TTP #325423: Making it so we correctly identify which execution-pin a node was triggered from.

  - Wire-trace op-codes are now inserted before every statement in GotoFixupRequestMap
  - ExpandNode() now uses FKismetCompilerContext::MovePinLinksToIntermediate/CopyPinLinksToIntermediate for transfering pin links (to track pin associations)

diff --git a/Engine/Plugins/Experimental/LiveEditor/Source/LiveEditor/Private/K2Node_LiveEditObject.cpp b/Engine/Plugins/Experimental/LiveEditor/Source/LiveEditor/Private/K2Node_LiveEditObject.cpp
--- a/Engine/Plugins/Experimental/LiveEditor/Source/LiveEditor/Private/K2Node_LiveEditObject.cpp
+++ b/Engine/Plugins/Experimental/LiveEditor/Source/LiveEditor/Private/K2Node_LiveEditObject.cpp
@@ -311,10 +311,10 @@ void UK2Node_LiveEditObject::ExpandNode(class FKismetCompilerContext& CompilerCo
 			RegisterForMIDINode->AllocateDefaultPins();
 
 			UEdGraphPin *ExecPin = RegisterForMIDINode->GetExecPin();
-			CompilerContext.CheckConnectionResponse(Schema->MovePinLinks(*SourceExecPin, *ExecPin), this);
+			CompilerContext.MovePinLinksToIntermediate(*SourceExecPin, *ExecPin);
 
 			UEdGraphPin *ThenPin = RegisterForMIDINode->GetThenPin();
-			CompilerContext.CheckConnectionResponse(Schema->MovePinLinks(*SourceThenPin, *ThenPin), this);
+			CompilerContext.MovePinLinksToIntermediate(*SourceThenPin, *ThenPin);
 
 			UEdGraphPin *TargetPin = RegisterForMIDINode->FindPinChecked( FString(TEXT("Target")) );
 			TargetPin->MakeLinkTo(SelfNodeThenPin);
@@ -323,10 +323,10 @@ void UK2Node_LiveEditObject::ExpandNode(class FKismetCompilerContext& CompilerCo
 			EventNamePin->DefaultValue = EventNameGuid;
 		
 			UEdGraphPin *DescriptionPin = RegisterForMIDINode->FindPinChecked( FString(TEXT("Description")) );
-			CompilerContext.CheckConnectionResponse(Schema->CopyPinLinks( *SourceDescriptionPin, *DescriptionPin), this );
+			CompilerContext.CopyPinLinksToIntermediate( *SourceDescriptionPin, *DescriptionPin);
 
 			UEdGraphPin *PermittedBindingsPin = RegisterForMIDINode->FindPinChecked( FString(TEXT("PermittedBindings")) );
-			CompilerContext.CheckConnectionResponse(Schema->CopyPinLinks( *SourcePermittedBindingsPin, *PermittedBindingsPin), this );
+			CompilerContext.CopyPinLinksToIntermediate( *SourcePermittedBindingsPin, *PermittedBindingsPin);
 		}
 
 		//Create the event handling part of the LiveEditor binding process
@@ -361,7 +361,7 @@ void UK2Node_LiveEditObject::ExpandNode(class FKismetCompilerContext& CompilerCo
 
 				//Set A Pin to the Blueprint Pin
 				UEdGraphPin *CompareBlueprintToNullAPin = CompareBlueprintToNullNode->FindPinChecked( FString(TEXT("A")) );
-				CompilerContext.CheckConnectionResponse(Schema->CopyPinLinks( *SourceBlueprintPin, *CompareBlueprintToNullAPin), this );
+				CompilerContext.CopyPinLinksToIntermediate( *SourceBlueprintPin, *CompareBlueprintToNullAPin);
 
 				// hook for Compare Blueprint to NULL result
 				UEdGraphPin *CompareBlueprintToNullResultPin = CompareBlueprintToNullNode->GetReturnValuePin();
@@ -388,7 +388,7 @@ void UK2Node_LiveEditObject::ExpandNode(class FKismetCompilerContext& CompilerCo
 			GetClassDefaultObjectNode->AllocateDefaultPins();
 
 			UEdGraphPin *GetClassDefaultObjectBlueprintPin = GetClassDefaultObjectNode->FindPinChecked( TEXT("Blueprint") );
-			CompilerContext.CheckConnectionResponse(Schema->CopyPinLinks( *SourceBlueprintPin, *GetClassDefaultObjectBlueprintPin), this );
+			CompilerContext.CopyPinLinksToIntermediate( *SourceBlueprintPin, *GetClassDefaultObjectBlueprintPin);
 
 			//hook for later -> the pointer to the ClassDefaultObject of our BlueprintPin
 			UEdGraphPin *GetClassDefaultObjectResultPin = GetClassDefaultObjectNode->GetReturnValuePin();
@@ -425,7 +425,7 @@ void UK2Node_LiveEditObject::ExpandNode(class FKismetCompilerContext& CompilerCo
 
 				//copy our BaseClass Pin into the ClassIsChildOf Parameter
 				UEdGraphPin *ClassIsChildOfParentPin = ClassIsChildOfNode->FindPinChecked( FString(TEXT("ParentClass")) );
-				CompilerContext.CheckConnectionResponse(Schema->CopyPinLinks( *SourceBaseClassPin, *ClassIsChildOfParentPin), this );
+				CompilerContext.CopyPinLinksToIntermediate( *SourceBaseClassPin, *ClassIsChildOfParentPin);
 
 				//hook for return value
 				UEdGraphPin *ClassIsChildOfResultPin = ClassIsChildOfNode->GetReturnValuePin();
@@ -481,15 +481,15 @@ void UK2Node_LiveEditObject::ExpandNode(class FKismetCompilerContext& CompilerCo
 			if ( ShouldClampPin->DefaultValue == FString(TEXT("true")) )
 			{
 				UEdGraphPin *ModifyVarNodeShouldClampPin = ModifyVarNode->FindPinChecked( TEXT("bShouldClamp") );
-				CompilerContext.CheckConnectionResponse(Schema->CopyPinLinks( *ShouldClampPin, *ModifyVarNodeShouldClampPin), this );
+				CompilerContext.CopyPinLinksToIntermediate( *ShouldClampPin, *ModifyVarNodeShouldClampPin);
 
 				check( ClampMinPin != NULL );
 				UEdGraphPin *ModifyVarNodeClampMinPin = ModifyVarNode->FindPinChecked( TEXT("ClampMin") );
-				CompilerContext.CheckConnectionResponse(Schema->CopyPinLinks( *ClampMinPin, *ModifyVarNodeClampMinPin), this );
+				CompilerContext.CopyPinLinksToIntermediate( *ClampMinPin, *ModifyVarNodeClampMinPin);
 
 				check( ClampMaxPin != NULL );
 				UEdGraphPin *ModifyVarNodeClampMaxPin = ModifyVarNode->FindPinChecked( TEXT("ClampMax") );
-				CompilerContext.CheckConnectionResponse(Schema->CopyPinLinks( *ClampMaxPin, *ModifyVarNodeClampMaxPin), this );
+				CompilerContext.CopyPinLinksToIntermediate( *ClampMaxPin, *ModifyVarNodeClampMaxPin);
 			}
 
 			//hook for ModifyVar THEN
@@ -507,7 +507,7 @@ void UK2Node_LiveEditObject::ExpandNode(class FKismetCompilerContext& CompilerCo
 
 			// 2nd input to the Add function comes from the Current variable value
 			UEdGraphPin *MultiplyNodeSecondPin = MultiplyNode->FindPinChecked( FString(TEXT("B")) );
-			CompilerContext.CheckConnectionResponse(Schema->CopyPinLinks( *DeltaMultPin, *MultiplyNodeSecondPin), this );
+			CompilerContext.CopyPinLinksToIntermediate( *DeltaMultPin, *MultiplyNodeSecondPin);
 
 			UEdGraphPin *MultiplyNodeReturnValuePin = MultiplyNode->GetReturnValuePin();
 			MultiplyNodeReturnValuePin->MakeLinkTo( ModifyVarNodeDeltaPin );
@@ -556,7 +556,7 @@ void UK2Node_LiveEditObject::ExpandNode(class FKismetCompilerContext& CompilerCo
 			//
 			// Finally, activate our OnMidiInput pin
 			//
-			CompilerContext.CheckConnectionResponse(Schema->CopyPinLinks( *SourceOnMidiInputPin, *ReplicationNodeThenPin), this );
+			CompilerContext.CopyPinLinksToIntermediate( *SourceOnMidiInputPin, *ReplicationNodeThenPin);
 			
 		}
 
diff --git a/Engine/Source/Editor/BlueprintGraph/Private/CallFunctionHandler.cpp b/Engine/Source/Editor/BlueprintGraph/Private/CallFunctionHandler.cpp
--- a/Engine/Source/Editor/BlueprintGraph/Private/CallFunctionHandler.cpp
+++ b/Engine/Source/Editor/BlueprintGraph/Private/CallFunctionHandler.cpp
@@ -152,6 +152,7 @@ void FKCHandler_CallFunction::CreateFunctionCallStatement(FKismetFunctionContext
 		// Parameter info to be stored, and assigned to all function call statements generated below
 		FBPTerminal* LHSTerm = NULL;
 		TArray<FBPTerminal*> RHSTerms;
+		UEdGraphPin* ThenExecPin = NULL;
 		UEdGraphNode* LatentTargetNode = NULL;
 		int32 LatentTargetParamIndex = INDEX_NONE;
 
@@ -224,7 +225,7 @@ void FKCHandler_CallFunction::CreateFunctionCallStatement(FKismetFunctionContext
 								if (PinMatch == LatentInfoPin)
 								{
 									// Record the (latent) output impulse from this node
-									UEdGraphPin* ThenExecPin = CompilerContext.GetSchema()->FindExecutionPin(*Node, EGPD_Output);
+									ThenExecPin = CompilerContext.GetSchema()->FindExecutionPin(*Node, EGPD_Output);
 
 									if( (ThenExecPin != NULL) && (ThenExecPin->LinkedTo.Num() > 0) )
 									{
@@ -331,8 +332,10 @@ void FKCHandler_CallFunction::CreateFunctionCallStatement(FKismetFunctionContext
 					// Fixup ubergraph calls
 					if (pSrcEventNode != NULL)
 					{
+						UEdGraphPin* ExecOut = CompilerContext.GetSchema()->FindExecutionPin(**pSrcEventNode, EGPD_Output);
+
 						check(CompilerContext.UbergraphContext);
-						CompilerContext.UbergraphContext->GotoFixupRequestMap.Add(&Statement, *pSrcEventNode);
+						CompilerContext.UbergraphContext->GotoFixupRequestMap.Add(&Statement, ExecOut);
 						Statement.UbergraphCallIndex = 0;
 					}
 				}
@@ -343,7 +346,7 @@ void FKCHandler_CallFunction::CreateFunctionCallStatement(FKismetFunctionContext
 					{
 						check(LatentTargetParamIndex != INDEX_NONE);
 						Statement.UbergraphCallIndex = LatentTargetParamIndex;
-						Context.GotoFixupRequestMap.Add(&Statement, LatentTargetNode);
+						Context.GotoFixupRequestMap.Add(&Statement, ThenExecPin);
 					}
 				}
 
@@ -640,6 +643,8 @@ void FKCHandler_CallFunction::Transform(FKismetFunctionContext& Context, UEdGrap
 
 			if ((NewInPin != NULL) && (NewOutPin != NULL))
 			{
+				CompilerContext.MessageLog.NotifyIntermediateObjectCreation(NewOutPin, OldOutPin);
+
 				while (OldOutPin->LinkedTo.Num() > 0)
 				{
 					UEdGraphPin* LinkedPin = OldOutPin->LinkedTo[0];
diff --git a/Engine/Source/Editor/BlueprintGraph/Private/DynamicCastHandler.cpp b/Engine/Source/Editor/BlueprintGraph/Private/DynamicCastHandler.cpp
--- a/Engine/Source/Editor/BlueprintGraph/Private/DynamicCastHandler.cpp
+++ b/Engine/Source/Editor/BlueprintGraph/Private/DynamicCastHandler.cpp
@@ -38,7 +38,7 @@ void FKCHandler_DynamicCast::Compile(FKismetFunctionContext& Context, UEdGraphNo
 	}
 
 	// Successful interface cast execution
-	const UEdGraphPin* SuccessPin = DynamicCastNode->GetValidCastPin();
+	UEdGraphPin* SuccessPin = DynamicCastNode->GetValidCastPin();
 	UEdGraphNode* SuccessNode = NULL;
 	if (SuccessPin->LinkedTo.Num() > 0)
 	{
@@ -46,7 +46,7 @@ void FKCHandler_DynamicCast::Compile(FKismetFunctionContext& Context, UEdGraphNo
 	}
 
 	// Failed interface cast execution
-	const UEdGraphPin* FailurePin = DynamicCastNode->GetInvalidCastPin();
+	UEdGraphPin* FailurePin = DynamicCastNode->GetInvalidCastPin();
 	UEdGraphNode* FailureNode = NULL;
 	if (FailurePin->LinkedTo.Num() > 0)
 	{
@@ -141,13 +141,13 @@ void FKCHandler_DynamicCast::Compile(FKismetFunctionContext& Context, UEdGraphNo
 	FBlueprintCompiledStatement& FailCastGoto = Context.AppendStatementForNode(Node);
 	FailCastGoto.Type = KCST_GotoIfNot;
 	FailCastGoto.LHS = BoolTerm;
-	Context.GotoFixupRequestMap.Add(&FailCastGoto, FailureNode);
+	Context.GotoFixupRequestMap.Add(&FailCastGoto, FailurePin);
 
 	// Successful cast...hit the success output node
 	FBlueprintCompiledStatement& SuccessCastGoto = Context.AppendStatementForNode(Node);
 	SuccessCastGoto.Type = KCST_UnconditionalGoto;
 	SuccessCastGoto.LHS = BoolTerm;
-	Context.GotoFixupRequestMap.Add(&SuccessCastGoto, SuccessNode);
+	Context.GotoFixupRequestMap.Add(&SuccessCastGoto, SuccessPin);
 }
 
 #undef LOCTEXT_NAMESPACE
diff --git a/Engine/Source/Editor/BlueprintGraph/Private/K2Node_AddComponent.cpp b/Engine/Source/Editor/BlueprintGraph/Private/K2Node_AddComponent.cpp
--- a/Engine/Source/Editor/BlueprintGraph/Private/K2Node_AddComponent.cpp
+++ b/Engine/Source/Editor/BlueprintGraph/Private/K2Node_AddComponent.cpp
@@ -346,17 +346,17 @@ void UK2Node_AddComponent::ExpandNode(class FKismetCompilerContext& CompilerCont
 		NewNode->AllocateDefaultPinsWithoutExposedVariables();
 
 		// function parameters
-		CompilerContext.CheckConnectionResponse(Schema->MovePinLinks(*GetTemplateNamePinChecked(), *NewNode->GetTemplateNamePinChecked()), this);
-		CompilerContext.CheckConnectionResponse(Schema->MovePinLinks(*GetRelativeTransformPin(), *NewNode->GetRelativeTransformPin()), this);
-		CompilerContext.CheckConnectionResponse(Schema->MovePinLinks(*GetManualAttachmentPin(), *NewNode->GetManualAttachmentPin()), this);
+		CompilerContext.MovePinLinksToIntermediate(*GetTemplateNamePinChecked(), *NewNode->GetTemplateNamePinChecked());
+		CompilerContext.MovePinLinksToIntermediate(*GetRelativeTransformPin(), *NewNode->GetRelativeTransformPin());
+		CompilerContext.MovePinLinksToIntermediate(*GetManualAttachmentPin(), *NewNode->GetManualAttachmentPin());
 
 		UEdGraphPin* ReturnPin = NewNode->GetReturnValuePin();
 		UEdGraphPin* OriginalReturnPin = GetReturnValuePin();
 		check((NULL != ReturnPin) && (NULL != OriginalReturnPin));
 		ReturnPin->PinType = OriginalReturnPin->PinType;
-		CompilerContext.CheckConnectionResponse(Schema->MovePinLinks(*OriginalReturnPin, *ReturnPin), this);
+		CompilerContext.MovePinLinksToIntermediate(*OriginalReturnPin, *ReturnPin);
 		// exec in
-		CompilerContext.CheckConnectionResponse(Schema->MovePinLinks(*GetExecPin(), *NewNode->GetExecPin()), this);
+		CompilerContext.MovePinLinksToIntermediate(*GetExecPin(), *NewNode->GetExecPin());
 
 		UEdGraphPin* LastThen = NewNode->GetThenPin();
 		for(int32 PinIndex = 0; PinIndex < Pins.Num(); PinIndex++)
@@ -391,14 +391,14 @@ void UK2Node_AddComponent::ExpandNode(class FKismetCompilerContext& CompilerCont
 
 				// Move connection from the variable pin on the spawn node to the 'value' pin
 				UEdGraphPin* ValuePin = SetVarNode->FindPinChecked(ValueParamName);
-				CompilerContext.CheckConnectionResponse(Schema->MovePinLinks(*OrgPin, *ValuePin), this);
+				CompilerContext.MovePinLinksToIntermediate(*OrgPin, *ValuePin);
 				if(OrgPin->PinType.bIsArray)
 				{
 					SetVarNode->PinConnectionListChanged(ValuePin);
 				}
 			}
 		}
-		CompilerContext.CheckConnectionResponse(Schema->MovePinLinks(*GetThenPin(), *LastThen), this);
+		CompilerContext.MovePinLinksToIntermediate(*GetThenPin(), *LastThen);
 		BreakAllNodeLinks();
 	}
 }
diff --git a/Engine/Source/Editor/BlueprintGraph/Private/K2Node_BaseAsyncTask.cpp b/Engine/Source/Editor/BlueprintGraph/Private/K2Node_BaseAsyncTask.cpp
--- a/Engine/Source/Editor/BlueprintGraph/Private/K2Node_BaseAsyncTask.cpp
+++ b/Engine/Source/Editor/BlueprintGraph/Private/K2Node_BaseAsyncTask.cpp
@@ -170,7 +170,7 @@ void UK2Node_BaseAsyncTask::ExpandNode(class FKismetCompilerContext& CompilerCon
 	UK2Node_CallFunction* CallCreateMoveToProxyObjectNode = CompilerContext.SpawnIntermediateNode<UK2Node_CallFunction>(CurrentNode, SourceGraph);
 	CallCreateMoveToProxyObjectNode->FunctionReference.SetExternalMember(CreateMoveToProxyObjectBlueprintFuncName, CurrentNode->GetProxyFactoryClass());
 	CallCreateMoveToProxyObjectNode->AllocateDefaultPins();
-	CompilerContext.CheckConnectionResponse(CompilerContext.GetSchema()->MovePinLinks(*CurrentNode->FindPinChecked(Schema->PN_Execute), *CallCreateMoveToProxyObjectNode->FindPinChecked(Schema->PN_Execute)), CurrentNode);
+	CompilerContext.MovePinLinksToIntermediate(*CurrentNode->FindPinChecked(Schema->PN_Execute), *CallCreateMoveToProxyObjectNode->FindPinChecked(Schema->PN_Execute));
 
 	// match function inputs, to pass data to function from CallFunction node
 	for (int32 Index=0; Index < CurrentNode->Pins.Num(); ++Index)
@@ -183,7 +183,7 @@ void UK2Node_BaseAsyncTask::ExpandNode(class FKismetCompilerContext& CompilerCon
 
 		if (UEdGraphPin* DestPin = CallCreateMoveToProxyObjectNode->FindPin(CurrentPin->PinName))
 		{
-			CompilerContext.CheckConnectionResponse(Schema->MovePinLinks(*CurrentPin, *DestPin), CurrentNode);
+			CompilerContext.MovePinLinksToIntermediate(*CurrentPin, *DestPin);
 		}
 	}
 
@@ -325,7 +325,7 @@ void UK2Node_BaseAsyncTask::ExpandNode(class FKismetCompilerContext& CompilerCon
 		}
 
 		// connect CE output with OnSuccess output of AIMoveTo node
-		CompilerContext.CheckConnectionResponse(Schema->MovePinLinks(*PinForCurrentDelegateProperty, *LastActivatedNodeThen), CurrentNode);
+		CompilerContext.MovePinLinksToIntermediate(*PinForCurrentDelegateProperty, *LastActivatedNodeThen);
 	}
 
 	for (int32 OutIndex = 0; OutIndex < VariableOutputPins.Num(); ++OutIndex)
@@ -333,16 +333,16 @@ void UK2Node_BaseAsyncTask::ExpandNode(class FKismetCompilerContext& CompilerCon
 		UEdGraphPin* OutPin = VariableOutputPins[OutIndex];
 		TArray<UEdGraphPin*> OptionPins;
 		OutputSelectNodes[OutIndex]->GetOptionPins(OptionPins);
-		CompilerContext.CheckConnectionResponse(Schema->MovePinLinks(*OutPin, *OutputSelectNodes[OutIndex]->GetReturnValuePin()), CurrentNode);
+		CompilerContext.MovePinLinksToIntermediate(*OutPin, *OutputSelectNodes[OutIndex]->GetReturnValuePin());
 	}
 
 	if (LastThenPin != NULL)
 	{
-		CompilerContext.CheckConnectionResponse(Schema->MovePinLinks(*CurrentNode->FindPinChecked(Schema->PN_Then), *LastThenPin), CurrentNode);
+		CompilerContext.MovePinLinksToIntermediate(*CurrentNode->FindPinChecked(Schema->PN_Then), *LastThenPin);
 	}
 	else
 	{
-		CompilerContext.CheckConnectionResponse(Schema->MovePinLinks(*CurrentNode->FindPinChecked(Schema->PN_Then), *CallCreateMoveToProxyObjectNode->FindPinChecked(Schema->PN_Then)), CurrentNode);
+		CompilerContext.MovePinLinksToIntermediate(*CurrentNode->FindPinChecked(Schema->PN_Then), *CallCreateMoveToProxyObjectNode->FindPinChecked(Schema->PN_Then));
 	}
 
 	CurrentNode->BreakAllNodeLinks();
diff --git a/Engine/Source/Editor/BlueprintGraph/Private/K2Node_CallFunction.cpp b/Engine/Source/Editor/BlueprintGraph/Private/K2Node_CallFunction.cpp
--- a/Engine/Source/Editor/BlueprintGraph/Private/K2Node_CallFunction.cpp
+++ b/Engine/Source/Editor/BlueprintGraph/Private/K2Node_CallFunction.cpp
@@ -1224,7 +1224,7 @@ void UK2Node_CallFunction::ExpandNode(class FKismetCompilerContext& CompilerCont
 							AssignNode->AllocateDefaultPins();
 
 							// Move connections from fake 'enum exec' pint to this assignment node
-							CompilerContext.CheckConnectionResponse(Schema->MovePinLinks(*Pin, *AssignNode->GetExecPin()), this);
+							CompilerContext.MovePinLinksToIntermediate(*Pin, *AssignNode->GetExecPin());
 
 							// Connect this to out temp enum var
 							Schema->TryCreateConnection(AssignNode->GetVariablePin(), TempEnumVarOutput);
@@ -1282,13 +1282,13 @@ void UK2Node_CallFunction::CallForEachElementInArrayExpansion(UK2Node* Node, UEd
 	InteratorInitialize->AllocateDefaultPins();
 	InteratorInitialize->GetValuePin()->DefaultValue = TEXT("0");
 	Schema->TryCreateConnection(IteratorVar->GetVariablePin(), InteratorInitialize->GetVariablePin());
-	CompilerContext.CheckConnectionResponse(Schema->MovePinLinks(*Node->GetExecPin(), *InteratorInitialize->GetExecPin()), Node);
+	CompilerContext.MovePinLinksToIntermediate(*Node->GetExecPin(), *InteratorInitialize->GetExecPin());
 
 	// Do loop branch
 	UK2Node_IfThenElse* Branch = CompilerContext.SpawnIntermediateNode<UK2Node_IfThenElse>(Node, SourceGraph);
 	Branch->AllocateDefaultPins();
 	Schema->TryCreateConnection(InteratorInitialize->GetThenPin(), Branch->GetExecPin());
-	CompilerContext.CheckConnectionResponse(Schema->MovePinLinks(*ThenPin, *Branch->GetElsePin()), Node);
+	CompilerContext.MovePinLinksToIntermediate(*ThenPin, *Branch->GetElsePin());
 
 	// Do loop condition
 	UK2Node_CallFunction* Condition = CompilerContext.SpawnIntermediateNode<UK2Node_CallFunction>(Node, SourceGraph); 
@@ -1301,7 +1301,7 @@ void UK2Node_CallFunction::CallForEachElementInArrayExpansion(UK2Node* Node, UEd
 	UK2Node_CallArrayFunction* ArrayLength = CompilerContext.SpawnIntermediateNode<UK2Node_CallArrayFunction>(Node, SourceGraph); 
 	ArrayLength->SetFromFunction(UKismetArrayLibrary::StaticClass()->FindFunctionByName(TEXT("Array_Length")));
 	ArrayLength->AllocateDefaultPins();
-	CompilerContext.CheckConnectionResponse(Schema->CopyPinLinks(*MultiSelf, *ArrayLength->GetTargetArrayPin()), Node);
+	CompilerContext.CopyPinLinksToIntermediate(*MultiSelf, *ArrayLength->GetTargetArrayPin());
 	ArrayLength->PinConnectionListChanged(ArrayLength->GetTargetArrayPin());
 	Schema->TryCreateConnection(Condition->FindPinChecked(TEXT("B")), ArrayLength->GetReturnValuePin());
 
@@ -1309,7 +1309,7 @@ void UK2Node_CallFunction::CallForEachElementInArrayExpansion(UK2Node* Node, UEd
 	UK2Node_CallArrayFunction* GetElement = CompilerContext.SpawnIntermediateNode<UK2Node_CallArrayFunction>(Node, SourceGraph); 
 	GetElement->SetFromFunction(UKismetArrayLibrary::StaticClass()->FindFunctionByName(TEXT("Array_Get")));
 	GetElement->AllocateDefaultPins();
-	CompilerContext.CheckConnectionResponse(Schema->CopyPinLinks(*MultiSelf, *GetElement->GetTargetArrayPin()), Node);
+	CompilerContext.CopyPinLinksToIntermediate(*MultiSelf, *GetElement->GetTargetArrayPin());
 	GetElement->PinConnectionListChanged(GetElement->GetTargetArrayPin());
 	Schema->TryCreateConnection(GetElement->FindPinChecked(TEXT("Index")), IteratorVar->GetVariablePin());
 
diff --git a/Engine/Source/Editor/BlueprintGraph/Private/K2Node_CallFunctionOnMember.cpp b/Engine/Source/Editor/BlueprintGraph/Private/K2Node_CallFunctionOnMember.cpp
--- a/Engine/Source/Editor/BlueprintGraph/Private/K2Node_CallFunctionOnMember.cpp
+++ b/Engine/Source/Editor/BlueprintGraph/Private/K2Node_CallFunctionOnMember.cpp
@@ -101,7 +101,7 @@ void UK2Node_CallFunctionOnMember::ExpandNode(class FKismetCompilerContext& Comp
 				UEdGraphPin* DestPin = CallFuncNode->FindPin(SrcPin->PinName);
 				if(DestPin != NULL)
 				{
-					CompilerContext.CheckConnectionResponse(Schema->MovePinLinks(*SrcPin, *DestPin), this); // Source node is assumed to be owner...
+					CompilerContext.MovePinLinksToIntermediate(*SrcPin, *DestPin); // Source node is assumed to be owner...
 				}
 			}
 		}
diff --git a/Engine/Source/Editor/BlueprintGraph/Private/K2Node_CastByteToEnum.cpp b/Engine/Source/Editor/BlueprintGraph/Private/K2Node_CastByteToEnum.cpp
--- a/Engine/Source/Editor/BlueprintGraph/Private/K2Node_CastByteToEnum.cpp
+++ b/Engine/Source/Editor/BlueprintGraph/Private/K2Node_CastByteToEnum.cpp
@@ -71,7 +71,7 @@ void UK2Node_CastByteToEnum::ExpandNode(class FKismetCompilerContext& CompilerCo
 		UEdGraphPin* OrgInputPin = FindPinChecked(ByteInputPinName);
 		UEdGraphPin* FunctionIndexPin = CallValidation->FindPinChecked(TEXT("EnumeratorIndex"));
 		check(EGPD_Input == FunctionIndexPin->Direction && Schema->PC_Byte == FunctionIndexPin->PinType.PinCategory);
-		CompilerContext.CheckConnectionResponse(Schema->MovePinLinks(*OrgInputPin, *FunctionIndexPin), this);
+		CompilerContext.MovePinLinksToIntermediate(*OrgInputPin, *FunctionIndexPin);
 
 		// UNSAFE CAST NODE
 		UK2Node_CastByteToEnum* UsafeCast = CompilerContext.SpawnIntermediateNode<UK2Node_CastByteToEnum>(this, SourceGraph); 
@@ -88,7 +88,7 @@ void UK2Node_CastByteToEnum::ExpandNode(class FKismetCompilerContext& CompilerCo
 		// OPUTPUT PIN
 		UEdGraphPin* OrgReturnPin = FindPinChecked(Schema->PN_ReturnValue);
 		UEdGraphPin* NewReturnPin = UsafeCast->FindPinChecked(Schema->PN_ReturnValue);
-		CompilerContext.CheckConnectionResponse(Schema->MovePinLinks(*OrgReturnPin, *NewReturnPin), this);
+		CompilerContext.MovePinLinksToIntermediate(*OrgReturnPin, *NewReturnPin);
 
 		BreakAllNodeLinks();
 	}
diff --git a/Engine/Source/Editor/BlueprintGraph/Private/K2Node_DelegateSet.cpp b/Engine/Source/Editor/BlueprintGraph/Private/K2Node_DelegateSet.cpp
--- a/Engine/Source/Editor/BlueprintGraph/Private/K2Node_DelegateSet.cpp
+++ b/Engine/Source/Editor/BlueprintGraph/Private/K2Node_DelegateSet.cpp
@@ -279,7 +279,7 @@ void UK2Node_DelegateSet::ExpandNode(class FKismetCompilerContext& CompilerConte
 						// Hook up the exec pin specially, since it has a different name on the dynamic delegate node
 						UEdGraphPin* OldExecPin = FindPin(Schema->PN_DelegateEntry);
 						check(OldExecPin);
-						CompilerContext.CheckConnectionResponse(Schema->MovePinLinks(*OldExecPin, *CurrentPin), this);
+						CompilerContext.MovePinLinksToIntermediate(*OldExecPin, *CurrentPin);
 					}
 					else if( CurrentPin->PinName != UK2Node_Event::DelegateOutputName )
 					{
@@ -292,7 +292,7 @@ void UK2Node_DelegateSet::ExpandNode(class FKismetCompilerContext& CompilerConte
 							return;
 						}
 
-						CompilerContext.CheckConnectionResponse(Schema->MovePinLinks(*OldPin, *CurrentPin), this);
+						CompilerContext.MovePinLinksToIntermediate(*OldPin, *CurrentPin);
 					}
 				}
 			}
diff --git a/Engine/Source/Editor/BlueprintGraph/Private/K2Node_EnumEquality.cpp b/Engine/Source/Editor/BlueprintGraph/Private/K2Node_EnumEquality.cpp
--- a/Engine/Source/Editor/BlueprintGraph/Private/K2Node_EnumEquality.cpp
+++ b/Engine/Source/Editor/BlueprintGraph/Private/K2Node_EnumEquality.cpp
@@ -227,9 +227,9 @@ void UK2Node_EnumEquality::ExpandNode(class FKismetCompilerContext& CompilerCont
 		UEdGraphPin* Input2Pin = GetInput2Pin();
 		LeftSideConditionalPin->PinType = Input1Pin->PinType;
 		RightSideConditionalPin->PinType = Input2Pin->PinType;
-		CompilerContext.CheckConnectionResponse(Schema->MovePinLinks(*Input1Pin, *LeftSideConditionalPin), this);
-		CompilerContext.CheckConnectionResponse(Schema->MovePinLinks(*Input2Pin, *RightSideConditionalPin), this);
-		CompilerContext.CheckConnectionResponse(Schema->MovePinLinks(*GetReturnValuePin(), *ReturnConditionalPin), this);
+		CompilerContext.MovePinLinksToIntermediate(*Input1Pin, *LeftSideConditionalPin);
+		CompilerContext.MovePinLinksToIntermediate(*Input2Pin, *RightSideConditionalPin);
+		CompilerContext.MovePinLinksToIntermediate(*GetReturnValuePin(), *ReturnConditionalPin);
 
 		// Break all links to the Select node so it goes away for at scheduling time
 		BreakAllNodeLinks();
diff --git a/Engine/Source/Editor/BlueprintGraph/Private/K2Node_Event.cpp b/Engine/Source/Editor/BlueprintGraph/Private/K2Node_Event.cpp
--- a/Engine/Source/Editor/BlueprintGraph/Private/K2Node_Event.cpp
+++ b/Engine/Source/Editor/BlueprintGraph/Private/K2Node_Event.cpp
@@ -531,7 +531,7 @@ void UK2Node_Event::ExpandNode(class FKismetCompilerContext& CompilerContext, UE
 
 			UK2Node_CreateDelegate* CreateDelegateNode = CompilerContext.SpawnIntermediateNode<UK2Node_CreateDelegate>(this, SourceGraph);
 			CreateDelegateNode->AllocateDefaultPins();
-			CompilerContext.CheckConnectionResponse(Schema->MovePinLinks(*OrgDelegatePin, *CreateDelegateNode->GetDelegateOutPin()), this);
+			CompilerContext.MovePinLinksToIntermediate(*OrgDelegatePin, *CreateDelegateNode->GetDelegateOutPin());
 			Schema->TryCreateConnection(SelfNode->FindPinChecked(Schema->PN_Self), CreateDelegateNode->GetObjectInPin());
 			CreateDelegateNode->SelectedFunctionName = FunctionName;
 			CreateDelegateNode->HandleAnyChangeInner();
diff --git a/Engine/Source/Editor/BlueprintGraph/Private/K2Node_ExecutionSequence.cpp b/Engine/Source/Editor/BlueprintGraph/Private/K2Node_ExecutionSequence.cpp
--- a/Engine/Source/Editor/BlueprintGraph/Private/K2Node_ExecutionSequence.cpp
+++ b/Engine/Source/Editor/BlueprintGraph/Private/K2Node_ExecutionSequence.cpp
@@ -66,8 +66,6 @@ public:
 
 				for (int32 i = 0; i < OutputPins.Num(); ++i)
 				{
-					UEdGraphNode* StateNode = OutputPins[i]->LinkedTo[0]->GetOwningNode();
-
 					// Emit the debug site and patch up the previous jump if we're on subsequent steps
 					const bool bNotFirstIndex = i > 0;
 					if (bNotFirstIndex)
@@ -95,7 +93,7 @@ public:
 					// Emit the goto to the actual state
 					FBlueprintCompiledStatement& GotoSequenceLinkedState = Context.AppendStatementForNode(Node);
 					GotoSequenceLinkedState.Type = KCST_UnconditionalGoto;
-					Context.GotoFixupRequestMap.Add(&GotoSequenceLinkedState, StateNode);
+					Context.GotoFixupRequestMap.Add(&GotoSequenceLinkedState, OutputPins[i]);
 				}
 			}
 			else
@@ -103,17 +101,16 @@ public:
 				// Directly emit pushes to execute the remaining branches
 				for (int32 i = OutputPins.Num() - 1; i > 0; i--)
 				{
-					UEdGraphNode* StateNode = OutputPins[i]->LinkedTo[0]->GetOwningNode();
 					FBlueprintCompiledStatement& PushExecutionState = Context.AppendStatementForNode(Node);
 					PushExecutionState.Type = KCST_PushState;
-					Context.GotoFixupRequestMap.Add(&PushExecutionState, StateNode);
+					Context.GotoFixupRequestMap.Add(&PushExecutionState, OutputPins[i]);
 				}
 
 				// Immediately jump to the first pin
 				UEdGraphNode* NextNode = OutputPins[0]->LinkedTo[0]->GetOwningNode();
 				FBlueprintCompiledStatement& NextExecutionState = Context.AppendStatementForNode(Node);
 				NextExecutionState.Type = KCST_UnconditionalGoto;
-				Context.GotoFixupRequestMap.Add(&NextExecutionState, NextNode);
+				Context.GotoFixupRequestMap.Add(&NextExecutionState, OutputPins[0]);
 			}
 		}
 	}
diff --git a/Engine/Source/Editor/BlueprintGraph/Private/K2Node_ForEachElementInEnum.cpp b/Engine/Source/Editor/BlueprintGraph/Private/K2Node_ForEachElementInEnum.cpp
--- a/Engine/Source/Editor/BlueprintGraph/Private/K2Node_ForEachElementInEnum.cpp
+++ b/Engine/Source/Editor/BlueprintGraph/Private/K2Node_ForEachElementInEnum.cpp
@@ -149,9 +149,9 @@ void UK2Node_ForEachElementInEnum::ExpandNode(class FKismetCompilerContext& Comp
 
 		const UEdGraphSchema_K2* Schema = CompilerContext.GetSchema();
 
-		CompilerContext.CheckConnectionResponse(Schema->MovePinLinks(*GetExecPin(), *ForLoop.StartLoopExecInPin), this);
-		CompilerContext.CheckConnectionResponse(Schema->MovePinLinks(*FindPinChecked(Schema->PN_Then), *ForLoop.LoopCompleteOutExecPin), this);
-		CompilerContext.CheckConnectionResponse(Schema->MovePinLinks(*FindPinChecked(InsideLoopPinName), *ForLoop.InsideLoopExecOutPin), this);
+		CompilerContext.MovePinLinksToIntermediate(*GetExecPin(), *ForLoop.StartLoopExecInPin);
+		CompilerContext.MovePinLinksToIntermediate(*FindPinChecked(Schema->PN_Then), *ForLoop.LoopCompleteOutExecPin);
+		CompilerContext.MovePinLinksToIntermediate(*FindPinChecked(InsideLoopPinName), *ForLoop.InsideLoopExecOutPin);
 
 		UK2Node_GetNumEnumEntries* GetNumEnumEntries = CompilerContext.SpawnIntermediateNode<UK2Node_GetNumEnumEntries>(this, SourceGraph);
 		GetNumEnumEntries->Enum = Enum;
@@ -169,7 +169,7 @@ void UK2Node_ForEachElementInEnum::ExpandNode(class FKismetCompilerContext& Comp
 		CastByteToEnum->bSafe = true;
 		CastByteToEnum->AllocateDefaultPins();
 		bResult &= Schema->TryCreateConnection(Conv_Func->FindPinChecked(Schema->PN_ReturnValue), CastByteToEnum->FindPinChecked(UK2Node_CastByteToEnum::ByteInputPinName));
-		CompilerContext.CheckConnectionResponse(Schema->MovePinLinks(*FindPinChecked(EnumOuputPinName), *CastByteToEnum->FindPinChecked(Schema->PN_ReturnValue)), this);
+		CompilerContext.MovePinLinksToIntermediate(*FindPinChecked(EnumOuputPinName), *CastByteToEnum->FindPinChecked(Schema->PN_ReturnValue));
 
 		if (!bResult)
 		{
diff --git a/Engine/Source/Editor/BlueprintGraph/Private/K2Node_FormatText.cpp b/Engine/Source/Editor/BlueprintGraph/Private/K2Node_FormatText.cpp
--- a/Engine/Source/Editor/BlueprintGraph/Private/K2Node_FormatText.cpp
+++ b/Engine/Source/Editor/BlueprintGraph/Private/K2Node_FormatText.cpp
@@ -198,7 +198,7 @@ void UK2Node_FormatText::ExpandNode(class FKismetCompilerContext& CompilerContex
 			PinMakeStruct->GetSchema()->TrySetDefaultText(*PinMakeStruct->FindPin("ArgumentName"), FText::FromString(ArgumentPin->PinName));
 
 			// Move the connection of the argument pin to the struct's "TextValue" pin, this will move the literal value if present.
-			CompilerContext.CheckConnectionResponse(Schema->MovePinLinks(*ArgumentPin, *PinMakeStruct->FindPin("TextValue")), this);
+			CompilerContext.MovePinLinksToIntermediate(*ArgumentPin, *PinMakeStruct->FindPin("TextValue"));
 
 			// The "Make Array" node already has one pin available, so don't create one for ArgIdx == 0
 			if(ArgIdx > 0)
@@ -216,9 +216,9 @@ void UK2Node_FormatText::ExpandNode(class FKismetCompilerContext& CompilerContex
 		}
 
 		// Move connection of FormatText's "Result" pin to the call function's return value pin.
-		CompilerContext.CheckConnectionResponse(Schema->MovePinLinks(*FindPin(TEXT("Result")), *CallFunction->GetReturnValuePin()), this);
+		CompilerContext.MovePinLinksToIntermediate(*FindPin(TEXT("Result")), *CallFunction->GetReturnValuePin());
 		// Move connection of FormatText's "Format" pin to the call function's "InPattern" pin
-		CompilerContext.CheckConnectionResponse(Schema->MovePinLinks(*FindPin(TEXT("Format")), *CallFunction->FindPin(TEXT("InPattern"))), this);
+		CompilerContext.MovePinLinksToIntermediate(*FindPin(TEXT("Format")), *CallFunction->FindPin(TEXT("InPattern")));
 
 		BreakAllNodeLinks();
 	}
diff --git a/Engine/Source/Editor/BlueprintGraph/Private/K2Node_GetEnumeratorName.cpp b/Engine/Source/Editor/BlueprintGraph/Private/K2Node_GetEnumeratorName.cpp
--- a/Engine/Source/Editor/BlueprintGraph/Private/K2Node_GetEnumeratorName.cpp
+++ b/Engine/Source/Editor/BlueprintGraph/Private/K2Node_GetEnumeratorName.cpp
@@ -139,7 +139,7 @@ void UK2Node_GetEnumeratorName::ExpandNode(class FKismetCompilerContext& Compile
 		UEdGraphPin* OrgReturnPin = FindPinChecked(Schema->PN_ReturnValue);
 		UEdGraphPin* NewReturnPin = CallGetName->GetReturnValuePin();
 		check(NULL != NewReturnPin);
-		CompilerContext.CheckConnectionResponse(Schema->MovePinLinks(*OrgReturnPin, *NewReturnPin), this);
+		CompilerContext.MovePinLinksToIntermediate(*OrgReturnPin, *NewReturnPin);
 
 		//ENUM PIN
 		UEdGraphPin* EnumPin = CallGetName->FindPinChecked(TEXT("Enum"));
@@ -150,7 +150,7 @@ void UK2Node_GetEnumeratorName::ExpandNode(class FKismetCompilerContext& Compile
 		UEdGraphPin* OrgInputPin = FindPinChecked(EnumeratorPinName);
 		UEdGraphPin* IndexPin = CallGetName->FindPinChecked(TEXT("EnumeratorIndex"));
 		check(EGPD_Input == IndexPin->Direction && Schema->PC_Byte == IndexPin->PinType.PinCategory);
-		CompilerContext.CheckConnectionResponse(Schema->MovePinLinks(*OrgInputPin, *IndexPin), this);
+		CompilerContext.MovePinLinksToIntermediate(*OrgInputPin, *IndexPin);
 
 		if (!IndexPin->LinkedTo.Num())
 		{
diff --git a/Engine/Source/Editor/BlueprintGraph/Private/K2Node_GetNumEnumEntries.cpp b/Engine/Source/Editor/BlueprintGraph/Private/K2Node_GetNumEnumEntries.cpp
--- a/Engine/Source/Editor/BlueprintGraph/Private/K2Node_GetNumEnumEntries.cpp
+++ b/Engine/Source/Editor/BlueprintGraph/Private/K2Node_GetNumEnumEntries.cpp
@@ -59,7 +59,7 @@ void UK2Node_GetNumEnumEntries::ExpandNode(class FKismetCompilerContext& Compile
 		UEdGraphPin* OrgReturnPin = FindPinChecked(Schema->PN_ReturnValue);
 		UEdGraphPin* NewReturnPin = MakeLiteralInt->GetReturnValuePin();
 		check(NULL != NewReturnPin);
-		CompilerContext.CheckConnectionResponse(Schema->MovePinLinks(*OrgReturnPin, *NewReturnPin), this);
+		CompilerContext.MovePinLinksToIntermediate(*OrgReturnPin, *NewReturnPin);
 
 		//INPUT PIN
 		UEdGraphPin* InputPin = MakeLiteralInt->FindPinChecked(TEXT("Value"));
diff --git a/Engine/Source/Editor/BlueprintGraph/Private/K2Node_IfThenElse.cpp b/Engine/Source/Editor/BlueprintGraph/Private/K2Node_IfThenElse.cpp
--- a/Engine/Source/Editor/BlueprintGraph/Private/K2Node_IfThenElse.cpp
+++ b/Engine/Source/Editor/BlueprintGraph/Private/K2Node_IfThenElse.cpp
@@ -47,30 +47,18 @@ public:
 			{
 				// First skip the if, if the term is false
 				{
-					UEdGraphNode* TargetNode = NULL;
-					if (ElsePin->LinkedTo.Num() > 0)
-					{
-						TargetNode = ElsePin->LinkedTo[0]->GetOwningNode();
-					}
-
 					FBlueprintCompiledStatement& SkipIfGoto = Context.AppendStatementForNode(Node);
 					SkipIfGoto.Type = KCST_GotoIfNot;
 					SkipIfGoto.LHS = *CondTerm;
-					Context.GotoFixupRequestMap.Add(&SkipIfGoto, TargetNode);
+					Context.GotoFixupRequestMap.Add(&SkipIfGoto, ElsePin);
 				}
 
 				// Now go to the If branch
 				{
-					UEdGraphNode* TargetNode = NULL;
-					if (ThenPin->LinkedTo.Num() > 0)
-					{
-						TargetNode = ThenPin->LinkedTo[0]->GetOwningNode();
-					}
-
 					FBlueprintCompiledStatement& GotoThen = Context.AppendStatementForNode(Node);
 					GotoThen.Type = KCST_UnconditionalGoto;
 					GotoThen.LHS = *CondTerm;
-					Context.GotoFixupRequestMap.Add(&GotoThen, TargetNode);
+					Context.GotoFixupRequestMap.Add(&GotoThen, ThenPin);
 				}
 			}
 			else
diff --git a/Engine/Source/Editor/BlueprintGraph/Private/K2Node_InputAction.cpp b/Engine/Source/Editor/BlueprintGraph/Private/K2Node_InputAction.cpp
--- a/Engine/Source/Editor/BlueprintGraph/Private/K2Node_InputAction.cpp
+++ b/Engine/Source/Editor/BlueprintGraph/Private/K2Node_InputAction.cpp
@@ -91,7 +91,7 @@ void UK2Node_InputAction::CreateInputActionEvent(FKismetCompilerContext& Compile
 
 		if(EventOutput != NULL)
 		{
-			CompilerContext.CheckConnectionResponse(CompilerContext.GetSchema()->MovePinLinks(*InputActionPin, *EventOutput), this);
+			CompilerContext.MovePinLinksToIntermediate(*InputActionPin, *EventOutput);
 		}
 	}
 }
diff --git a/Engine/Source/Editor/BlueprintGraph/Private/K2Node_InputKey.cpp b/Engine/Source/Editor/BlueprintGraph/Private/K2Node_InputKey.cpp
--- a/Engine/Source/Editor/BlueprintGraph/Private/K2Node_InputKey.cpp
+++ b/Engine/Source/Editor/BlueprintGraph/Private/K2Node_InputKey.cpp
@@ -229,7 +229,7 @@ void UK2Node_InputKey::CreateInputKeyEvent(FKismetCompilerContext& CompilerConte
 
 		if(EventOutput != NULL)
 		{
-			CompilerContext.CheckConnectionResponse(CompilerContext.GetSchema()->MovePinLinks(*InputKeyPin, *EventOutput), this);
+			CompilerContext.MovePinLinksToIntermediate(*InputKeyPin, *EventOutput);
 		}
 	}
 }
diff --git a/Engine/Source/Editor/BlueprintGraph/Private/K2Node_MatineeController.cpp b/Engine/Source/Editor/BlueprintGraph/Private/K2Node_MatineeController.cpp
--- a/Engine/Source/Editor/BlueprintGraph/Private/K2Node_MatineeController.cpp
+++ b/Engine/Source/Editor/BlueprintGraph/Private/K2Node_MatineeController.cpp
@@ -103,7 +103,7 @@ void UK2Node_MatineeController::ExpandNode(FKismetCompilerContext& CompilerConte
 				// Move connection from matinee output to event node output
 				UEdGraphPin* EventOutputPin = Schema->FindExecutionPin(*MatineeEventNode, EGPD_Output);
 				check(EventOutputPin != NULL);
-				CompilerContext.CheckConnectionResponse(Schema->MovePinLinks(*MatineePin, *EventOutputPin), this);
+				CompilerContext.MovePinLinksToIntermediate(*MatineePin, *EventOutputPin);
 			}
 		}
 	}
diff --git a/Engine/Source/Editor/BlueprintGraph/Private/K2Node_Message.cpp b/Engine/Source/Editor/BlueprintGraph/Private/K2Node_Message.cpp
--- a/Engine/Source/Editor/BlueprintGraph/Private/K2Node_Message.cpp
+++ b/Engine/Source/Editor/BlueprintGraph/Private/K2Node_Message.cpp
@@ -112,10 +112,10 @@ void UK2Node_Message::ExpandNode(class FKismetCompilerContext& CompilerContext,
 			CastToInterfaceResultPin->PinType.PinSubCategoryObject = *CastToInterfaceNode->TargetType;
 
 			// Wire up the connections
-			CompilerContext.CheckConnectionResponse(Schema->MovePinLinks(*ExecPin, *CastToInterfaceNode->GetExecPin()), this);
+			CompilerContext.MovePinLinksToIntermediate(*ExecPin, *CastToInterfaceNode->GetExecPin());
 
 			UEdGraphPin* CastToInterfaceSourceObjectPin = CastToInterfaceNode->GetCastSourcePin();
-			CompilerContext.CheckConnectionResponse(Schema->MovePinLinks(*MessageSelfPin, *CastToInterfaceSourceObjectPin), this);
+			CompilerContext.MovePinLinksToIntermediate(*MessageSelfPin, *CastToInterfaceSourceObjectPin);
 
 			// Next, create the function call node
 			UK2Node_CallFunction* FunctionCallNode = CompilerContext.SpawnIntermediateNode<UK2Node_CallFunction>(this, SourceGraph);
@@ -147,7 +147,7 @@ void UK2Node_Message::ExpandNode(class FKismetCompilerContext& CompilerContext,
 					if( FunctionCallPin )
 					{
 						// Move pin links if the pin is connected...
-						CompilerContext.CheckConnectionResponse(Schema->MovePinLinks(*CurrentPin, *FunctionCallPin), this);
+						CompilerContext.MovePinLinksToIntermediate(*CurrentPin, *FunctionCallPin);
 
 						// when cast fails all return values must be cleared.
 						if (EEdGraphPinDirection::EGPD_Output == CurrentPin->Direction)
@@ -197,24 +197,12 @@ void UK2Node_Message::ExpandNode(class FKismetCompilerContext& CompilerContext,
 			if( bThenPinConnected )
 			{
 				// Failure case for the cast runs straight through to the exit
-				CompilerContext.CheckConnectionResponse(Schema->MovePinLinks(*ThenPin, *LastOutCastFaildPin), this);
+				CompilerContext.CopyPinLinksToIntermediate(*ThenPin, *LastOutCastFaildPin);
 
 				// Copy all links from the invalid cast case above to the call function node
 				if (UEdGraphPin* CallFunctionThenPin = Schema->FindExecutionPin(*FunctionCallNode, EGPD_Output))
 				{
-					for( int32 i = 0; i < LastOutCastFaildPin->LinkedTo.Num(); i++ )
-					{
-						UEdGraphPin* Pin = LastOutCastFaildPin->LinkedTo[i];
-						FPinConnectionResponse Response = Schema->CanCreateConnection(CallFunctionThenPin, Pin);
-						if(Response.CanSafeConnect())
-						{
-							CallFunctionThenPin->MakeLinkTo(Pin);
-						}
-						else
-						{
-							CompilerContext.CheckConnectionResponse(Response, this);
-						}
-					}
+					CompilerContext.MovePinLinksToIntermediate(*ThenPin, *CallFunctionThenPin);
 				}
 				else
 				{
diff --git a/Engine/Source/Editor/BlueprintGraph/Private/K2Node_MultiGate.cpp b/Engine/Source/Editor/BlueprintGraph/Private/K2Node_MultiGate.cpp
--- a/Engine/Source/Editor/BlueprintGraph/Private/K2Node_MultiGate.cpp
+++ b/Engine/Source/Editor/BlueprintGraph/Private/K2Node_MultiGate.cpp
@@ -649,7 +649,7 @@ void UK2Node_MultiGate::ExpandNode(class FKismetCompilerContext& CompilerContext
 			AssignmentNode->GetValuePin()->DefaultValue = TEXT("0");
 
 			// Move the "Reset" link to the Assignment node
-			CompilerContext.CheckConnectionResponse(Schema->MovePinLinks(*ResetPin, *AssignmentNode->GetExecPin()), this);
+			CompilerContext.MovePinLinksToIntermediate(*ResetPin, *AssignmentNode->GetExecPin());
 		}
 	}
 }
diff --git a/Engine/Source/Editor/BlueprintGraph/Private/K2Node_SpawnActor.cpp b/Engine/Source/Editor/BlueprintGraph/Private/K2Node_SpawnActor.cpp
--- a/Engine/Source/Editor/BlueprintGraph/Private/K2Node_SpawnActor.cpp
+++ b/Engine/Source/Editor/BlueprintGraph/Private/K2Node_SpawnActor.cpp
@@ -335,12 +335,12 @@ void UK2Node_SpawnActor::ExpandNode(class FKismetCompilerContext& CompilerContex
 		UEdGraphPin* CallBeginResult = CallBeginSpawnNode->GetReturnValuePin();
 
 		// Move 'exec' connection from spawn node to 'begin spawn'
-		CompilerContext.CheckConnectionResponse(Schema->MovePinLinks(*SpawnNodeExec, *CallBeginExec), this);
+		CompilerContext.MovePinLinksToIntermediate(*SpawnNodeExec, *CallBeginExec);
 
 		if(SpawnBlueprintPin->LinkedTo.Num() > 0)
 		{
 			// Copy the 'blueprint' connection from the spawn node to 'begin spawn'
-			CompilerContext.CheckConnectionResponse(Schema->MovePinLinks(*SpawnBlueprintPin, *CallBeginBlueprintPin), this);
+			CompilerContext.MovePinLinksToIntermediate(*SpawnBlueprintPin, *CallBeginBlueprintPin);
 		}
 		else
 		{
@@ -351,14 +351,14 @@ void UK2Node_SpawnActor::ExpandNode(class FKismetCompilerContext& CompilerContex
 		// Copy the world context connection from the spawn node to 'begin spawn' if necessary
 		if (SpawnWorldContextPin)
 		{
-			CompilerContext.CheckConnectionResponse(Schema->MovePinLinks(*SpawnWorldContextPin, *CallBeginWorldContextPin), this);
+			CompilerContext.MovePinLinksToIntermediate(*SpawnWorldContextPin, *CallBeginWorldContextPin);
 		}
 
 		// Copy the 'transform' connection from the spawn node to 'begin spawn'
-		CompilerContext.CheckConnectionResponse(Schema->MovePinLinks(*SpawnNodeTransform, *CallBeginTransform), this);
+		CompilerContext.MovePinLinksToIntermediate(*SpawnNodeTransform, *CallBeginTransform);
 		
 		// Copy the 'bNoCollisionFail' connection from the spawn node to 'begin spawn'
-		CompilerContext.CheckConnectionResponse(Schema->MovePinLinks(*SpawnNodeNoCollisionFail, *CallBeginNoCollisionFail), this);
+		CompilerContext.MovePinLinksToIntermediate(*SpawnNodeNoCollisionFail, *CallBeginNoCollisionFail);
 
 		//////////////////////////////////////////////////////////////////////////
 		// create 'finish spawn' call node
@@ -373,17 +373,17 @@ void UK2Node_SpawnActor::ExpandNode(class FKismetCompilerContext& CompilerContex
 		UEdGraphPin* CallFinishResult = CallFinishSpawnNode->GetReturnValuePin();
 
 		// Move 'then' connection from spawn node to 'finish spawn'
-		CompilerContext.CheckConnectionResponse(Schema->MovePinLinks(*SpawnNodeThen, *CallFinishThen), this);
+		CompilerContext.MovePinLinksToIntermediate(*SpawnNodeThen, *CallFinishThen);
 			
 		// Copy transform connection
-		CompilerContext.CheckConnectionResponse(Schema->CopyPinLinks(*CallBeginTransform, *CallFinishTransform), this);
+		CompilerContext.CopyPinLinksToIntermediate(*CallBeginTransform, *CallFinishTransform);
 		
 		// Connect output actor from 'begin' to 'finish'
 		CallBeginResult->MakeLinkTo(CallFinishActor);
 
 		// Move result connection from spawn node to 'finish spawn'
 		CallFinishResult->PinType = SpawnNodeResult->PinType; // Copy type so it uses the right actor subclass
-		CompilerContext.CheckConnectionResponse(Schema->MovePinLinks(*SpawnNodeResult, *CallFinishResult), this);
+		CompilerContext.MovePinLinksToIntermediate(*SpawnNodeResult, *CallFinishResult);
 
 		//////////////////////////////////////////////////////////////////////////
 		// create 'set var' nodes
@@ -428,7 +428,7 @@ void UK2Node_SpawnActor::ExpandNode(class FKismetCompilerContext& CompilerContex
 
 					// Move connection from the variable pin on the spawn node to the 'value' pin
 					UEdGraphPin* ValuePin = SetVarNode->FindPinChecked(ValueParamName);
-					CompilerContext.CheckConnectionResponse(Schema->MovePinLinks(*SpawnVarPin, *ValuePin), this);
+					CompilerContext.MovePinLinksToIntermediate(*SpawnVarPin, *ValuePin);
 					if(SpawnVarPin->PinType.bIsArray)
 					{
 						SetVarNode->PinConnectionListChanged(ValuePin);
diff --git a/Engine/Source/Editor/BlueprintGraph/Private/K2Node_SpawnActorFromClass.cpp b/Engine/Source/Editor/BlueprintGraph/Private/K2Node_SpawnActorFromClass.cpp
--- a/Engine/Source/Editor/BlueprintGraph/Private/K2Node_SpawnActorFromClass.cpp
+++ b/Engine/Source/Editor/BlueprintGraph/Private/K2Node_SpawnActorFromClass.cpp
@@ -330,12 +330,12 @@ void UK2Node_SpawnActorFromClass::ExpandNode(class FKismetCompilerContext& Compi
 		UEdGraphPin* CallBeginResult = CallBeginSpawnNode->GetReturnValuePin();
 
 		// Move 'exec' connection from spawn node to 'begin spawn'
-		CompilerContext.CheckConnectionResponse(Schema->MovePinLinks(*SpawnNodeExec, *CallBeginExec), SpawnNode);
+		CompilerContext.MovePinLinksToIntermediate(*SpawnNodeExec, *CallBeginExec);
 
 		if(SpawnClassPin->LinkedTo.Num() > 0)
 		{
 			// Copy the 'blueprint' connection from the spawn node to 'begin spawn'
-			CompilerContext.CheckConnectionResponse(Schema->MovePinLinks(*SpawnClassPin, *CallBeginActorClassPin), SpawnNode);
+			CompilerContext.MovePinLinksToIntermediate(*SpawnClassPin, *CallBeginActorClassPin);
 		}
 		else
 		{
@@ -346,14 +346,14 @@ void UK2Node_SpawnActorFromClass::ExpandNode(class FKismetCompilerContext& Compi
 		// Copy the world context connection from the spawn node to 'begin spawn' if necessary
 		if (SpawnWorldContextPin)
 		{
-			CompilerContext.CheckConnectionResponse(Schema->MovePinLinks(*SpawnWorldContextPin, *CallBeginWorldContextPin), SpawnNode);
+			CompilerContext.MovePinLinksToIntermediate(*SpawnWorldContextPin, *CallBeginWorldContextPin);
 		}
 
 		// Copy the 'transform' connection from the spawn node to 'begin spawn'
-		CompilerContext.CheckConnectionResponse(Schema->MovePinLinks(*SpawnNodeTransform, *CallBeginTransform), SpawnNode);
+		CompilerContext.MovePinLinksToIntermediate(*SpawnNodeTransform, *CallBeginTransform);
 
 		// Copy the 'bNoCollisionFail' connection from the spawn node to 'begin spawn'
-		CompilerContext.CheckConnectionResponse(Schema->MovePinLinks(*SpawnNodeNoCollisionFail, *CallBeginNoCollisionFail), this);
+		CompilerContext.MovePinLinksToIntermediate(*SpawnNodeNoCollisionFail, *CallBeginNoCollisionFail);
 
 		//////////////////////////////////////////////////////////////////////////
 		// create 'finish spawn' call node
@@ -368,17 +368,17 @@ void UK2Node_SpawnActorFromClass::ExpandNode(class FKismetCompilerContext& Compi
 		UEdGraphPin* CallFinishResult = CallFinishSpawnNode->GetReturnValuePin();
 
 		// Move 'then' connection from spawn node to 'finish spawn'
-		CompilerContext.CheckConnectionResponse(Schema->MovePinLinks(*SpawnNodeThen, *CallFinishThen), SpawnNode);
+		CompilerContext.MovePinLinksToIntermediate(*SpawnNodeThen, *CallFinishThen);
 
 		// Copy transform connection
-		CompilerContext.CheckConnectionResponse(Schema->CopyPinLinks(*CallBeginTransform, *CallFinishTransform), SpawnNode);
+		CompilerContext.CopyPinLinksToIntermediate(*CallBeginTransform, *CallFinishTransform);
 
 		// Connect output actor from 'begin' to 'finish'
 		CallBeginResult->MakeLinkTo(CallFinishActor);
 
 		// Move result connection from spawn node to 'finish spawn'
 		CallFinishResult->PinType = SpawnNodeResult->PinType; // Copy type so it uses the right actor subclass
-		CompilerContext.CheckConnectionResponse(Schema->MovePinLinks(*SpawnNodeResult, *CallFinishResult), SpawnNode);
+		CompilerContext.MovePinLinksToIntermediate(*SpawnNodeResult, *CallFinishResult);
 
 		//////////////////////////////////////////////////////////////////////////
 		// create 'set var' nodes
@@ -423,7 +423,7 @@ void UK2Node_SpawnActorFromClass::ExpandNode(class FKismetCompilerContext& Compi
 
 					// Move connection from the variable pin on the spawn node to the 'value' pin
 					UEdGraphPin* ValuePin = SetVarNode->FindPinChecked(ValueParamName);
-					CompilerContext.CheckConnectionResponse(Schema->MovePinLinks(*SpawnVarPin, *ValuePin), SpawnNode);
+					CompilerContext.MovePinLinksToIntermediate(*SpawnVarPin, *ValuePin);
 					if(SpawnVarPin->PinType.bIsArray)
 					{
 						SetVarNode->PinConnectionListChanged(ValuePin);
diff --git a/Engine/Source/Editor/BlueprintGraph/Private/K2Node_Switch.cpp b/Engine/Source/Editor/BlueprintGraph/Private/K2Node_Switch.cpp
--- a/Engine/Source/Editor/BlueprintGraph/Private/K2Node_Switch.cpp
+++ b/Engine/Source/Editor/BlueprintGraph/Private/K2Node_Switch.cpp
@@ -120,8 +120,7 @@ public:
 					IfFailTest_SucceedAtBeingEqualGoto.Type = KCST_GotoIfNot;
 					IfFailTest_SucceedAtBeingEqualGoto.LHS = BoolTerm;
 
-					UEdGraphNode* ConnectedCaseStatementNode = (Pin->LinkedTo.Num() > 0) ? Pin->LinkedTo[0]->GetOwningNode() : NULL;
-					Context.GotoFixupRequestMap.Add(&IfFailTest_SucceedAtBeingEqualGoto, ConnectedCaseStatementNode);
+					Context.GotoFixupRequestMap.Add(&IfFailTest_SucceedAtBeingEqualGoto, Pin);
 				}
 			}
 
diff --git a/Engine/Source/Editor/BlueprintGraph/Private/K2Node_Timeline.cpp b/Engine/Source/Editor/BlueprintGraph/Private/K2Node_Timeline.cpp
--- a/Engine/Source/Editor/BlueprintGraph/Private/K2Node_Timeline.cpp
+++ b/Engine/Source/Editor/BlueprintGraph/Private/K2Node_Timeline.cpp
@@ -470,7 +470,7 @@ void UK2Node_Timeline::ExpandForPin(UEdGraphPin* TimelinePin, const FName Proper
 		UEdGraphPin* ValuePin = GetVarNode->GetValuePin();
 		if (NULL != ValuePin)
 		{
-			CompilerContext.CheckConnectionResponse(CompilerContext.GetSchema()->MovePinLinks(*TimelinePin, *ValuePin), this);
+			CompilerContext.MovePinLinksToIntermediate(*TimelinePin, *ValuePin);
 		}
 		else
 		{
diff --git a/Engine/Source/Editor/GameplayTagsEditor/Private/GameplayTagsK2Node_LiteralGameplayTag.cpp b/Engine/Source/Editor/GameplayTagsEditor/Private/GameplayTagsK2Node_LiteralGameplayTag.cpp
--- a/Engine/Source/Editor/GameplayTagsEditor/Private/GameplayTagsK2Node_LiteralGameplayTag.cpp
+++ b/Engine/Source/Editor/GameplayTagsEditor/Private/GameplayTagsK2Node_LiteralGameplayTag.cpp
@@ -132,7 +132,7 @@ void UGameplayTagsK2Node_LiteralGameplayTag::ExpandNode(class FKismetCompilerCon
 		if( OutPin && TagOutPin )
 		{
 			OutPin->PinType = TagOutPin->PinType; // Copy type so it uses the right actor subclass
-			CompilerContext.CheckConnectionResponse(Schema->MovePinLinks(*TagOutPin, *OutPin), this);
+			CompilerContext.MovePinLinksToIntermediate(*TagOutPin, *OutPin);
 		}
 
 		// Break any links to the expanded node
diff --git a/Engine/Source/Editor/GraphEditor/Private/ConnectionDrawingPolicy.cpp b/Engine/Source/Editor/GraphEditor/Private/ConnectionDrawingPolicy.cpp
--- a/Engine/Source/Editor/GraphEditor/Private/ConnectionDrawingPolicy.cpp
+++ b/Engine/Source/Editor/GraphEditor/Private/ConnectionDrawingPolicy.cpp
@@ -9,6 +9,8 @@
 #include "SoundDefinitions.h"
 #include "AnimationStateNodes/SGraphNodeAnimTransition.h"
 
+DEFINE_LOG_CATEGORY_STATIC(LogConnectionDrawingPolicy, Log, All);
+
 //////////////////////////////////////////////////////////////////////////
 // FGeometryHelper
 
@@ -595,6 +597,7 @@ void FKismetConnectionDrawingPolicy::BuildExecutionRoadmap()
 
 	TArray<UEdGraphNode*> SequentialNodesInGraph;
 	TArray<double> SequentialNodeTimes;
+	TArray<UEdGraphPin*> SequentialExecPinsInGraph;
 
 	{
 		const TSimpleRingBuffer<FKismetTraceSample>& TraceStack = FKismetDebugUtilities::GetTraceStack();
@@ -610,12 +613,15 @@ void FKismetConnectionDrawingPolicy::BuildExecutionRoadmap()
 			{
 				if (TestObject == ActiveObject)
 				{
+					UEdGraphPin* AssociatedPin = DebugData.FindExecPinFromCodeLocation(Sample.Function.Get(), Sample.Offset);
+
 					if (UEdGraphNode* Node = DebugData.FindSourceNodeFromCodeLocation(Sample.Function.Get(), Sample.Offset, /*bAllowImpreciseHit=*/ false))
 					{
 						if (GraphObj == Node->GetGraph())
 						{
 							SequentialNodesInGraph.Add(Node);
 							SequentialNodeTimes.Add(Sample.ObservationTime);
+							SequentialExecPinsInGraph.Add(AssociatedPin);
 						}
 						else
 						{
@@ -633,6 +639,7 @@ void FKismetConnectionDrawingPolicy::BuildExecutionRoadmap()
 										// Add it to the sequential node list
 										SequentialNodesInGraph.Add(MacroSourceNode);
 										SequentialNodeTimes.Add(Sample.ObservationTime);
+										SequentialExecPinsInGraph.Add(AssociatedPin);
 									}
 									else
 									{
@@ -651,6 +658,7 @@ void FKismetConnectionDrawingPolicy::BuildExecutionRoadmap()
 												// Add it to the sequential node list
 												SequentialNodesInGraph.Add(MacroInstanceNode);
 												SequentialNodeTimes.Add(Sample.ObservationTime);
+												SequentialExecPinsInGraph.Add(AssociatedPin);
 
 												// Exit the loop; we're done
 												break;
@@ -679,23 +687,59 @@ void FKismetConnectionDrawingPolicy::BuildExecutionRoadmap()
 		LatestTimeDiscovered = FMath::Max<double>(LatestTimeDiscovered, ObservationTime);
 	}
 
-	// Record the unique node->node pairings, keeping only the most recent times for each pairing
+	UEdGraphPin* LastExecPin = NULL;
+	// Record the unique exec-pin to time pairings, keeping only the most recent 
+	// times for each pairing... reverse the "SequentialNodes" because right now
+	// it is in stack order (with the last executed node first)
 	for (int32 i = SequentialNodesInGraph.Num() - 1; i >= 1; --i)
 	{
-		UEdGraphNode* CurNode = SequentialNodesInGraph[i];
-		double CurNodeTime = SequentialNodeTimes[i];
+		UEdGraphNode* CurNode  = SequentialNodesInGraph[i];
 		UEdGraphNode* NextNode = SequentialNodesInGraph[i-1];
-		double NextNodeTime = SequentialNodeTimes[i-1];
-
-		FExecPairingMap& Predecessors = PredecessorNodes.FindOrAdd(NextNode);
 
-		// Update the timings if this is a more recent pairing
-		FTimePair& Timings = Predecessors.FindOrAdd(CurNode);
-		if (Timings.ThisExecTime < NextNodeTime)
+		// keep track of the last exec-pin executed by CurNode (these tracked 
+		// pins coincide with "WireTraceSite" op-codes that have been injected 
+		// prior to every "goto" statement... this way we have context for which
+		// pin executed the jump)
+		if (UEdGraphPin* AssociatedPin = SequentialExecPinsInGraph[i])
 		{
-			Timings.PredExecTime = CurNodeTime;
-			Timings.ThisExecTime = NextNodeTime;
+			LastExecPin = AssociatedPin;
+		}
+		
+		// if this statement is a jump (from one node to another)
+		if (CurNode != NextNode)
+		{
+			// if there was a wire-trace op-code inserted before this jump
+			if (LastExecPin != NULL)
+			{
+				//ensure(LastExecPin->GetOwningNode() == CurNode);
+				double NextNodeTime = SequentialNodeTimes[i-1];
+
+				FExecPairingMap& ExecPaths  = PredecessorPins.FindOrAdd(NextNode);
+				FTimePair&       ExecTiming = ExecPaths.FindOrAdd(LastExecPin);
+				// make sure that if we've already visited this exec-pin (like 
+				// in a for-loop or something), that we're replacing it with a 
+				// more recent execution time
+				//
+				// @TODO I don't see when this wouldn't be the case
+				if (ExecTiming.ThisExecTime < NextNodeTime)
+				{
+					double CurNodeTime = SequentialNodeTimes[i];
+					ExecTiming.ThisExecTime = NextNodeTime;
+					ExecTiming.PredExecTime = CurNodeTime;
+				}
+			}
+			else 
+			{
+				UE_LOG(LogConnectionDrawingPolicy, Warning, TEXT("Looks like a wire-trace was not injected before the jump from '%s' to '%s'."), 
+					*CurNode->GetNodeTitle(ENodeTitleType::FullTitle), *NextNode->GetNodeTitle(ENodeTitleType::FullTitle));
+			}
+
+			// clear the exec-pin (we're moving to a new node and want to find 
+			// it's executed out pin)
+			LastExecPin = NULL;
 		}
+		// else, we're only collecting this data for tracing node-to-node
+		// executions (so we don't care about this sequence of statements)
 	}
 
 	// Fade only when free-running (since we're using GCurrentTime, instead of FPlatformTime::Seconds)
@@ -770,19 +814,21 @@ void FKismetConnectionDrawingPolicy::DetermineWiringStyle(UEdGraphPin* OutputPin
 		{
 			if (CanBuildRoadmap())
 			{
+				// track if this node connection was ran or not
 				bool bExecuted = false;
 
-				// Run thru the predecessors, and on
-				if (FExecPairingMap* PredecessorMap = PredecessorNodes.Find(InputPin->GetOwningNode()))
+				UEdGraphNode* InputNode = InputPin->GetOwningNode();
+				// if the node belonging to InputPin was actually executed
+				if (FExecPairingMap* ExecPaths = PredecessorPins.Find(InputNode))
 				{
-					if (FTimePair* Times = PredecessorMap->Find(OutputPin->GetOwningNode()))
+					// if the output pin is one of the pins that lead to InputNode being ran
+					if (FTimePair* ExecTiming = ExecPaths->Find(OutputPin))
 					{
 						bExecuted = true;
-
-						DetermineStyleOfExecWire(/*inout*/ Thickness, /*inout*/ WireColor, /*inout*/ bDrawBubbles, *Times);
+						DetermineStyleOfExecWire(/*inout*/ Thickness, /*inout*/ WireColor, /*inout*/ bDrawBubbles, *ExecTiming);
 					}
 				}
-
+				
 				if (!bExecuted)
 				{
 					// It's not followed, fade it and keep it thin
@@ -1071,11 +1117,27 @@ void FAnimGraphConnectionDrawingPolicy::BuildExecutionRoadmap()
 			{
 				if (UAnimGraphNode_Base* TargetNode = Cast<UAnimGraphNode_Base>(PropertySourceMap.FindRef(AnimBlueprintClass->AnimNodeProperties[VisitRecord.TargetID])))
 				{
-					//@TODO: Extend the rendering code to allow using the recorded blend weight instead of faked exec times
-					FExecPairingMap& Predecessors = PredecessorNodes.FindOrAdd((UEdGraphNode*)SourceNode);
-					FTimePair& Timings = Predecessors.FindOrAdd((UEdGraphNode*)TargetNode);
-					Timings.PredExecTime = 0.0;
-					Timings.ThisExecTime = VisitRecord.Weight;
+					UEdGraphPin* PoseNet = NULL;
+
+					UAnimationGraphSchema const* AnimSchema = GetDefault<UAnimationGraphSchema>();
+					for (int32 PinIndex = 0; PinIndex < TargetNode->Pins.Num(); ++PinIndex)
+					{
+						UEdGraphPin* Pin = TargetNode->Pins[PinIndex];
+						if (AnimSchema->IsPosePin(Pin->PinType) && (Pin->Direction == EGPD_Output))
+						{
+							PoseNet = Pin;
+							break;
+						}
+					}
+
+					if (PoseNet != NULL)
+					{
+						//@TODO: Extend the rendering code to allow using the recorded blend weight instead of faked exec times
+						FExecPairingMap& Predecessors = PredecessorPins.FindOrAdd((UEdGraphNode*)SourceNode);
+						FTimePair& Timings = Predecessors.FindOrAdd(PoseNet);
+						Timings.PredExecTime = 0.0;
+						Timings.ThisExecTime = VisitRecord.Weight;
+					}
 				}
 			}
 		}
diff --git a/Engine/Source/Editor/GraphEditor/Private/ConnectionDrawingPolicy.h b/Engine/Source/Editor/GraphEditor/Private/ConnectionDrawingPolicy.h
--- a/Engine/Source/Editor/GraphEditor/Private/ConnectionDrawingPolicy.h
+++ b/Engine/Source/Editor/GraphEditor/Private/ConnectionDrawingPolicy.h
@@ -92,8 +92,8 @@ protected:
 	// Times for one execution pair within the current graph
 	struct FTimePair
 	{
-		double PredExecTime;
-		double ThisExecTime;
+		double PredExecTime; // time that the node with the associated exec-out was executed
+		double ThisExecTime; // time that the following node was executed
 
 		FTimePair()
 			: PredExecTime(0.0)
@@ -101,12 +101,15 @@ protected:
 		{
 		}
 	};
-
-	// Map of pairings
-	typedef TMap<UEdGraphNode*, FTimePair> FExecPairingMap;
-
-	// Map of nodes that preceeded before a given node in the execution sequence (one entry for each pairing)
-	TMap<UEdGraphNode*, FExecPairingMap> PredecessorNodes;
+	/** 
+	 * Map of exec-pins to a pair of execution times (holds timing info for the 
+	 * owning-node, paired with a time of execution for the linked node)... only 
+	 * one timing-pair per pin because exec outputs can only connect to one node
+	 */
+	typedef TMap<UEdGraphPin*, FTimePair> FExecPairingMap;
+
+	// Map of executed nodes, and the execution pins that lead to them being ran
+	TMap<UEdGraphNode*, FExecPairingMap> PredecessorPins;
 
 	UEdGraph* GraphObj;
 	double CurrentTime;
diff --git a/Engine/Source/Editor/KismetCompiler/Private/KismetCompiler.cpp b/Engine/Source/Editor/KismetCompiler/Private/KismetCompiler.cpp
--- a/Engine/Source/Editor/KismetCompiler/Private/KismetCompiler.cpp
+++ b/Engine/Source/Editor/KismetCompiler/Private/KismetCompiler.cpp
@@ -1321,6 +1321,12 @@ void FKismetCompilerContext::PostcompileFunction(FKismetFunctionContext& Context
 	// Sort the 'linear execution list' again by likely execution order.
 	Context.FinalSortLinearExecList();
 
+// 	if (Context.GetBlueprint()->GetLinker()->IsCooking())
+// 	{
+// 		int i = 0;
+// 		++i;
+// 	}
+
 	// Resolve goto links
 	Context.ResolveGotoFixups();
 
@@ -1632,7 +1638,7 @@ void FKismetCompilerContext::CreatePinEventNodeForTimelineFunction(UK2Node_Timel
 
 	if(UpdatePin != NULL && UpdateOutput != NULL)
 	{
-		CheckConnectionResponse(Schema->MovePinLinks(*UpdatePin, *UpdateOutput), TimelineNode);
+		MovePinLinksToIntermediate(*UpdatePin, *UpdateOutput);
 	}
 }
 
@@ -1649,7 +1655,7 @@ UK2Node_CallFunction* FKismetCompilerContext::CreateCallTimelineFunction(UK2Node
 
 	// Move any exec links from 'play' pin to the 'call play' node
 	UEdGraphPin* CallExecInput = Schema->FindExecutionPin(*CallNode, EGPD_Input);
-	CheckConnectionResponse(Schema->MovePinLinks(*TimelineFunctionPin, *CallExecInput), TimelineNode);
+	MovePinLinksToIntermediate(*TimelineFunctionPin, *CallExecInput);
 	return CallNode;
 }
 
@@ -1749,7 +1755,7 @@ void FKismetCompilerContext::ExpandTimelineNodes(UEdGraph* SourceGraph)
 						if (CallNode && NewTimePin)
 						{
 							UEdGraphPin* InputPin = CallNode->FindPinChecked(TEXT("NewTime"));
-							CheckConnectionResponse(Schema->MovePinLinks(*NewTimePin, *InputPin), CallNode);
+							MovePinLinksToIntermediate(*NewTimePin, *InputPin);
 						}
 					}
 				}
@@ -1934,7 +1940,7 @@ void FKismetCompilerContext::ExpandPlayMovieSceneNodes( UEdGraph* SourceGraph )
 					SequenceNode->AllocateDefaultPins();
 				
 					// Move input links from 'Play' to the exec pin on the Sequence node
-					CheckConnectionResponse( Schema->MovePinLinks( *PlayPin, *SequenceNode->GetExecPin() ), PlayMovieSceneNode );
+					MovePinLinksToIntermediate( *PlayPin, *SequenceNode->GetExecPin() );
 
 					SequenceNode->GetThenPinGivenIndex( 0 )->MakeLinkTo( AllocateRuntimeMovieScenePlayerExecPin );
 
@@ -1961,7 +1967,7 @@ void FKismetCompilerContext::ExpandPlayMovieSceneNodes( UEdGraph* SourceGraph )
 					SequenceNode->AllocateDefaultPins();
 				
 					// Move input links from 'Pause' to the exec pin on the Sequence node
-					CheckConnectionResponse( Schema->MovePinLinks( *PausePin, *SequenceNode->GetExecPin() ), PlayMovieSceneNode );
+					MovePinLinksToIntermediate( *PausePin, *SequenceNode->GetExecPin() );
 
 					SequenceNode->GetThenPinGivenIndex( 0 )->MakeLinkTo( AllocateRuntimeMovieScenePlayerExecPin );
 
@@ -1978,6 +1984,28 @@ void FKismetCompilerContext::ExpandPlayMovieSceneNodes( UEdGraph* SourceGraph )
 	}
 }
 
+FPinConnectionResponse FKismetCompilerContext::MovePinLinksToIntermediate(UEdGraphPin& SourcePin, UEdGraphPin& IntermediatePin)
+{
+	 UEdGraphSchema_K2 const* K2Schema = GetSchema();
+	 FPinConnectionResponse ConnectionResult = K2Schema->MovePinLinks(SourcePin, IntermediatePin);
+
+	 CheckConnectionResponse(ConnectionResult, SourcePin.GetOwningNode());
+	 MessageLog.NotifyIntermediateObjectCreation(&IntermediatePin, &SourcePin);
+
+	 return ConnectionResult;
+}
+
+FPinConnectionResponse FKismetCompilerContext::CopyPinLinksToIntermediate(UEdGraphPin& SourcePin, UEdGraphPin& IntermediatePin)
+{
+	UEdGraphSchema_K2 const* K2Schema = GetSchema();
+	FPinConnectionResponse ConnectionResult = K2Schema->CopyPinLinks(SourcePin, IntermediatePin);
+
+	CheckConnectionResponse(ConnectionResult, SourcePin.GetOwningNode());
+	MessageLog.NotifyIntermediateObjectCreation(&IntermediatePin, &SourcePin);
+
+	return ConnectionResult;
+}
+
 UK2Node_TemporaryVariable* FKismetCompilerContext::SpawnInternalVariable(UEdGraphNode* SourceNode, FString Category, FString SubCategory, UObject* SubcategoryObject, bool bIsArray)
 {
 	UK2Node_TemporaryVariable* Result = SpawnIntermediateNode<UK2Node_TemporaryVariable>(SourceNode);
@@ -2589,7 +2617,29 @@ void FKismetCompilerContext::ExpandTunnelsAndMacros(UEdGraph* SourceGraph)
 		}
 		else if (UK2Node_Tunnel* TunnelNode = Cast<UK2Node_Tunnel>(*NodeIt))
 		{
-			bool bSuccess = Schema->CollapseGatewayNode(TunnelNode, TunnelNode->GetInputSink(), TunnelNode->GetOutputSource());
+			UEdGraphNode* InputSink = TunnelNode->GetInputSink();
+			for (UEdGraphPin const* TunnelPin : TunnelNode->Pins)
+			{
+				if ((TunnelPin->Direction != EGPD_Input) || (TunnelPin->PinType.PinCategory != Schema->PC_Exec))
+				{
+					continue;
+				}
+				check(InputSink != NULL);
+
+				UEdGraphPin* SinkPin = InputSink->FindPin(TunnelPin->PinName);
+				if (SinkPin == NULL)
+				{
+					continue;
+				}
+				check(SinkPin->Direction == EGPD_Output);
+
+				for (UEdGraphPin* TunnelLinkedPin : TunnelPin->LinkedTo)
+				{
+					MessageLog.NotifyIntermediateObjectCreation(TunnelLinkedPin, SinkPin);
+				}
+			}
+
+			bool bSuccess = Schema->CollapseGatewayNode(TunnelNode, InputSink, TunnelNode->GetOutputSource());
 			if (!bSuccess)
 			{
 				MessageLog.Error(*LOCTEXT("CollapseTunnel_Error", "Failed to collapse tunnel @@").ToString(), TunnelNode);
diff --git a/Engine/Source/Editor/KismetCompiler/Private/KismetCompilerMisc.cpp b/Engine/Source/Editor/KismetCompiler/Private/KismetCompilerMisc.cpp
--- a/Engine/Source/Editor/KismetCompiler/Private/KismetCompilerMisc.cpp
+++ b/Engine/Source/Editor/KismetCompiler/Private/KismetCompilerMisc.cpp
@@ -728,7 +728,7 @@ FBlueprintCompiledStatement& FNodeHandlingFunctor::GenerateSimpleThenGoto(FKisme
 
 	FBlueprintCompiledStatement& GotoStatement = Context.AppendStatementForNode(&Node);
 	GotoStatement.Type = KCST_UnconditionalGoto;
-	Context.GotoFixupRequestMap.Add(&GotoStatement, TargetNode);
+	Context.GotoFixupRequestMap.Add(&GotoStatement, ThenExecPin);
 
 	return GotoStatement;
 }
@@ -854,7 +854,8 @@ FKismetFunctionContext::FKismetFunctionContext(FCompilerResultsLog& InMessageLog
 	, NetFlags(0)
 	, bIsInterfaceStub(false)
 	, bIsConstFunction(false)
-	, bCreateDebugData(true)
+	// only need debug-data when running in the editor app:
+	, bCreateDebugData(GIsEditor && !IsRunningCommandlet())
 	, bIsSimpleStubGraphWithNoParams(false)
 	, SourceEventFromStubGraph(NULL)
 {
diff --git a/Engine/Source/Editor/KismetCompiler/Private/KismetCompilerVMBackend.cpp b/Engine/Source/Editor/KismetCompiler/Private/KismetCompilerVMBackend.cpp
--- a/Engine/Source/Editor/KismetCompiler/Private/KismetCompilerVMBackend.cpp
+++ b/Engine/Source/Editor/KismetCompiler/Private/KismetCompilerVMBackend.cpp
@@ -1060,6 +1060,15 @@ public:
 			Writer << EX_Nothing;
 			break;
 		case KCST_WireTraceSite:
+			{
+				UEdGraphPin const* TrueSourcePin = Cast<UEdGraphPin const>(FunctionContext.MessageLog.FindSourceObject(Statement.ExecContext));
+				if (TrueSourcePin)
+				{
+					int32 Offset = Writer.ScriptBuffer.Num();
+					ClassBeingBuilt->GetDebugData().RegisterPinToCodeAssociation(TrueSourcePin, FunctionContext.Function, Offset);
+				}
+			}
+			// no break, continue down through KCST_DebugSite
 		case KCST_DebugSite:
 			if (SourceNode != NULL)
 			{
diff --git a/Engine/Source/Editor/KismetCompiler/Public/BlueprintCompiledStatement.h b/Engine/Source/Editor/KismetCompiler/Public/BlueprintCompiledStatement.h
--- a/Engine/Source/Editor/KismetCompiler/Public/BlueprintCompiledStatement.h
+++ b/Engine/Source/Editor/KismetCompiler/Public/BlueprintCompiledStatement.h
@@ -47,6 +47,7 @@ struct FBlueprintCompiledStatement
 		, bIsJumpTarget(false)
 		, bIsInterfaceContext(false)
 		, bIsParentContext(false)
+		, ExecContext(NULL)
 	{
 	}
 
@@ -82,6 +83,9 @@ struct FBlueprintCompiledStatement
 	// Is this function called on a parent class (super, etc)?  (KCST_CallFunction)
 	bool bIsParentContext;
 
+	// Exec pin about to execute (KCST_WireTraceSite)
+	class UEdGraphPin* ExecContext;
+
 	// Comment text
 	FString Comment;
 };
diff --git a/Engine/Source/Editor/KismetCompiler/Public/KismetCompiledFunctionContext.h b/Engine/Source/Editor/KismetCompiler/Public/KismetCompiledFunctionContext.h
--- a/Engine/Source/Editor/KismetCompiler/Public/KismetCompiledFunctionContext.h
+++ b/Engine/Source/Editor/KismetCompiler/Public/KismetCompiledFunctionContext.h
@@ -45,14 +45,14 @@ public:
 	FCompilerResultsLog& MessageLog;
 	UEdGraphSchema_K2* Schema;
 
-	//
+	// An UNORDERED listing of all statements (used for cleaning up the dynamically allocated statements)
 	TArray< FBlueprintCompiledStatement* > AllGeneratedStatements;
 
 	// Individual execution lists for every node that generated code to be consumed by the backend
 	TMap< UEdGraphNode*, TArray<FBlueprintCompiledStatement*> > StatementsPerNode;
 
-	// Goto fixup requests (each statement (key) wants to goto the first statement created for the node (value))
-	TMap< FBlueprintCompiledStatement*, UEdGraphNode* > GotoFixupRequestMap;
+	// Goto fixup requests (each statement (key) wants to goto the first statement attached to the exec out-pin (value))
+	TMap< FBlueprintCompiledStatement*, UEdGraphPin* > GotoFixupRequestMap;
 
 	// Map from a net to an term (either a literal or a storage location)
 	TIndirectArray<FBPTerminal> Parameters;
@@ -325,8 +325,15 @@ public:
 
 					if (NextNode != NULL)
 					{
+						UEdGraphNode* JumpTargetNode = NULL;
+
+						UEdGraphPin const* JumpPin = GotoFixupRequestMap.FindChecked(&LastStatementInCurrentNode);
+						if ((JumpPin != NULL) && (JumpPin->LinkedTo.Num() > 0))
+						{
+							JumpTargetNode = JumpPin->LinkedTo[0]->GetOwningNode();
+						}
+
 						// Is this node the target of the last statement in the previous node?
-						UEdGraphNode* JumpTargetNode = GotoFixupRequestMap.FindChecked(&LastStatementInCurrentNode);
 						if (NextNode == JumpTargetNode)
 						{
 							// Jump and label are adjacent, can optimize it out
@@ -339,19 +346,115 @@ public:
 		}
 	}
 
+	/**
+	 * Makes sure an KCST_WireTraceSite is inserted before the specified 
+	 * statement, and associates the specified pin with the inserted wire-trace 
+	 * (so we can backwards engineer which pin triggered the goto).
+	 * 
+	 * @param  GotoStatement		The statement to insert a goto before.
+	 * @param  AssociatedExecPin	The pin responsible for executing the goto.
+	 */
+	void InsertWireTrace(FBlueprintCompiledStatement* GotoStatement, UEdGraphPin* AssociatedExecPin)
+	{
+		// only need wire traces if we're debugging the blueprint is available (not for cooked builds)
+		if (bCreateDebugData && (AssociatedExecPin != NULL))
+		{
+			UEdGraphNode* PreJumpNode = AssociatedExecPin->GetOwningNode();
+
+			TArray<FBlueprintCompiledStatement*>* NodeStatementList = StatementsPerNode.Find(PreJumpNode);
+			if (NodeStatementList != NULL)
+			{
+				// @TODO: this Find() is potentially costly (if the node initially generated a lot of statements)
+				int32 GotoIndex = NodeStatementList->Find(GotoStatement);
+				if (GotoIndex != INDEX_NONE)
+				{
+					FBlueprintCompiledStatement* PrevStatement = NULL;
+					if (GotoIndex > 0)
+					{
+						PrevStatement = (*NodeStatementList)[GotoIndex - 1];
+					}
+
+					// if a wire-trace has already been inserted for us
+					if ((PrevStatement != NULL) && (PrevStatement->Type == KCST_WireTraceSite))
+					{
+						PrevStatement->ExecContext = AssociatedExecPin;
+					}
+					else if (PrevStatement != NULL)
+					{
+						FBlueprintCompiledStatement* TraceStatement = new FBlueprintCompiledStatement();
+						TraceStatement->Type        = KCST_WireTraceSite;
+						TraceStatement->Comment     = PreJumpNode->NodeComment.IsEmpty() ? PreJumpNode->GetName() : PreJumpNode->NodeComment;
+						TraceStatement->ExecContext = AssociatedExecPin;
+
+						NodeStatementList->Insert(TraceStatement, GotoIndex);
+						// AllGeneratedStatements is an unordered list, so it doesn't matter that we throw this at the end
+						AllGeneratedStatements.Add(TraceStatement);
+					}
+				}
+			}
+		}
+	}
+
 	/** Resolves all pending goto fixups; Should only be called after all nodes have had a chance to generate code! */
 	void ResolveGotoFixups()
 	{
+		/**
+		 * This local lambda is intended to save on code duplication within this
+		 * function. It tests the supplied statement to see if it is an event 
+		 * entry point.
+		 *
+		 * @param GotoStatement		The compiled statement you wish to test.
+		 * @return True if the passed statement is an event entry point, false if not.
+		 */
+		auto IsUberGraphEventStatement = [](FBlueprintCompiledStatement* GotoStatement) -> bool 
+		{
+			return ((GotoStatement->Type == KCST_CallFunction) && (GotoStatement->UbergraphCallIndex == 0));
+		};
+
+		if (bCreateDebugData)
+		{
+			// if we're debugging, go through an insert a wire trace before  
+			// every "goto" statement so we can trace what execution pin a node
+			// was executed from
+			for (auto GotoIt = GotoFixupRequestMap.CreateIterator(); GotoIt; ++GotoIt)
+			{
+				FBlueprintCompiledStatement* GotoStatement = GotoIt.Key();
+				if (IsUberGraphEventStatement(GotoStatement))
+				{
+					continue;
+				}
+
+				InsertWireTrace(GotoIt.Key(), GotoIt.Value());
+			}
+		}
+
 		if (KismetCompilerDebugOptions::MergeAdjacentStates)
 		{
 			MergeAdjacentStates();
 		}
 
 		// Resolve the remaining fixups
-		for (TMap< FBlueprintCompiledStatement*, UEdGraphNode* >::TIterator It(GotoFixupRequestMap); It; ++It)
+		for (auto GotoIt = GotoFixupRequestMap.CreateIterator(); GotoIt; ++GotoIt)
 		{
-			FBlueprintCompiledStatement* GotoStatement = It.Key();
-			UEdGraphNode* TargetNode = It.Value();
+			FBlueprintCompiledStatement* GotoStatement = GotoIt.Key();
+
+			UEdGraphPin const* ExecNet = GotoIt.Value();
+			if (ExecNet == NULL)
+			{
+				// Execution flow ended here; pop the stack
+				GotoStatement->Type = (GotoStatement->Type == KCST_GotoIfNot) ? KCST_EndOfThreadIfNot : KCST_EndOfThread;
+				continue;
+			}
+
+			UEdGraphNode* TargetNode = NULL;
+			if (IsUberGraphEventStatement(GotoStatement))
+			{
+				TargetNode = ExecNet->GetOwningNode();
+			}
+			else if (ExecNet->LinkedTo.Num() > 0)
+			{
+				TargetNode = ExecNet->LinkedTo[0]->GetOwningNode();
+			}
 
 			if (TargetNode == NULL)
 			{
@@ -370,7 +473,7 @@ public:
 				}
 				else
 				{
-					// Wire up the jump target and notify the target that it is targetted
+					// Wire up the jump target and notify the target that it is targeted
 					FBlueprintCompiledStatement& FirstStatement = *((*StatementList)[0]);
 					GotoStatement->TargetLabel = &FirstStatement;
 					FirstStatement.bIsJumpTarget = true;
diff --git a/Engine/Source/Editor/KismetCompiler/Public/KismetCompiler.h b/Engine/Source/Editor/KismetCompiler/Public/KismetCompiler.h
--- a/Engine/Source/Editor/KismetCompiler/Public/KismetCompiler.h
+++ b/Engine/Source/Editor/KismetCompiler/Public/KismetCompiler.h
@@ -113,6 +113,28 @@ public:
 		return Result;
 	}
 
+	/**
+	 * Moves pin links over from the source-pin to the specified intermediate, 
+	 * and validates the result (additionally logs a redirect from the 
+	 * intermediate-pin back to the source so we can back trace for debugging, etc.)
+	 * 
+	 * @param  SourcePin		The pin you want disconnected.
+	 * @param  IntermediatePin	The pin you want the SourcePin's links moved to.
+	 * @return The result from calling the schema's MovePinLinks().
+	 */
+	KISMETCOMPILER_API FPinConnectionResponse MovePinLinksToIntermediate(UEdGraphPin& SourcePin, UEdGraphPin& IntermediatePin);
+
+	/**
+	 * Copies pin links over from the source-pin to the specified intermediate, 
+	 * and validates the result (additionally logs a redirect from the 
+	 * intermediate-pin back to the source so we can back trace for debugging, etc.)
+	 * 
+	 * @param  SourcePin		The pin whose links you want copied.
+	 * @param  IntermediatePin	The pin you want the SourcePin's links copied to.
+	 * @return The result from calling the schema's CopyPinLinks().
+	 */
+	KISMETCOMPILER_API FPinConnectionResponse CopyPinLinksToIntermediate(UEdGraphPin& SourcePin, UEdGraphPin& IntermediatePin);
+
 	KISMETCOMPILER_API UK2Node_TemporaryVariable* SpawnInternalVariable(UEdGraphNode* SourceNode, FString Category, FString SubCategory = TEXT(""), UObject* SubcategoryObject = NULL, bool bIsArray = false);
 
 	/** Checks a connection response, and errors if it didn't succeed */
diff --git a/Engine/Source/Runtime/Engine/Classes/Engine/BlueprintGeneratedClass.h b/Engine/Source/Runtime/Engine/Classes/Engine/BlueprintGeneratedClass.h
--- a/Engine/Source/Runtime/Engine/Classes/Engine/BlueprintGeneratedClass.h
+++ b/Engine/Source/Runtime/Engine/Classes/Engine/BlueprintGeneratedClass.h
@@ -48,6 +48,9 @@ public:
 	// Reverse map from code offset to macro instance node(s)
 	TMultiMap< int32, TWeakObjectPtr<UEdGraphNode> > LineNumberToMacroInstanceNodeMap;
 
+	// Reverse map from code offset to exec pin
+	TMap< int32, TWeakObjectPtr<UEdGraphPin> > LineNumberToPinMap;
+
 public:
 	FDebuggingInfoForSingleFunction()
 	{
@@ -233,6 +236,17 @@ public:
 		}
 	}
 
+	// Finds the macro source node associated with the code location Function+CodeOffset, or NULL if there isn't one
+	UEdGraphPin* FindExecPinFromCodeLocation(UFunction* Function, int32 CodeOffset) const
+	{
+		if (const FDebuggingInfoForSingleFunction* pFuncInfo = PerFunctionLineNumbers.Find(Function))
+		{
+			return pFuncInfo->LineNumberToPinMap.FindRef(CodeOffset).Get();
+		}
+
+		return NULL;
+	}
+
 	// Finds the breakpoint injection site(s) in bytecode if any were associated with the given node
 	void FindBreakpointInjectionSites(UEdGraphNode* Node, TArray<uint8*>& InstallSites) const
 	{
@@ -291,6 +305,12 @@ public:
 		}
 	}
 
+	void RegisterPinToCodeAssociation(UEdGraphPin const* ExecPin, UFunction* InFunction, int32 CodeOffset)
+	{
+		FDebuggingInfoForSingleFunction& PerFuncInfo = PerFunctionLineNumbers.FindOrAdd(InFunction);
+		PerFuncInfo.LineNumberToPinMap.Add(CodeOffset, ExecPin);
+	}
+
 	// Registers an association between an object (pin or node typically) and an associated class member property
 	void RegisterClassPropertyAssociation(class UObject* TrueSourceObject, class UProperty* AssociatedProperty)
 	{
