#summary Can search for nodes in the palette menu in Blueprints with both the native name and the localized name.

diff --git a/Engine/Source/Editor/BlueprintGraph/Private/K2Node_InputKey.cpp b/Engine/Source/Editor/BlueprintGraph/Private/K2Node_InputKey.cpp
--- a/Engine/Source/Editor/BlueprintGraph/Private/K2Node_InputKey.cpp
+++ b/Engine/Source/Editor/BlueprintGraph/Private/K2Node_InputKey.cpp
@@ -128,15 +128,31 @@ FText UK2Node_InputKey::GetKeyText() const
 	return InputKey.GetDisplayName();
 }
 
-FString UK2Node_InputKey::GetNodeTitle(ENodeTitleType::Type TitleType) const
+FText UK2Node_InputKey::GetNodeTitle(ENodeTitleType::Type TitleType) const
 {
 	if (bControl || bAlt || bShift)
 	{
-		return FText::Format(NSLOCTEXT("K2Node", "InputKey_Name_WithModifiers", "{0} {1}"), GetModifierText(), GetKeyText()).ToString();
+		FFormatNamedArguments Args;
+		Args.Add(TEXT("ModifierKey"), GetModifierText());
+		Args.Add(TEXT("Key"), GetKeyText());
+		return FText::Format(NSLOCTEXT("K2Node", "InputKey_Name_WithModifiers", "{ModifierKey} {Key}"), Args);
 	}
 	else
 	{
-		return FText::Format(NSLOCTEXT("K2Node", "InputKey_Name", "{0}"), GetKeyText()).ToString();
+		return GetKeyText();
+	}
+}
+
+FString UK2Node_InputKey::GetNodeNativeTitle(ENodeTitleType::Type TitleType) const
+{
+	// Do not setup this function for localization, intentionally left unlocalized!
+	if (bControl || bAlt || bShift)
+	{
+		return FString::Printf(TEXT("%s %s"), *GetModifierName().ToString(), *InputKey.ToString());
+	}
+	else
+	{
+		return GetKeyText().ToString();
 	}
 }
 
diff --git a/Engine/Source/Editor/BlueprintGraph/Private/K2Node_InputTouch.cpp b/Engine/Source/Editor/BlueprintGraph/Private/K2Node_InputTouch.cpp
--- a/Engine/Source/Editor/BlueprintGraph/Private/K2Node_InputTouch.cpp
+++ b/Engine/Source/Editor/BlueprintGraph/Private/K2Node_InputTouch.cpp
@@ -55,9 +55,9 @@ FLinearColor UK2Node_InputTouch::GetNodeTitleColor() const
 	return GEditor->AccessEditorUserSettings().EventNodeTitleColor;
 }
 
-FString UK2Node_InputTouch::GetNodeTitle(ENodeTitleType::Type TitleType) const
+FText UK2Node_InputTouch::GetNodeTitle(ENodeTitleType::Type TitleType) const
 {
-	return NSLOCTEXT("K2Node", "InputTouch_Name", "InputTouch").ToString();
+	return NSLOCTEXT("K2Node", "InputTouch_Name", "InputTouch");
 }
 
 UEdGraphPin* UK2Node_InputTouch::GetPressedPin() const
diff --git a/Engine/Source/Editor/BlueprintGraph/Private/K2Node_Literal.cpp b/Engine/Source/Editor/BlueprintGraph/Private/K2Node_Literal.cpp
--- a/Engine/Source/Editor/BlueprintGraph/Private/K2Node_Literal.cpp
+++ b/Engine/Source/Editor/BlueprintGraph/Private/K2Node_Literal.cpp
@@ -113,23 +113,23 @@ FString UK2Node_Literal::GetTooltip() const
 	return NSLOCTEXT("K2Node", "Literal_Tooltip", "Stores a reference to an actor in the level").ToString();
 }
 
-FString UK2Node_Literal::GetNodeTitle(ENodeTitleType::Type TitleType) const
+FText UK2Node_Literal::GetNodeTitle(ENodeTitleType::Type TitleType) const
 {
 	if( ObjectRef != NULL )
 	{
 		AActor* Actor = Cast<AActor>(ObjectRef);
 		if(Actor != NULL)
 		{
-			return Actor->GetActorLabel();
+			return FText::FromString(Actor->GetActorLabel());
 		}
 		else
 		{
-			return ObjectRef->GetName();
+			return FText::FromString(ObjectRef->GetName());
 		}
 	}
 	else
 	{
-		return NSLOCTEXT("K2Node", "Unknown", "Unknown").ToString();
+		return NSLOCTEXT("K2Node", "Unknown", "Unknown");
 	}
 }
 
@@ -184,7 +184,8 @@ void UK2Node_Literal::SetObjectRef(UObject* NewValue)
 
 	if( ValuePin )
 	{
-		ValuePin->PinName = GetNodeTitle(ENodeTitleType::ListView);
+		ValuePin->PinName = GetNodeNativeTitle(ENodeTitleType::ListView);
+		ValuePin->PinFriendlyName = GetNodeTitle(ENodeTitleType::ListView);
 	}
 }
 
diff --git a/Engine/Source/Editor/BlueprintGraph/Private/K2Node_LocalVariable.cpp b/Engine/Source/Editor/BlueprintGraph/Private/K2Node_LocalVariable.cpp
--- a/Engine/Source/Editor/BlueprintGraph/Private/K2Node_LocalVariable.cpp
+++ b/Engine/Source/Editor/BlueprintGraph/Private/K2Node_LocalVariable.cpp
@@ -22,18 +22,37 @@ FString UK2Node_LocalVariable::GetTooltip() const
 	return VariableTooltip.ToString();
 }
 
-FString UK2Node_LocalVariable::GetNodeTitle(ENodeTitleType::Type TitleType) const
+FText UK2Node_LocalVariable::GetNodeTitle(ENodeTitleType::Type TitleType) const
 {
+	if (TitleType == ENodeTitleType::EditableTitle)
+	{
+		return FText::FromName(CustomVariableName);
+	}
+	else if(TitleType == ENodeTitleType::ListView)
+	{
+		FFormatNamedArguments Args;
+		Args.Add(TEXT("TypeName"), UEdGraphSchema_K2::TypeToText(VariableType));
+		return FText::Format(NSLOCTEXT("K2Node", "LocalVariable", "Local {TypeName}"), Args);
+	}
+
+	FFormatNamedArguments Args;
+	Args.Add(TEXT("Title"), FText::FromName(CustomVariableName));
+	return FText::Format(NSLOCTEXT("K2Node", "LocalVariable_Name", "{Title}\nLocal Variable"), Args);
+}
+
+FString UK2Node_LocalVariable::GetNodeNativeTitle(ENodeTitleType::Type TitleType) const
+{
+	// Do not setup this function for localization, intentionally left unlocalized!
 	if (TitleType == ENodeTitleType::EditableTitle)
 	{
 		return CustomVariableName.ToString();
 	}
 	else if(TitleType == ENodeTitleType::ListView)
 	{
-		return NSLOCTEXT("K2Node", "LocalVariable", "Local ").ToString() + UEdGraphSchema_K2::TypeToString(VariableType);
+		return FString(TEXT("Local ")) + UEdGraphSchema_K2::TypeToString(VariableType);
 	}
 
-	return FString::Printf(*NSLOCTEXT("K2Node", "LocalVariable_Name", "%s\nLocal Variable").ToString(), *CustomVariableName.ToString());
+	return FString::Printf(TEXT("%s\nLocal Variable"), *CustomVariableName.ToString());
 }
 
 void UK2Node_LocalVariable::OnRenameNode(const FString& NewName)
@@ -50,7 +69,7 @@ void UK2Node_LocalVariable::OnRenameNode(const FString& NewName)
 
 TSharedPtr<class INameValidatorInterface> UK2Node_LocalVariable::MakeNameValidator() const
 {
-	return MakeShareable(new FKismetNameValidator(GetBlueprint(), *GetNodeTitle(ENodeTitleType::EditableTitle)));
+	return MakeShareable(new FKismetNameValidator(GetBlueprint(), *GetNodeTitle(ENodeTitleType::EditableTitle).ToString()));
 }
 
 void UK2Node_LocalVariable::ChangeVariableType(const FEdGraphPinType& InVariableType)
diff --git a/Engine/Source/Editor/BlueprintGraph/Private/K2Node_MCDelegate.cpp b/Engine/Source/Editor/BlueprintGraph/Private/K2Node_MCDelegate.cpp
--- a/Engine/Source/Editor/BlueprintGraph/Private/K2Node_MCDelegate.cpp
+++ b/Engine/Source/Editor/BlueprintGraph/Private/K2Node_MCDelegate.cpp
@@ -73,7 +73,7 @@ void UK2Node_BaseMCDelegate::AllocateDefaultPins()
 
 	if(SelfPin)
 	{
-		SelfPin->PinFriendlyName = NSLOCTEXT("K2Node", "BaseMCDelegateSelfPinName", "Target").ToString();
+		SelfPin->PinFriendlyName = NSLOCTEXT("K2Node", "BaseMCDelegateSelfPinName", "Target");
 	}
 }
 
@@ -201,13 +201,22 @@ void UK2Node_AddDelegate::AllocateDefaultPins()
 
 	if(UEdGraphPin* DelegatePin = CreatePin(EGPD_Input, K2Schema->PC_Delegate, TEXT(""), GetDelegateSignature(), false, false, FK2Node_BaseMCDelegateHelper::DelegatePinName))
 	{
-		DelegatePin->PinFriendlyName = *NSLOCTEXT("K2Node", "PinFriendlyDelegatetName", "Event").ToString();
+		DelegatePin->PinFriendlyName = NSLOCTEXT("K2Node", "PinFriendlyDelegatetName", "Event");
 	}
 }
 
-FString UK2Node_AddDelegate::GetNodeTitle(ENodeTitleType::Type TitleType) const
+FText UK2Node_AddDelegate::GetNodeTitle(ENodeTitleType::Type TitleType) const
 {
-	return FString::Printf(*NSLOCTEXT("K2Node", "AddDelegate", "Bind Event to %s").ToString(), *GetPropertyName().ToString());
+	FFormatNamedArguments Args;
+	Args.Add(TEXT("ProeprtyName"), FText::FromName(GetPropertyName()));
+	return FText::Format(NSLOCTEXT("K2Node", "AddDelegate", "Bind Event to {ProeprtyName}"), Args);
+}
+
+FString UK2Node_AddDelegate::GetNodeNativeTitle(ENodeTitleType::Type TitleType) const
+{
+	// Do not setup this function for localization, intentionally left unlocalized!
+	
+	return FString::Printf(TEXT("Bind Event to %s"), *GetPropertyName().ToString());
 }
 
 FNodeHandlingFunctor* UK2Node_AddDelegate::CreateNodeHandler(FKismetCompilerContext& CompilerContext) const
@@ -222,9 +231,18 @@ UK2Node_ClearDelegate::UK2Node_ClearDelegate(const class FPostConstructInitializ
 {
 }
 
-FString UK2Node_ClearDelegate::GetNodeTitle(ENodeTitleType::Type TitleType) const
+FText UK2Node_ClearDelegate::GetNodeTitle(ENodeTitleType::Type TitleType) const
+{
+	FFormatNamedArguments Args;
+	Args.Add(TEXT("PropertyName"), FText::FromName(GetPropertyName()));
+	return FText::Format(NSLOCTEXT("K2Node", "ClearDelegate", "Unbind all Events from {PropertyName}"), Args);
+}
+
+FString UK2Node_ClearDelegate::GetNodeNativeTitle(ENodeTitleType::Type TitleType) const
 {
-	return FString::Printf(*NSLOCTEXT("K2Node", "ClearDelegate", "Unbind all Events from %s").ToString(), *GetPropertyName().ToString());
+	// Do not setup this function for localization, intentionally left unlocalized!
+	
+	return FString::Printf(TEXT("Unbind all Events from %s"), *GetPropertyName().ToString());
 }
 
 FNodeHandlingFunctor* UK2Node_ClearDelegate::CreateNodeHandler(FKismetCompilerContext& CompilerContext) const
@@ -247,13 +265,22 @@ void UK2Node_RemoveDelegate::AllocateDefaultPins()
 
 	if(UEdGraphPin* DelegatePin = CreatePin(EGPD_Input, K2Schema->PC_Delegate, TEXT(""), GetDelegateSignature(), false, false, FK2Node_BaseMCDelegateHelper::DelegatePinName))
 	{
-		DelegatePin->PinFriendlyName = *NSLOCTEXT("K2Node", "PinFriendlyDelegatetName", "Event").ToString();
+		DelegatePin->PinFriendlyName = NSLOCTEXT("K2Node", "PinFriendlyDelegatetName", "Event");
 	}
 }
 
-FString UK2Node_RemoveDelegate::GetNodeTitle(ENodeTitleType::Type TitleType) const
+FText UK2Node_RemoveDelegate::GetNodeTitle(ENodeTitleType::Type TitleType) const
 {
-	return FString::Printf(*NSLOCTEXT("K2Node", "RemoveDelegate", "Unbind Event from %s").ToString(), *GetPropertyName().ToString());
+	FFormatNamedArguments Args;
+	Args.Add(TEXT("PropertyName"), FText::FromName(GetPropertyName()));
+	return FText::Format(NSLOCTEXT("K2Node", "RemoveDelegate", "Unbind Event from {PropertyName}"), Args);
+}
+
+FString UK2Node_RemoveDelegate::GetNodeNativeTitle(ENodeTitleType::Type TitleType) const
+{
+	// Do not setup this function for localization, intentionally left unlocalized!
+	
+	return FString::Printf(TEXT("Unbind Event from %s"), *GetPropertyName().ToString());
 }
 
 FNodeHandlingFunctor* UK2Node_RemoveDelegate::CreateNodeHandler(FKismetCompilerContext& CompilerContext) const
@@ -297,9 +324,18 @@ void UK2Node_CallDelegate::AllocateDefaultPins()
 	CreatePinsForFunctionInputs(GetDelegateSignature());
 }
 
-FString UK2Node_CallDelegate::GetNodeTitle(ENodeTitleType::Type TitleType) const
+FText UK2Node_CallDelegate::GetNodeTitle(ENodeTitleType::Type TitleType) const
+{
+	FFormatNamedArguments Args;
+	Args.Add(TEXT("PropertyName"), FText::FromName(GetPropertyName()));
+	return FText::Format(NSLOCTEXT("K2Node", "CallDelegate", "Call {PropertyName}"), Args);
+}
+
+FString UK2Node_CallDelegate::GetNodeNativeTitle(ENodeTitleType::Type TitleType) const
 {
-	return FString::Printf(*NSLOCTEXT("K2Node", "CallDelegate", "Call %s").ToString(), *GetPropertyName().ToString());
+	// Do not setup this function for localization, intentionally left unlocalized!
+	
+	return FString::Printf(TEXT("Call %s"), *GetPropertyName().ToString());
 }
 
 void UK2Node_CallDelegate::ValidateNodeDuringCompilation(class FCompilerResultsLog& MessageLog) const
diff --git a/Engine/Source/Editor/BlueprintGraph/Private/K2Node_MacroInstance.cpp b/Engine/Source/Editor/BlueprintGraph/Private/K2Node_MacroInstance.cpp
--- a/Engine/Source/Editor/BlueprintGraph/Private/K2Node_MacroInstance.cpp
+++ b/Engine/Source/Editor/BlueprintGraph/Private/K2Node_MacroInstance.cpp
@@ -118,10 +118,24 @@ FString UK2Node_MacroInstance::GetTooltip() const
 	}
 }
 
-FString UK2Node_MacroInstance::GetNodeTitle(ENodeTitleType::Type TitleType) const
+FText UK2Node_MacroInstance::GetNodeTitle(ENodeTitleType::Type TitleType) const
 {
 	UEdGraph* MacroGraph = MacroGraphReference.GetGraph();
-	FString Result = NSLOCTEXT("K2Node", "MacroInstance", "Macro instance").ToString();
+	FText Result = NSLOCTEXT("K2Node", "MacroInstance", "Macro instance");
+	if (MacroGraph)
+	{
+		Result = FText::FromString(MacroGraph->GetName());
+	}
+
+	return Result;
+}
+
+FString UK2Node_MacroInstance::GetNodeNativeTitle(ENodeTitleType::Type TitleType) const
+{
+	// Do not setup this function for localization, intentionally left unlocalized!
+
+	UEdGraph* MacroGraph = MacroGraphReference.GetGraph();
+	FString Result = TEXT("Macro instance");
 	if (MacroGraph)
 	{
 		Result = MacroGraph->GetName();
diff --git a/Engine/Source/Editor/BlueprintGraph/Private/K2Node_MakeArray.cpp b/Engine/Source/Editor/BlueprintGraph/Private/K2Node_MakeArray.cpp
--- a/Engine/Source/Editor/BlueprintGraph/Private/K2Node_MakeArray.cpp
+++ b/Engine/Source/Editor/BlueprintGraph/Private/K2Node_MakeArray.cpp
@@ -81,9 +81,9 @@ FNodeHandlingFunctor* UK2Node_MakeArray::CreateNodeHandler(FKismetCompilerContex
 	return new FKCHandler_MakeArray(CompilerContext);
 }
 
-FString UK2Node_MakeArray::GetNodeTitle(ENodeTitleType::Type TitleType) const
+FText UK2Node_MakeArray::GetNodeTitle(ENodeTitleType::Type TitleType) const
 {
-	return LOCTEXT("NodeTitle", "Make Array").ToString();
+	return LOCTEXT("NodeTitle", "Make Array");
 }
 
 UEdGraphPin* UK2Node_MakeArray::GetOutputPin() const
diff --git a/Engine/Source/Editor/BlueprintGraph/Private/K2Node_MakeStruct.cpp b/Engine/Source/Editor/BlueprintGraph/Private/K2Node_MakeStruct.cpp
--- a/Engine/Source/Editor/BlueprintGraph/Private/K2Node_MakeStruct.cpp
+++ b/Engine/Source/Editor/BlueprintGraph/Private/K2Node_MakeStruct.cpp
@@ -249,8 +249,17 @@ void UK2Node_MakeStruct::ValidateNodeDuringCompilation(class FCompilerResultsLog
 }
 
 
-FString UK2Node_MakeStruct::GetNodeTitle(ENodeTitleType::Type TitleType) const 
+FText UK2Node_MakeStruct::GetNodeTitle(ENodeTitleType::Type TitleType) const 
 {
+	FFormatNamedArguments Args;
+	Args.Add(TEXT("StructName"), FText::FromString(StructType ? StructType->GetName() : FString()));
+	return FText::Format(LOCTEXT("MakeNodeTitle", "Make {StructName}"), Args);
+}
+
+FString UK2Node_MakeStruct::GetNodeNativeTitle(ENodeTitleType::Type TitleType) const 
+{
+	// Do not setup this function for localization, intentionally left unlocalized!
+	
 	return FString(TEXT("Make ")) + (StructType ? StructType->GetName() : FString());
 }
 
diff --git a/Engine/Source/Editor/BlueprintGraph/Private/K2Node_MatineeController.cpp b/Engine/Source/Editor/BlueprintGraph/Private/K2Node_MatineeController.cpp
--- a/Engine/Source/Editor/BlueprintGraph/Private/K2Node_MatineeController.cpp
+++ b/Engine/Source/Editor/BlueprintGraph/Private/K2Node_MatineeController.cpp
@@ -45,15 +45,15 @@ FLinearColor UK2Node_MatineeController::GetNodeTitleColor() const
 }
 
 
-FString UK2Node_MatineeController::GetNodeTitle(ENodeTitleType::Type TitleType) const
+FText UK2Node_MatineeController::GetNodeTitle(ENodeTitleType::Type TitleType) const
 {
 	if(MatineeActor != NULL)
 	{
-		return MatineeActor->GetActorLabel();
+		return FText::FromString(MatineeActor->GetActorLabel());
 	}
 	else
 	{
-		return NSLOCTEXT("K2Node", "InvalidMatineeController", "INVALID MATINEECONTROLLER").ToString();
+		return NSLOCTEXT("K2Node", "InvalidMatineeController", "INVALID MATINEECONTROLLER");
 	}
 }
 
diff --git a/Engine/Source/Editor/BlueprintGraph/Private/K2Node_Message.cpp b/Engine/Source/Editor/BlueprintGraph/Private/K2Node_Message.cpp
--- a/Engine/Source/Editor/BlueprintGraph/Private/K2Node_Message.cpp
+++ b/Engine/Source/Editor/BlueprintGraph/Private/K2Node_Message.cpp
@@ -12,18 +12,39 @@ UK2Node_Message::UK2Node_Message(const class FPostConstructInitializeProperties&
 {
 }
 
-FString UK2Node_Message::GetNodeTitle(ENodeTitleType::Type TitleType) const
+FText UK2Node_Message::GetNodeTitle(ENodeTitleType::Type TitleType) const
 {
+	FText NodeName;
+	if (UFunction* Function = GetTargetFunction())
+	{
+		FString NodeNameStr = UK2Node_CallFunction::GetUserFacingFunctionName(Function);
+
+		FFormatNamedArguments Args;
+		Args.Add(TEXT("NodeName"), FText::FromString(NodeNameStr));
+		Args.Add(TEXT("OuterClassName"), FText::FromString(Function->GetOuterUClass()->GetName()));
+		NodeName = FText::Format(NSLOCTEXT("K2Node", "CallInterfaceContext", "{NodeName}\nUsing Interface {OuterClassName}"), Args);
+	}
+	else
+	{
+		NodeName = NSLOCTEXT("K2Node", "InvalidMessageNode", "Invalid Message Node");
+	}
+
+	return NodeName;
+}
+
+FString UK2Node_Message::GetNodeNativeTitle(ENodeTitleType::Type TitleType) const
+{
+	// Do not setup this function for localization, intentionally left unlocalized!
 	FString NodeName;
 	if (UFunction* Function = GetTargetFunction())
 	{
 		NodeName = UK2Node_CallFunction::GetUserFacingFunctionName(Function);
 		NodeName += FString(TEXT("\n"));
-		NodeName += FString::Printf(*NSLOCTEXT("K2Node", "CallInterfaceContext", "Using Interface %s").ToString(), *(Function->GetOuterUClass()->GetName()));
+		NodeName += FString::Printf(TEXT("Using Interface %s"), *(Function->GetOuterUClass()->GetName()));
 	}
 	else
 	{
-		NodeName = NSLOCTEXT("K2Node", "InvalidMessageNode", "Invalid Message Node").ToString();
+		NodeName = TEXT("Invalid Message Node");
 	}
 
 	return NodeName;
diff --git a/Engine/Source/Editor/BlueprintGraph/Private/K2Node_MultiGate.cpp b/Engine/Source/Editor/BlueprintGraph/Private/K2Node_MultiGate.cpp
--- a/Engine/Source/Editor/BlueprintGraph/Private/K2Node_MultiGate.cpp
+++ b/Engine/Source/Editor/BlueprintGraph/Private/K2Node_MultiGate.cpp
@@ -469,9 +469,9 @@ FLinearColor UK2Node_MultiGate::GetNodeTitleColor() const
 	return FLinearColor::White;
 }
 
-FString UK2Node_MultiGate::GetNodeTitle(ENodeTitleType::Type TitleType) const
+FText UK2Node_MultiGate::GetNodeTitle(ENodeTitleType::Type TitleType) const
 {
-	return NSLOCTEXT("K2Node", "MultiGate", "MultiGate").ToString();
+	return NSLOCTEXT("K2Node", "MultiGate", "MultiGate");
 }
 
 void UK2Node_MultiGate::AllocateDefaultPins()
diff --git a/Engine/Source/Editor/BlueprintGraph/Private/K2Node_Select.cpp b/Engine/Source/Editor/BlueprintGraph/Private/K2Node_Select.cpp
--- a/Engine/Source/Editor/BlueprintGraph/Private/K2Node_Select.cpp
+++ b/Engine/Source/Editor/BlueprintGraph/Private/K2Node_Select.cpp
@@ -233,11 +233,11 @@ void UK2Node_Select::AllocateDefaultPins()
 		{
 			if (EnumEntryFriendlyNames[Idx] != NAME_None)
 			{
-				NewPin->PinFriendlyName = EnumEntryFriendlyNames[Idx].ToString();
+				NewPin->PinFriendlyName = FText::FromName(EnumEntryFriendlyNames[Idx]);
 			}
 			else
 			{
-				NewPin->PinFriendlyName = TEXT("");
+				NewPin->PinFriendlyName = FText::GetEmpty();
 			}
 		}
 	}
@@ -256,9 +256,9 @@ FString UK2Node_Select::GetTooltip() const
 	return TEXT("Return the option at Index, (first option is indexed at 0)");
 }
 
-FString UK2Node_Select::GetNodeTitle(ENodeTitleType::Type TitleType) const
+FText UK2Node_Select::GetNodeTitle(ENodeTitleType::Type TitleType) const
 {
-	return TEXT("Select");
+	return LOCTEXT("Select", "Select");
 }
 
 void UK2Node_Select::ReallocatePinsDuringReconstruction(TArray<UEdGraphPin*>& OldPins)
diff --git a/Engine/Source/Editor/BlueprintGraph/Private/K2Node_SelectEnum.cpp b/Engine/Source/Editor/BlueprintGraph/Private/K2Node_SelectEnum.cpp
--- a/Engine/Source/Editor/BlueprintGraph/Private/K2Node_SelectEnum.cpp
+++ b/Engine/Source/Editor/BlueprintGraph/Private/K2Node_SelectEnum.cpp
@@ -44,8 +44,19 @@ FString UDEPRECATED_K2Node_SelectEnum::GetTooltip() const
 	return TEXT("Return literal value set for each enum value");
 }
 
-FString UDEPRECATED_K2Node_SelectEnum::GetNodeTitle(ENodeTitleType::Type TitleType) const
+FText  UDEPRECATED_K2Node_SelectEnum::GetNodeTitle(ENodeTitleType::Type TitleType) const
 {
+	FText EnumName = (Enum != NULL) ? FText::FromString(Enum->GetName()) : NSLOCTEXT("K2Node", "BadEnum", "(bad enum)");
+
+	FFormatNamedArguments Args;
+	Args.Add(TEXT("EnumName"), EnumName);
+	return FText::Format(NSLOCTEXT("K2Node", "SelectByEnum", "Select by {EnumName}"), Args);
+}
+
+FString UDEPRECATED_K2Node_SelectEnum::GetNodeNativeTitle(ENodeTitleType::Type TitleType) const
+{
+	// Do not setup this function for localization, intentionally left unlocalized!
+
 	FString EnumName = (Enum != NULL) ? Enum->GetName() : TEXT("(bad enum)");
 	return FString::Printf(TEXT("Select by %s"), *EnumName);
 }
diff --git a/Engine/Source/Editor/BlueprintGraph/Private/K2Node_Self.cpp b/Engine/Source/Editor/BlueprintGraph/Private/K2Node_Self.cpp
--- a/Engine/Source/Editor/BlueprintGraph/Private/K2Node_Self.cpp
+++ b/Engine/Source/Editor/BlueprintGraph/Private/K2Node_Self.cpp
@@ -52,9 +52,9 @@ FString UK2Node_Self::GetKeywords() const
 	return TEXT("This");
 }
 
-FString UK2Node_Self::GetNodeTitle(ENodeTitleType::Type TitleType) const
+FText UK2Node_Self::GetNodeTitle(ENodeTitleType::Type TitleType) const
 {
-	return FString::Printf(*NSLOCTEXT("K2Node", "SelfReferenceName", "Self-Reference").ToString());
+	return NSLOCTEXT("K2Node", "SelfReferenceName", "Self-Reference");
 }
 
 FNodeHandlingFunctor* UK2Node_Self::CreateNodeHandler(FKismetCompilerContext& CompilerContext) const
diff --git a/Engine/Source/Editor/BlueprintGraph/Private/K2Node_SpawnActor.cpp b/Engine/Source/Editor/BlueprintGraph/Private/K2Node_SpawnActor.cpp
--- a/Engine/Source/Editor/BlueprintGraph/Private/K2Node_SpawnActor.cpp
+++ b/Engine/Source/Editor/BlueprintGraph/Private/K2Node_SpawnActor.cpp
@@ -228,11 +228,35 @@ FLinearColor UK2Node_SpawnActor::GetNodeTitleColor() const
 }
 
 
-FString UK2Node_SpawnActor::GetNodeTitle(ENodeTitleType::Type TitleType) const
+FText UK2Node_SpawnActor::GetNodeTitle(ENodeTitleType::Type TitleType) const
 {
 	UEdGraphPin* BlueprintPin = GetBlueprintPin();
 
-	FString SpawnString = NSLOCTEXT("K2Node", "None", "NONE").ToString();
+	FText SpawnString = NSLOCTEXT("K2Node", "None", "NONE");
+	if(BlueprintPin != NULL)
+	{
+		if(BlueprintPin->LinkedTo.Num() > 0)
+		{
+			// Blueprint will be determined dynamically, so we don't have the name in this case
+			SpawnString = FText::GetEmpty();
+		}
+		else if(BlueprintPin->DefaultObject != NULL)
+		{
+			SpawnString = FText::FromString(BlueprintPin->DefaultObject->GetName());
+		}
+	}
+
+	FFormatNamedArguments Args;
+	Args.Add(TEXT("ActorName"), SpawnString);
+	return FText::Format(NSLOCTEXT("K2Node", "SpawnActor", "SpawnActor {ActorName}"), Args);
+}
+
+FString UK2Node_SpawnActor::GetNodeNativeTitle(ENodeTitleType::Type TitleType) const
+{
+	// Do not setup this function for localization, intentionally left unlocalized!
+	UEdGraphPin* BlueprintPin = GetBlueprintPin();
+
+	FString SpawnString = TEXT("NONE");
 	if(BlueprintPin != NULL)
 	{
 		if(BlueprintPin->LinkedTo.Num() > 0)
@@ -246,7 +270,7 @@ FString UK2Node_SpawnActor::GetNodeTitle(ENodeTitleType::Type TitleType) const
 		}
 	}
 
-	return FString::Printf(*NSLOCTEXT("K2Node", "SpawnActor", "SpawnActor %s").ToString(), *SpawnString);
+	return FString::Printf(TEXT("SpawnActor %s"), *SpawnString);
 }
 
 bool UK2Node_SpawnActor::CanPasteHere( const UEdGraph* TargetGraph, const UEdGraphSchema* Schema ) const 
diff --git a/Engine/Source/Editor/BlueprintGraph/Private/K2Node_SpawnActorFromClass.cpp b/Engine/Source/Editor/BlueprintGraph/Private/K2Node_SpawnActorFromClass.cpp
--- a/Engine/Source/Editor/BlueprintGraph/Private/K2Node_SpawnActorFromClass.cpp
+++ b/Engine/Source/Editor/BlueprintGraph/Private/K2Node_SpawnActorFromClass.cpp
@@ -250,11 +250,36 @@ FLinearColor UK2Node_SpawnActorFromClass::GetNodeTitleColor() const
 }
 
 
-FString UK2Node_SpawnActorFromClass::GetNodeTitle(ENodeTitleType::Type TitleType) const
+FText UK2Node_SpawnActorFromClass::GetNodeTitle(ENodeTitleType::Type TitleType) const
 {
 	UEdGraphPin* ClassPin = GetClassPin();
 
-	FString SpawnString = NSLOCTEXT("K2Node", "None", "NONE").ToString();
+	FText SpawnString = NSLOCTEXT("K2Node", "None", "NONE");
+	if(ClassPin != NULL)
+	{
+		if(ClassPin->LinkedTo.Num() > 0)
+		{
+			// Blueprint will be determined dynamically, so we don't have the name in this case
+			SpawnString = FText::GetEmpty();
+		}
+		else if(ClassPin->DefaultObject != NULL)
+		{
+			SpawnString = FText::FromString(ClassPin->DefaultObject->GetName());
+		}
+	}
+
+	FFormatNamedArguments Args;
+	Args.Add(TEXT("ClassName"), SpawnString);
+	return FText::Format(NSLOCTEXT("K2Node", "SpawnActor", "SpawnActor {ClassName}"), Args);
+}
+
+FString UK2Node_SpawnActorFromClass::GetNodeNativeTitle(ENodeTitleType::Type TitleType) const
+{
+	// Do not setup this function for localization, intentionally left unlocalized!
+	
+	UEdGraphPin* ClassPin = GetClassPin();
+
+	FString SpawnString = TEXT("NONE");
 	if(ClassPin != NULL)
 	{
 		if(ClassPin->LinkedTo.Num() > 0)
@@ -268,7 +293,7 @@ FString UK2Node_SpawnActorFromClass::GetNodeTitle(ENodeTitleType::Type TitleType
 		}
 	}
 
-	return FString::Printf(*NSLOCTEXT("K2Node", "SpawnActor", "SpawnActor %s").ToString(), *SpawnString);
+	return FString::Printf(TEXT("SpawnActor %s"), *SpawnString);
 }
 
 bool UK2Node_SpawnActorFromClass::CanPasteHere( const UEdGraph* TargetGraph, const UEdGraphSchema* Schema ) const 
diff --git a/Engine/Source/Editor/BlueprintGraph/Private/K2Node_StructMemberGet.cpp b/Engine/Source/Editor/BlueprintGraph/Private/K2Node_StructMemberGet.cpp
--- a/Engine/Source/Editor/BlueprintGraph/Private/K2Node_StructMemberGet.cpp
+++ b/Engine/Source/Editor/BlueprintGraph/Private/K2Node_StructMemberGet.cpp
@@ -6,6 +6,8 @@
 //////////////////////////////////////////////////////////////////////////
 // UK2Node_StructMemberGet
 
+#define LOCTEXT_NAMESPACE "K2Node"
+
 UK2Node_StructMemberGet::UK2Node_StructMemberGet(const class FPostConstructInitializeProperties& PCIP)
 	: Super(PCIP)
 {
@@ -69,11 +71,22 @@ void UK2Node_StructMemberGet::AllocatePinsForSingleMemberGet(FName MemberName)
 
 FString UK2Node_StructMemberGet::GetTooltip() const
 {
-	return FString::Printf(TEXT("Get member variables of %s"), *GetVarNameString());
+	FFormatNamedArguments Args;
+	Args.Add(TEXT("VariableName"), FText::FromString(GetVarNameString()));
+	return FText::Format(LOCTEXT("K2Node_StructMemberGet_Tooltip", "Get member variables of {VariableName}"), Args).ToString();
+}
+
+FText UK2Node_StructMemberGet::GetNodeTitle(ENodeTitleType::Type TitleType) const
+{
+	FFormatNamedArguments Args;
+	Args.Add(TEXT("VariableName"), FText::FromString(GetVarNameString()));
+	return FText::Format(LOCTEXT("GetMembersInVariable", "Get members in {VariableName}"), Args);
 }
 
-FString UK2Node_StructMemberGet::GetNodeTitle(ENodeTitleType::Type TitleType) const
+FString UK2Node_StructMemberGet::GetNodeNativeTitle(ENodeTitleType::Type TitleType) const
 {
+	// Do not setup this function for localization, intentionally left unlocalized!
+
 	return FString::Printf(TEXT("Get members in %s"), *GetVarNameString());
 }
 
@@ -81,3 +94,5 @@ FNodeHandlingFunctor* UK2Node_StructMemberGet::CreateNodeHandler(FKismetCompiler
 {
 	return new FKCHandler_StructMemberVariableGet(CompilerContext);
 }
+
+#undef LOCTEXT_NAMESPACE
\ No newline at end of file
diff --git a/Engine/Source/Editor/BlueprintGraph/Private/K2Node_StructMemberSet.cpp b/Engine/Source/Editor/BlueprintGraph/Private/K2Node_StructMemberSet.cpp
--- a/Engine/Source/Editor/BlueprintGraph/Private/K2Node_StructMemberSet.cpp
+++ b/Engine/Source/Editor/BlueprintGraph/Private/K2Node_StructMemberSet.cpp
@@ -7,6 +7,8 @@
 //////////////////////////////////////////////////////////////////////////
 // UK2Node_StructMemberSet
 
+#define LOCTEXT_NAMESPACE "K2Node"
+
 UK2Node_StructMemberSet::UK2Node_StructMemberSet(const class FPostConstructInitializeProperties& PCIP)
 	: Super(PCIP)
 {
@@ -42,11 +44,22 @@ void UK2Node_StructMemberSet::AllocateDefaultPins()
 
 FString UK2Node_StructMemberSet::GetTooltip() const
 {
-	return FString::Printf(TEXT("Set member variables of %s"), *VariableReference.GetMemberName().ToString());
+	FFormatNamedArguments Args;
+	Args.Add(TEXT("VariableName"), FText::FromName(VariableReference.GetMemberName()));
+	return FText::Format(LOCTEXT("K2Node_StructMemberSet_Tooltip", "Set member variables of {VariableName}"), Args).ToString();
+}
+
+FText UK2Node_StructMemberSet::GetNodeTitle(ENodeTitleType::Type TitleType) const
+{
+	FFormatNamedArguments Args;
+	Args.Add(TEXT("VariableName"), FText::FromName(VariableReference.GetMemberName()));
+	return FText::Format(LOCTEXT("SetMembersInVariable", "Set members in {VariableName}"), Args);
 }
 
-FString UK2Node_StructMemberSet::GetNodeTitle(ENodeTitleType::Type TitleType) const
+FString UK2Node_StructMemberSet::GetNodeNativeTitle(ENodeTitleType::Type TitleType) const
 {
+	// Do not setup this function for localization, intentionally left unlocalized!
+
 	return FString::Printf(TEXT("Set members in %s"), *VariableReference.GetMemberName().ToString());
 }
 
@@ -59,3 +72,5 @@ FNodeHandlingFunctor* UK2Node_StructMemberSet::CreateNodeHandler(FKismetCompiler
 {
 	return new FKCHandler_StructMemberVariableSet(CompilerContext);
 }
+
+#undef LOCTEXT_NAMESPACE
\ No newline at end of file
diff --git a/Engine/Source/Editor/BlueprintGraph/Private/K2Node_SwitchEnum.cpp b/Engine/Source/Editor/BlueprintGraph/Private/K2Node_SwitchEnum.cpp
--- a/Engine/Source/Editor/BlueprintGraph/Private/K2Node_SwitchEnum.cpp
+++ b/Engine/Source/Editor/BlueprintGraph/Private/K2Node_SwitchEnum.cpp
@@ -3,6 +3,8 @@
 #include "BlueprintGraphPrivatePCH.h"
 #include "../../../Runtime/Engine/Classes/Kismet/KismetMathLibrary.h"
 
+#define LOCTEXT_NAMESPACE "K2Node"
+
 UK2Node_SwitchEnum::UK2Node_SwitchEnum(const class FPostConstructInitializeProperties& PCIP)
 	: Super(PCIP)
 {
@@ -43,11 +45,23 @@ void UK2Node_SwitchEnum::SetEnum(UEnum* InEnum)
 	}
 }
 
-FString UK2Node_SwitchEnum::GetNodeTitle(ENodeTitleType::Type TitleType) const 
+FText UK2Node_SwitchEnum::GetNodeTitle(ENodeTitleType::Type TitleType) const 
 {
+	FText EnumName = (Enum != NULL) ? FText::FromString(Enum->GetName()) : LOCTEXT("BadEnum", "(bad enum)");
+
+	FFormatNamedArguments Args;
+	Args.Add(TEXT("EnumName"), EnumName);
+
+	return FText::Format(NSLOCTEXT("K2Node", "Switch_Enum", "Switch on {EnumName}"), Args);
+}
+
+FString UK2Node_SwitchEnum::GetNodeNativeTitle(ENodeTitleType::Type TitleType) const 
+{
+	// Do not setup this function for localization, intentionally left unlocalized!
+
 	FString EnumName = (Enum != NULL) ? Enum->GetName() : TEXT("(bad enum)");
 
-	return FString::Printf(*NSLOCTEXT("K2Node", "Switch_Enum", "Switch on %s").ToString(), *EnumName);
+	return FString::Printf(TEXT("Switch on %s"), *EnumName);
 }
 
 FString UK2Node_SwitchEnum::GetTooltip() const
@@ -96,7 +110,7 @@ void UK2Node_SwitchEnum::CreateCasePins()
 		int32 Index = EnumIt.GetIndex();
 		if (EnumFriendlyNames.IsValidIndex(Index))
 		{
-			NewPin->PinFriendlyName = EnumFriendlyNames[Index];
+			NewPin->PinFriendlyName = FText::FromString(EnumFriendlyNames[Index]);
 		}
 		
 		if(bShouldUseAdvancedView && (EnumIt.GetIndex() > 2))
@@ -159,4 +173,6 @@ void UK2Node_SwitchEnum::RemovePinFromSwitchNode(UEdGraphPin* Pin)
 		}
 		Pin->BreakAllPinLinks();
 	}
-}
\ No newline at end of file
+}
+
+#undef LOCTEXT_NAMESPACE
\ No newline at end of file
diff --git a/Engine/Source/Editor/BlueprintGraph/Private/K2Node_SwitchInteger.cpp b/Engine/Source/Editor/BlueprintGraph/Private/K2Node_SwitchInteger.cpp
--- a/Engine/Source/Editor/BlueprintGraph/Private/K2Node_SwitchInteger.cpp
+++ b/Engine/Source/Editor/BlueprintGraph/Private/K2Node_SwitchInteger.cpp
@@ -12,9 +12,9 @@ UK2Node_SwitchInteger::UK2Node_SwitchInteger(const class FPostConstructInitializ
 	FunctionClass = UKismetMathLibrary::StaticClass();
 }
 
-FString UK2Node_SwitchInteger::GetNodeTitle(ENodeTitleType::Type TitleType) const
+FText UK2Node_SwitchInteger::GetNodeTitle(ENodeTitleType::Type TitleType) const
 {
-	return NSLOCTEXT("K2Node", "Switch_Interger", "Switch on Int").ToString();
+	return NSLOCTEXT("K2Node", "Switch_Interger", "Switch on Int");
 }
 
 FString UK2Node_SwitchInteger::GetTooltip() const
diff --git a/Engine/Source/Editor/BlueprintGraph/Private/K2Node_SwitchName.cpp b/Engine/Source/Editor/BlueprintGraph/Private/K2Node_SwitchName.cpp
--- a/Engine/Source/Editor/BlueprintGraph/Private/K2Node_SwitchName.cpp
+++ b/Engine/Source/Editor/BlueprintGraph/Private/K2Node_SwitchName.cpp
@@ -27,9 +27,9 @@ void UK2Node_SwitchName::PostEditChangeProperty(struct FPropertyChangedEvent& Pr
 	Super::PostEditChangeProperty(PropertyChangedEvent);
 }
 
-FString UK2Node_SwitchName::GetNodeTitle(ENodeTitleType::Type TitleType) const
+FText UK2Node_SwitchName::GetNodeTitle(ENodeTitleType::Type TitleType) const
 {
-	return NSLOCTEXT("K2Node", "Switch_Name", "Switch on Name").ToString();
+	return NSLOCTEXT("K2Node", "Switch_Name", "Switch on Name");
 }
 
 FString UK2Node_SwitchName::GetTooltip() const
diff --git a/Engine/Source/Editor/BlueprintGraph/Private/K2Node_SwitchString.cpp b/Engine/Source/Editor/BlueprintGraph/Private/K2Node_SwitchString.cpp
--- a/Engine/Source/Editor/BlueprintGraph/Private/K2Node_SwitchString.cpp
+++ b/Engine/Source/Editor/BlueprintGraph/Private/K2Node_SwitchString.cpp
@@ -38,9 +38,9 @@ void UK2Node_SwitchString::PostEditChangeProperty(struct FPropertyChangedEvent&
 	Super::PostEditChangeProperty(PropertyChangedEvent);
 }
 
-FString UK2Node_SwitchString::GetNodeTitle(ENodeTitleType::Type TitleType) const
+FText UK2Node_SwitchString::GetNodeTitle(ENodeTitleType::Type TitleType) const
 {
-	return NSLOCTEXT("K2Node", "Switch_String", "Switch on String").ToString();
+	return NSLOCTEXT("K2Node", "Switch_String", "Switch on String");
 }
 
 FString UK2Node_SwitchString::GetTooltip() const
diff --git a/Engine/Source/Editor/BlueprintGraph/Private/K2Node_TemporaryVariable.cpp b/Engine/Source/Editor/BlueprintGraph/Private/K2Node_TemporaryVariable.cpp
--- a/Engine/Source/Editor/BlueprintGraph/Private/K2Node_TemporaryVariable.cpp
+++ b/Engine/Source/Editor/BlueprintGraph/Private/K2Node_TemporaryVariable.cpp
@@ -46,12 +46,26 @@ void UK2Node_TemporaryVariable::AllocateDefaultPins()
 
 FString UK2Node_TemporaryVariable::GetTooltip() const
 {
-	return FString::Printf(*NSLOCTEXT("K2Node", "LocalTemporaryVariable", "Local temporary %s variable").ToString(), *UEdGraphSchema_K2::TypeToString(VariableType));
+	FFormatNamedArguments Args;
+	Args.Add(TEXT("VariableType"), UEdGraphSchema_K2::TypeToText(VariableType));
+	return FText::Format(NSLOCTEXT("K2Node", "LocalTemporaryVariable", "Local temporary {VariableType} variable"), Args).ToString();
 }
 
-FString UK2Node_TemporaryVariable::GetNodeTitle(ENodeTitleType::Type TitleType) const
+FText UK2Node_TemporaryVariable::GetNodeTitle(ENodeTitleType::Type TitleType) const
 {
-	FString Result = !bIsPersistent ? NSLOCTEXT("K2Node", "LocalVariable", "Local ").ToString() : NSLOCTEXT("K2Node", "PersistentLocalVariable", "Persistent Local ").ToString();
+	FText Result = !bIsPersistent ? NSLOCTEXT("K2Node", "LocalVariable", "Local {VariableType}") : NSLOCTEXT("K2Node", "PersistentLocalVariable", "Persistent Local {VariableType}");
+
+	FFormatNamedArguments Args;
+	Args.Add(TEXT("VariableType"), UEdGraphSchema_K2::TypeToText(VariableType));
+	Result = FText::Format(Result, Args);
+	return Result;
+}
+
+FString UK2Node_TemporaryVariable::GetNodeNativeTitle(ENodeTitleType::Type TitleType) const
+{
+	// Do not setup this function for localization, intentionally left unlocalized!
+
+	FString Result = !bIsPersistent ? TEXT("Local ") : TEXT("Persistent Local ");
 	Result += UEdGraphSchema_K2::TypeToString(VariableType);
 	return Result;
 }
diff --git a/Engine/Source/Editor/BlueprintGraph/Private/K2Node_Timeline.cpp b/Engine/Source/Editor/BlueprintGraph/Private/K2Node_Timeline.cpp
--- a/Engine/Source/Editor/BlueprintGraph/Private/K2Node_Timeline.cpp
+++ b/Engine/Source/Editor/BlueprintGraph/Private/K2Node_Timeline.cpp
@@ -260,9 +260,9 @@ FLinearColor UK2Node_Timeline::GetNodeTitleColor() const
 	return FLinearColor(1.0f, 0.51f, 0.0f);
 }
 
-FString UK2Node_Timeline::GetNodeTitle(ENodeTitleType::Type TitleType) const
+FText UK2Node_Timeline::GetNodeTitle(ENodeTitleType::Type TitleType) const
 {
-	return TimelineName.ToString();
+	return FText::FromName(TimelineName);
 }
 
 UEdGraphPin* UK2Node_Timeline::GetDirectionPin() const
@@ -362,11 +362,16 @@ void FindExactTimelineDifference(struct FDiffResults& Results, FDiffSingleResult
 {
 	if(Tracks1.Num() != Tracks2.Num())
 	{
-		FString NodeName = *Result.Node1->GetNodeTitle(ENodeTitleType::ListView);
+		FText NodeName = Result.Node1->GetNodeTitle(ENodeTitleType::ListView);
+
+		FFormatNamedArguments Args;
+		Args.Add(TEXT("TrackType"), FText::FromString(TrackTypeStr));
+		Args.Add(TEXT("NodeName"), NodeName);
+
 		Result.Diff = EDiffType::TIMELINE_NUM_TRACKS;
-		Result.ToolTip =  FString::Printf(*LOCTEXT("DIF_TimelineNumTracksToolTip", "The number of %s tracks in Timeline '%s' has changed").ToString(), *TrackTypeStr, *NodeName);
+		Result.ToolTip =  FText::Format(LOCTEXT("DIF_TimelineNumTracksToolTip", "The number of {TrackType} tracks in Timeline '{NodeName}' has changed"), Args).ToString();
 		Result.DisplayColor = FLinearColor(0.05f,0.261f,0.775f);
-		Result.DisplayString = FString::Printf(*LOCTEXT("DIF_TimelineNumTracks", "%s Track Count '%s'").ToString(), *TrackTypeStr, *NodeName);
+		Result.DisplayString = FText::Format(LOCTEXT("DIF_TimelineNumTracks", "{TrackType} Track Count '{NodeName}'"), Args).ToString();
 		Results.Add(Result);
 		return;
 	}
@@ -376,11 +381,16 @@ void FindExactTimelineDifference(struct FDiffResults& Results, FDiffSingleResult
 		if(!(Tracks1[i] == Tracks2[i]))
 		{
 			FName TrackName = Tracks2[i].TrackName;
-			FString NodeName = *Result.Node1->GetNodeTitle(ENodeTitleType::ListView);
+			FText NodeName = Result.Node1->GetNodeTitle(ENodeTitleType::ListView);
+
+			FFormatNamedArguments Args;
+			Args.Add(TEXT("TrackName"), FText::FromName(TrackName));
+			Args.Add(TEXT("NodeName"), NodeName);
+
 			Result.Diff = EDiffType::TIMELINE_TRACK_MODIFIED;
-			Result.ToolTip =  FString::Printf(*LOCTEXT("DIF_TimelineTrackModifiedToolTip", "Track '%s' of Timeline '%s' was Modified").ToString(), *TrackName.ToString(),  *NodeName);
+			Result.ToolTip =  FText::Format(LOCTEXT("DIF_TimelineTrackModifiedToolTip", "Track '{TrackName}' of Timeline '{NodeName}' was Modified"), Args).ToString();
 			Result.DisplayColor = FLinearColor(0.75f,0.1f,0.15f);
-			Result.DisplayString = FString::Printf(*LOCTEXT("DIF_TimelineTrackModified", "Track Modified '%s'").ToString(), *TrackName.ToString());;
+			Result.DisplayString = FText::Format(LOCTEXT("DIF_TimelineTrackModified", "Track Modified '{TrackName}'"), Args).ToString();
 			Results.Add(Result);
 			break;
 		}
@@ -409,28 +419,42 @@ void UK2Node_Timeline::FindDiffs( class UEdGraphNode* OtherNode, struct FDiffRes
 		if(Template1->bAutoPlay != Template2->bAutoPlay)
 		{
 			Diff.Diff = EDiffType::TIMELINE_AUTOPLAY;
-			FString NodeName = *GetNodeTitle(ENodeTitleType::ListView);
-			Diff.ToolTip =  FString::Printf(*LOCTEXT("DIF_TimelineAutoPlayToolTip", "Timeline '%s' had its AutoPlay state changed").ToString(), *NodeName);
+			FText NodeName = GetNodeTitle(ENodeTitleType::ListView);
+
+			FFormatNamedArguments Args;
+			Args.Add(TEXT("NodeName"), NodeName);
+
+			Diff.ToolTip =  FText::Format(LOCTEXT("DIF_TimelineAutoPlayToolTip", "Timeline '{NodeName}' had its AutoPlay state changed"), Args).ToString();
 			Diff.DisplayColor = FLinearColor(0.15f,0.61f,0.15f);
-			Diff.DisplayString = FString::Printf(*LOCTEXT("DIF_TimelineAutoPlay", "Timeline AutoPlay Changed '%s'").ToString(), *NodeName);
+			Diff.DisplayString = FText::Format(LOCTEXT("DIF_TimelineAutoPlay", "Timeline AutoPlay Changed '{NodeName}'"), Args).ToString();
 			Results.Add(Diff);
 		}
 		if(Template1->bLoop != Template2->bLoop)
 		{
 			Diff.Diff = EDiffType::TIMELINE_LOOP;
-			FString NodeName = *GetNodeTitle(ENodeTitleType::ListView);
-			Diff.ToolTip =  FString::Printf(*LOCTEXT("DIF_TimelineLoopingToolTip", "Timeline '%s' had its looping state changed").ToString(), *NodeName);
+			FText NodeName = GetNodeTitle(ENodeTitleType::ListView);
+
+			FFormatNamedArguments Args;
+			Args.Add(TEXT("NodeName"), NodeName);
+
+			Diff.ToolTip =  FText::Format(LOCTEXT("DIF_TimelineLoopingToolTip", "Timeline '{NodeName}' had its looping state changed"), Args).ToString();
 			Diff.DisplayColor = FLinearColor(0.75f,0.1f,0.75f);
-			Diff.DisplayString =  FString::Printf(*LOCTEXT("DIF_TimelineLooping", "Timeline Loop Changed '%s'").ToString(), *NodeName);
+			Diff.DisplayString =  FText::Format(LOCTEXT("DIF_TimelineLooping", "Timeline Loop Changed '{NodeName}'"), Args).ToString();
 			Results.Add(Diff);
 		}
 		if(Template1->TimelineLength != Template2->TimelineLength)
 		{
-			FString NodeName = *GetNodeTitle(ENodeTitleType::ListView);
+			FText NodeName = GetNodeTitle(ENodeTitleType::ListView);
+
+			FFormatNamedArguments Args;
+			Args.Add(TEXT("NodeName"), NodeName);
+			Args.Add(TEXT("TimelineLength1"), Template1->TimelineLength);
+			Args.Add(TEXT("TimelineLength2"), Template2->TimelineLength);
+
 			Diff.Diff = EDiffType::TIMELINE_LENGTH;
-			Diff.ToolTip = FString::Printf(*LOCTEXT("DIF_TimelineLengthToolTip", "Length of Timeline '%s' has changed. Was %f, but is now %f").ToString(), *NodeName,Template1->TimelineLength, Template2->TimelineLength);
+			Diff.ToolTip = FText::Format(LOCTEXT("DIF_TimelineLengthToolTip", "Length of Timeline '{NodeName}' has changed. Was {TimelineLength1}, but is now {TimelineLength2}"), Args).ToString();
 			Diff.DisplayColor = FLinearColor(0.25f,0.1f,0.15f);
-			Diff.DisplayString =  FString::Printf(*LOCTEXT("DIF_TimelineLength", "Timeline Length '%s' [%f -> %f]").ToString(), *NodeName, Template1->TimelineLength, Template2->TimelineLength);
+			Diff.DisplayString =  FText::Format(LOCTEXT("DIF_TimelineLength", "Timeline Length '{NodeName}' [{TimelineLength1} -> {TimelineLength2}]"), Args).ToString();
 			Results.Add(Diff);
 		}
 
diff --git a/Engine/Source/Editor/BlueprintGraph/Private/K2Node_Tunnel.cpp b/Engine/Source/Editor/BlueprintGraph/Private/K2Node_Tunnel.cpp
--- a/Engine/Source/Editor/BlueprintGraph/Private/K2Node_Tunnel.cpp
+++ b/Engine/Source/Editor/BlueprintGraph/Private/K2Node_Tunnel.cpp
@@ -47,19 +47,19 @@ FString UK2Node_Tunnel::GetTooltip() const
 	}
 }
 
-FString UK2Node_Tunnel::GetNodeTitle(ENodeTitleType::Type TitleType) const
+FText UK2Node_Tunnel::GetNodeTitle(ENodeTitleType::Type TitleType) const
 {
 	if (bCanHaveInputs && !bCanHaveOutputs)
 	{
-		return NSLOCTEXT("K2Node", "OutputTunnelTitle", "Outputs").ToString();
+		return NSLOCTEXT("K2Node", "OutputTunnelTitle", "Outputs");
 	}
 	else if (!bCanHaveInputs && bCanHaveOutputs)
 	{
-		return NSLOCTEXT("K2Node", "InputTunnelTitle", "Inputs").ToString();
+		return NSLOCTEXT("K2Node", "InputTunnelTitle", "Inputs");
 	}
 	else
 	{
-		return NSLOCTEXT("K2Node", "TunnelConnectionTitle", "Tunnel Connection").ToString();
+		return NSLOCTEXT("K2Node", "TunnelConnectionTitle", "Tunnel Connection");
 	}
 }
 
diff --git a/Engine/Source/Editor/BlueprintGraph/Private/K2Node_Variable.cpp b/Engine/Source/Editor/BlueprintGraph/Private/K2Node_Variable.cpp
--- a/Engine/Source/Editor/BlueprintGraph/Private/K2Node_Variable.cpp
+++ b/Engine/Source/Editor/BlueprintGraph/Private/K2Node_Variable.cpp
@@ -93,7 +93,7 @@ void UK2Node_Variable::CreatePinForSelf()
 			UClass* MemberParentClass = VariableReference.GetMemberParentClass(this);
 			UClass* AuthoritativeClass = MemberParentClass ? MemberParentClass->GetAuthoritativeClass() : NULL;
 			UEdGraphPin* TargetPin = CreatePin(EGPD_Input, K2Schema->PC_Object, TEXT(""), AuthoritativeClass, false, false, K2Schema->PN_Self);
-			TargetPin->PinFriendlyName =  TEXT("Target");
+			TargetPin->PinFriendlyName =  LOCTEXT("Target", "Target");
 		}
 	}
 	else
diff --git a/Engine/Source/Editor/BlueprintGraph/Private/K2Node_VariableGet.cpp b/Engine/Source/Editor/BlueprintGraph/Private/K2Node_VariableGet.cpp
--- a/Engine/Source/Editor/BlueprintGraph/Private/K2Node_VariableGet.cpp
+++ b/Engine/Source/Editor/BlueprintGraph/Private/K2Node_VariableGet.cpp
@@ -7,6 +7,8 @@
 //////////////////////////////////////////////////////////////////////////
 // FKCHandler_VariableGet
 
+#define LOCTEXT_NAMESPACE "K2Node"
+
 class FKCHandler_VariableGet : public FNodeHandlingFunctor
 {
 public:
@@ -107,8 +109,41 @@ FString UK2Node_VariableGet::GetTooltip() const
 	return FText::Format( NSLOCTEXT( "K2Node", "GetVariable_ToolTip", "Read the value of variable {VarName}{TextPartition}{MetaData}"), Args ).ToString();
 }
 
-FString UK2Node_VariableGet::GetNodeTitle(ENodeTitleType::Type TitleType) const
+FText UK2Node_VariableGet::GetNodeTitle(ENodeTitleType::Type TitleType) const
 {
+	FText Result;
+
+	// If there is only one variable being read, the title can be made the variable name
+	FString OutputPinName;
+	int32 NumOutputsFound = 0;
+
+	for (int32 PinIndex = 0; PinIndex < Pins.Num(); ++PinIndex)
+	{
+		UEdGraphPin* Pin = Pins[PinIndex];
+		if (Pin->Direction == EGPD_Output)
+		{
+			++NumOutputsFound;
+			OutputPinName = Pin->PinName;
+		}
+	}
+
+	if (NumOutputsFound == 1)
+	{
+		FFormatNamedArguments Args;
+		Args.Add(TEXT("PinName"), FText::FromString(OutputPinName));
+		Result = FText::Format(LOCTEXT("GetPinName", "Get {PinName}"), Args);
+	}
+	else
+	{
+		Result = LOCTEXT("Get", "Get");
+	}
+
+	return Result;
+}
+
+FString UK2Node_VariableGet::GetNodeNativeTitle(ENodeTitleType::Type TitleType) const
+{
+	// Do not setup this function for localization, intentionally left unlocalized!
 	FString Result = TEXT("Get");
 
 	// If there is only one variable being read, the title can be made the variable name
@@ -137,3 +172,5 @@ FNodeHandlingFunctor* UK2Node_VariableGet::CreateNodeHandler(FKismetCompilerCont
 {
 	return new FKCHandler_VariableGet(CompilerContext);
 }
+
+#undef LOCTEXT_NAMESPACE
\ No newline at end of file
diff --git a/Engine/Source/Editor/BlueprintGraph/Private/K2Node_VariableSet.cpp b/Engine/Source/Editor/BlueprintGraph/Private/K2Node_VariableSet.cpp
--- a/Engine/Source/Editor/BlueprintGraph/Private/K2Node_VariableSet.cpp
+++ b/Engine/Source/Editor/BlueprintGraph/Private/K2Node_VariableSet.cpp
@@ -84,11 +84,55 @@ FString UK2Node_VariableSet::GetTooltip() const
 	return FText::Format( NSLOCTEXT( "K2Node", "SetValueOfVariable", "Set the value of variable {VarName}{ReplicationCall}{ReplicationNotifyName}{TextPartition}{MetaData}"), Args ).ToString();
 }
 
-FString UK2Node_VariableSet::GetNodeTitle(ENodeTitleType::Type TitleType) const
+FText UK2Node_VariableSet::GetNodeTitle(ENodeTitleType::Type TitleType) const
 {
 	const UEdGraphSchema_K2* K2Schema = GetDefault<UEdGraphSchema_K2>();
 
-	FString Result = HasLocalRepNotify() ? NSLOCTEXT("K2Node", "SetWithNotify", "Set with Notify").ToString() : NSLOCTEXT("K2Node", "Set", "Set").ToString();
+	FText Result = HasLocalRepNotify() ? NSLOCTEXT("K2Node", "SetWithNotify", "Set with Notify") : NSLOCTEXT("K2Node", "Set", "Set");
+
+	// If there is only one variable being written (one non-meta input pin), the title can be made the variable name
+	FString InputPinName;
+	int32 NumInputsFound = 0;
+
+	for (int32 PinIndex = 0; PinIndex < Pins.Num(); ++PinIndex)
+	{
+		UEdGraphPin* Pin = Pins[PinIndex];
+		if ((Pin->Direction == EGPD_Input) && (!K2Schema->IsMetaPin(*Pin)))
+		{
+			++NumInputsFound;
+			InputPinName = Pin->PinName;
+		}
+	}
+
+	if (NumInputsFound == 1)
+	{
+		FFormatNamedArguments Args;
+		Args.Add(TEXT("PinName"), FText::FromString(InputPinName));
+
+		if(HasLocalRepNotify())
+		{
+			Result = FText::Format(NSLOCTEXT("K2Node", "SetWithNotifyPinName", "Set with Notify {PinName}"), Args);
+		}
+		else
+		{
+			Result = FText::Format(NSLOCTEXT("K2Node", "SetPinName", "Set {PinName}"), Args);
+		}
+	}
+	else
+	{
+		Result = HasLocalRepNotify() ? NSLOCTEXT("K2Node", "SetWithNotify", "Set with Notify") : NSLOCTEXT("K2Node", "Set", "Set");
+	}
+
+	return Result;
+}
+
+FString UK2Node_VariableSet::GetNodeNativeTitle(ENodeTitleType::Type TitleType) const
+{
+	// Do not setup this function for localization, intentionally left unlocalized!
+
+	const UEdGraphSchema_K2* K2Schema = GetDefault<UEdGraphSchema_K2>();
+
+	FString Result = HasLocalRepNotify() ? TEXT("Set with Notify") : TEXT("Set");
 
 	// If there is only one variable being written (one non-meta input pin), the title can be made the variable name
 	FString InputPinName;
diff --git a/Engine/Source/Editor/BlueprintGraph/Private/K2Node_VariableSetRef.cpp b/Engine/Source/Editor/BlueprintGraph/Private/K2Node_VariableSetRef.cpp
--- a/Engine/Source/Editor/BlueprintGraph/Private/K2Node_VariableSetRef.cpp
+++ b/Engine/Source/Editor/BlueprintGraph/Private/K2Node_VariableSetRef.cpp
@@ -122,7 +122,7 @@ FString UK2Node_VariableSetRef::GetTooltip() const
 	return FString::Printf(*NSLOCTEXT("K2Node", "SetValueOfRefVariable", "Set the value of the connected pass-by-ref variable").ToString());
 }
 
-FString UK2Node_VariableSetRef::GetNodeTitle(ENodeTitleType::Type TitleType) const
+FText UK2Node_VariableSetRef::GetNodeTitle(ENodeTitleType::Type TitleType) const
 {
 	const UEdGraphSchema_K2* Schema = GetDefault<UEdGraphSchema_K2>();
 
@@ -130,11 +130,31 @@ FString UK2Node_VariableSetRef::GetNodeTitle(ENodeTitleType::Type TitleType) con
 
 	if( TargetPin && TargetPin->PinType.PinCategory != Schema->PC_Wildcard )
 	{
-		return FString::Printf(*NSLOCTEXT("K2Node", "SetRefVarNodeTitle_Typed", "Set %s").ToString(), *Schema->TypeToString(TargetPin->PinType));
+		FFormatNamedArguments Args;
+		Args.Add(TEXT("PinType"), Schema->TypeToText(TargetPin->PinType));
+		return FText::Format(NSLOCTEXT("K2Node", "SetRefVarNodeTitle_Typed", "Set {PinType}"), Args);
 	}
 	else
 	{
-		return FString::Printf(*NSLOCTEXT("K2Node", "SetRefVarNodeTitle", "Set By-Ref Var").ToString());
+		return NSLOCTEXT("K2Node", "SetRefVarNodeTitle", "Set By-Ref Var");
+	}
+}
+
+FString UK2Node_VariableSetRef::GetNodeNativeTitle(ENodeTitleType::Type TitleType) const
+{
+	// Do not setup this function for localization, intentionally left unlocalized!
+	
+	const UEdGraphSchema_K2* Schema = GetDefault<UEdGraphSchema_K2>();
+
+	UEdGraphPin* TargetPin = GetTargetPin();
+
+	if( TargetPin && TargetPin->PinType.PinCategory != Schema->PC_Wildcard )
+	{
+		return FString::Printf(TEXT("Set %s"), *Schema->TypeToString(TargetPin->PinType));
+	}
+	else
+	{
+		return TEXT("Set By-Ref Var");
 	}
 }
 
diff --git a/Engine/Source/Editor/BlueprintGraph/Public/K2ActionMenuBuilder.h b/Engine/Source/Editor/BlueprintGraph/Public/K2ActionMenuBuilder.h
--- a/Engine/Source/Editor/BlueprintGraph/Public/K2ActionMenuBuilder.h
+++ b/Engine/Source/Editor/BlueprintGraph/Public/K2ActionMenuBuilder.h
@@ -43,7 +43,7 @@ public:
 	static TSharedPtr<FEdGraphSchemaAction_K2AddComponent> CreateAddComponentAction(UObject* TemporaryOuter, const UBlueprint* Blueprint, TSubclassOf<UActorComponent> DestinationComponentType, UObject* Asset);
 
 	/** Utility to create a node creating action */
-	BLUEPRINTGRAPH_API static TSharedPtr<FEdGraphSchemaAction_K2NewNode> AddNewNodeAction(FGraphActionListBuilderBase& ContextMenuBuilder, const FString& Category, const FString& MenuDesc, const FString& Tooltip, const int32 Grouping = 0, const FString& Keywords = FString());
+	BLUEPRINTGRAPH_API static TSharedPtr<FEdGraphSchemaAction_K2NewNode> AddNewNodeAction(FGraphActionListBuilderBase& ContextMenuBuilder, const FString& Category, const FText& MenuDesc, const FString& Tooltip, const int32 Grouping = 0, const FString& Keywords = FString());
 
 	/** Utility that creates a 'new event' entry based on the function supplied. Can either override that function, or simply use the function as its signature. */
 	BLUEPRINTGRAPH_API static void AddEventForFunction(FGraphActionListBuilderBase& ActionMenuBuilder, class UFunction* Function);
diff --git a/Engine/Source/Editor/DetailCustomizations/Private/AnimTransitionNodeDetails.cpp b/Engine/Source/Editor/DetailCustomizations/Private/AnimTransitionNodeDetails.cpp
--- a/Engine/Source/Editor/DetailCustomizations/Private/AnimTransitionNodeDetails.cpp
+++ b/Engine/Source/Editor/DetailCustomizations/Private/AnimTransitionNodeDetails.cpp
@@ -95,7 +95,7 @@ void FAnimTransitionNodeDetails::CustomizeDetails( IDetailLayoutBuilder& DetailB
 					CanExecPin = ResultNode->FindPin(TEXT("bCanEnterTransition"));
 				}
 			}
-			TransitionCategory.AddCustomRow( CanExecPin ? CanExecPin->PinFriendlyName : FString() )
+			TransitionCategory.AddCustomRow( CanExecPin ? CanExecPin->PinFriendlyName.ToString() : FString() )
 			[
 				SNew(SKismetLinearExpression, CanExecPin)
 			];
diff --git a/Engine/Source/Editor/EnvironmentQueryEditor/Classes/EdGraphSchema_EnvironmentQuery.h b/Engine/Source/Editor/EnvironmentQueryEditor/Classes/EdGraphSchema_EnvironmentQuery.h
--- a/Engine/Source/Editor/EnvironmentQueryEditor/Classes/EdGraphSchema_EnvironmentQuery.h
+++ b/Engine/Source/Editor/EnvironmentQueryEditor/Classes/EdGraphSchema_EnvironmentQuery.h
@@ -20,7 +20,7 @@ struct FEnvironmentQuerySchemaAction_NewNode : public FEdGraphSchemaAction
 		, NodeTemplate(NULL)
 	{}
 
-	FEnvironmentQuerySchemaAction_NewNode(const FString& InNodeCategory, const FString& InMenuDesc, const FString& InToolTip, const int32 InGrouping)
+	FEnvironmentQuerySchemaAction_NewNode(const FString& InNodeCategory, const FText& InMenuDesc, const FString& InToolTip, const int32 InGrouping)
 		: FEdGraphSchemaAction(InNodeCategory, InMenuDesc, InToolTip, InGrouping) 
 		, NodeTemplate(NULL)
 	{}
@@ -60,7 +60,7 @@ struct FEnvironmentQuerySchemaAction_NewSubNode : public FEdGraphSchemaAction
 		, NodeTemplate(NULL)
 	{}
 
-	FEnvironmentQuerySchemaAction_NewSubNode(const FString& InNodeCategory, const FString& InMenuDesc, const FString& InToolTip, const int32 InGrouping)
+	FEnvironmentQuerySchemaAction_NewSubNode(const FString& InNodeCategory, const FText& InMenuDesc, const FString& InToolTip, const int32 InGrouping)
 		: FEdGraphSchemaAction(InNodeCategory, InMenuDesc, InToolTip, InGrouping) 
 		, NodeTemplate(NULL)
 	{}
diff --git a/Engine/Source/Editor/EnvironmentQueryEditor/Classes/EnvironmentQueryGraphNode.h b/Engine/Source/Editor/EnvironmentQueryEditor/Classes/EnvironmentQueryGraphNode.h
--- a/Engine/Source/Editor/EnvironmentQueryEditor/Classes/EnvironmentQueryGraphNode.h
+++ b/Engine/Source/Editor/EnvironmentQueryEditor/Classes/EnvironmentQueryGraphNode.h
@@ -30,8 +30,8 @@ class UEnvironmentQueryGraphNode : public UEdGraphNode
 	//
 	virtual UEdGraph* GetBoundGraph() const { return NULL; }
 
-	virtual FString GetNodeTitle(ENodeTitleType::Type TitleType) const OVERRIDE;
-	virtual FString GetDescription() const;
+	virtual FText GetNodeTitle(ENodeTitleType::Type TitleType) const OVERRIDE;
+	virtual FText GetDescription() const;
 
 	virtual void NodeConnectionListChanged() OVERRIDE;
 
diff --git a/Engine/Source/Editor/EnvironmentQueryEditor/Classes/EnvironmentQueryGraphNode_Option.h b/Engine/Source/Editor/EnvironmentQueryEditor/Classes/EnvironmentQueryGraphNode_Option.h
--- a/Engine/Source/Editor/EnvironmentQueryEditor/Classes/EnvironmentQueryGraphNode_Option.h
+++ b/Engine/Source/Editor/EnvironmentQueryEditor/Classes/EnvironmentQueryGraphNode_Option.h
@@ -14,8 +14,8 @@ class UEnvironmentQueryGraphNode_Option : public UEnvironmentQueryGraphNode
 	virtual void AllocateDefaultPins() OVERRIDE;
 	virtual void PostPlacedNewNode() OVERRIDE;
 	virtual void GetContextMenuActions(const FGraphNodeContextMenuBuilder& Context) const OVERRIDE;
-	virtual FString GetNodeTitle(ENodeTitleType::Type TitleType) const OVERRIDE;
-	virtual FString GetDescription() const OVERRIDE;
+	virtual FText GetNodeTitle(ENodeTitleType::Type TitleType) const OVERRIDE;
+	virtual FText GetDescription() const OVERRIDE;
 
 	void AddSubNode(class UEnvironmentQueryGraphNode_Test* NodeTemplate, class UEdGraph* ParentGraph);
 	void CreateAddTestSubMenu(class FMenuBuilder& MenuBuilder, UEdGraph* Graph) const;
diff --git a/Engine/Source/Editor/EnvironmentQueryEditor/Classes/EnvironmentQueryGraphNode_Test.h b/Engine/Source/Editor/EnvironmentQueryEditor/Classes/EnvironmentQueryGraphNode_Test.h
--- a/Engine/Source/Editor/EnvironmentQueryEditor/Classes/EnvironmentQueryGraphNode_Test.h
+++ b/Engine/Source/Editor/EnvironmentQueryEditor/Classes/EnvironmentQueryGraphNode_Test.h
@@ -26,8 +26,8 @@ class UEnvironmentQueryGraphNode_Test : public UEnvironmentQueryGraphNode
 	virtual void PostPlacedNewNode() OVERRIDE;
 	virtual void DestroyNode() OVERRIDE;
 	virtual bool IsSubNode() const OVERRIDE { return true; }
-	virtual FString GetNodeTitle(ENodeTitleType::Type TitleType) const OVERRIDE;
-	virtual FString GetDescription() const OVERRIDE;
+	virtual FText GetNodeTitle(ENodeTitleType::Type TitleType) const OVERRIDE;
+	virtual FText GetDescription() const OVERRIDE;
 
 	void SetDisplayedWeight(float Pct, bool bNamed);
 };
diff --git a/Engine/Source/Editor/EnvironmentQueryEditor/Private/EdGraphSchema_EnvironmentQuery.cpp b/Engine/Source/Editor/EnvironmentQueryEditor/Private/EdGraphSchema_EnvironmentQuery.cpp
--- a/Engine/Source/Editor/EnvironmentQueryEditor/Private/EdGraphSchema_EnvironmentQuery.cpp
+++ b/Engine/Source/Editor/EnvironmentQueryEditor/Private/EdGraphSchema_EnvironmentQuery.cpp
@@ -118,14 +118,14 @@ UEdGraphSchema_EnvironmentQuery::UEdGraphSchema_EnvironmentQuery(const class FPo
 {
 }
 
-TSharedPtr<FEnvironmentQuerySchemaAction_NewNode> AddNewNodeAction(FGraphContextMenuBuilder& ContextMenuBuilder, const FString& Category, const FString& MenuDesc, const FString& Tooltip)
+TSharedPtr<FEnvironmentQuerySchemaAction_NewNode> AddNewNodeAction(FGraphContextMenuBuilder& ContextMenuBuilder, const FString& Category, const FText& MenuDesc, const FString& Tooltip)
 {
 	const TSharedRef<FEnvironmentQuerySchemaAction_NewNode> NewAction = MakeShareable( new FEnvironmentQuerySchemaAction_NewNode(Category, MenuDesc, Tooltip, 0) );
 	ContextMenuBuilder.AddAction( NewAction );
 	return NewAction;
 }
 
-TSharedPtr<FEnvironmentQuerySchemaAction_NewSubNode> AddNewSubNodeAction(FGraphContextMenuBuilder& ContextMenuBuilder, const FString& Category, const FString& MenuDesc, const FString& Tooltip)
+TSharedPtr<FEnvironmentQuerySchemaAction_NewSubNode> AddNewSubNodeAction(FGraphContextMenuBuilder& ContextMenuBuilder, const FString& Category, const FText& MenuDesc, const FString& Tooltip)
 {
 	TSharedPtr<FEnvironmentQuerySchemaAction_NewSubNode> NewAction = MakeShareable(new FEnvironmentQuerySchemaAction_NewSubNode(Category, MenuDesc, Tooltip, 0));
 	ContextMenuBuilder.AddAction( NewAction );
@@ -219,11 +219,12 @@ void UEdGraphSchema_EnvironmentQuery::GetGraphContextActions(FGraphContextMenuBu
 	
 	for (int32 i = 0; i < AllowedClasses.Num(); i++)
 	{
-		TSharedPtr<FEnvironmentQuerySchemaAction_NewNode> AddOpAction = AddNewNodeAction(ContextMenuBuilder, "Generators", GetNodeDescriptionHelper(AllowedClasses[i]), FString());
+		TSharedPtr<FEnvironmentQuerySchemaAction_NewNode> AddOpAction = AddNewNodeAction(ContextMenuBuilder, "Generators", FText::FromString(GetNodeDescriptionHelper(AllowedClasses[i])), FString());
 		
 		UEnvironmentQueryGraphNode_Option* OpNode = NewObject<UEnvironmentQueryGraphNode_Option>(ContextMenuBuilder.OwnerOfTemporaries);
 		OpNode->EnvQueryNodeClass = AllowedClasses[i];
 		AddOpAction->NodeTemplate = OpNode;		
+		AddOpAction->SearchTitle = AddOpAction->NodeTemplate->GetNodeSearchTitle();
 	}
 }
 
@@ -240,9 +241,10 @@ void UEdGraphSchema_EnvironmentQuery::GetGraphNodeContextActions(FGraphContextMe
 		UEnvironmentQueryGraphNode_Test* OpNode = NewObject<UEnvironmentQueryGraphNode_Test>(Graph);
 		OpNode->EnvQueryNodeClass = AllowedClasses[i];
 
-		TSharedPtr<FEnvironmentQuerySchemaAction_NewSubNode> AddOpAction = AddNewSubNodeAction(ContextMenuBuilder, TEXT(""), GetNodeDescriptionHelper(AllowedClasses[i]), "");
+		TSharedPtr<FEnvironmentQuerySchemaAction_NewSubNode> AddOpAction = AddNewSubNodeAction(ContextMenuBuilder, TEXT(""), FText::FromString(GetNodeDescriptionHelper(AllowedClasses[i])), "");
 		AddOpAction->ParentNode = SelectedOption;
 		AddOpAction->NodeTemplate = OpNode;
+		AddOpAction->SearchTitle = AddOpAction->NodeTemplate->GetNodeSearchTitle();
 	}
 }
 
@@ -295,7 +297,7 @@ void UEdGraphSchema_EnvironmentQuery::GetBreakLinkToSubMenuActions( class FMenuB
 	for(TArray<class UEdGraphPin*>::TConstIterator Links(InGraphPin->LinkedTo); Links; ++Links)
 	{
 		UEdGraphPin* Pin = *Links;
-		FString TitleString = Pin->GetOwningNode()->GetNodeTitle(ENodeTitleType::ListView);
+		FString TitleString = Pin->GetOwningNode()->GetNodeTitle(ENodeTitleType::ListView).ToString();
 		FText Title = FText::FromString( TitleString );
 		if ( Pin->PinName != TEXT("") )
 		{
diff --git a/Engine/Source/Editor/EnvironmentQueryEditor/Private/EnvironmentQueryGraphNode.cpp b/Engine/Source/Editor/EnvironmentQueryEditor/Private/EnvironmentQueryGraphNode.cpp
--- a/Engine/Source/Editor/EnvironmentQueryEditor/Private/EnvironmentQueryGraphNode.cpp
+++ b/Engine/Source/Editor/EnvironmentQueryEditor/Private/EnvironmentQueryGraphNode.cpp
@@ -38,14 +38,14 @@ void UEnvironmentQueryGraphNode::ResetNodeOwner()
 	}
 }
 
-FString UEnvironmentQueryGraphNode::GetNodeTitle(ENodeTitleType::Type TitleType) const
+FText UEnvironmentQueryGraphNode::GetNodeTitle(ENodeTitleType::Type TitleType) const
 {
-	return FString();
+	return FText::GetEmpty();
 }
 
-FString UEnvironmentQueryGraphNode::GetDescription() const
+FText UEnvironmentQueryGraphNode::GetDescription() const
 {
-	return FString();
+	return FText::GetEmpty();
 }
 
 UEdGraphPin* UEnvironmentQueryGraphNode::GetInputPin(int32 InputIndex) const
diff --git a/Engine/Source/Editor/EnvironmentQueryEditor/Private/EnvironmentQueryGraphNode_Option.cpp b/Engine/Source/Editor/EnvironmentQueryEditor/Private/EnvironmentQueryGraphNode_Option.cpp
--- a/Engine/Source/Editor/EnvironmentQueryEditor/Private/EnvironmentQueryGraphNode_Option.cpp
+++ b/Engine/Source/Editor/EnvironmentQueryEditor/Private/EnvironmentQueryGraphNode_Option.cpp
@@ -31,16 +31,16 @@ void UEnvironmentQueryGraphNode_Option::PostPlacedNewNode()
 	}
 }
 
-FString UEnvironmentQueryGraphNode_Option::GetNodeTitle(ENodeTitleType::Type TitleType) const
+FText UEnvironmentQueryGraphNode_Option::GetNodeTitle(ENodeTitleType::Type TitleType) const
 {
 	UEnvQueryOption* OptionInstance = Cast<UEnvQueryOption>(NodeInstance);
-	return OptionInstance ? OptionInstance->GetDescriptionTitle() : FString();
+	return OptionInstance ? OptionInstance->GetDescriptionTitle() : FText::GetEmpty();
 }
 
-FString UEnvironmentQueryGraphNode_Option::GetDescription() const
+FText UEnvironmentQueryGraphNode_Option::GetDescription() const
 {
 	UEnvQueryOption* OptionInstance = Cast<UEnvQueryOption>(NodeInstance);
-	return OptionInstance ? OptionInstance->GetDescriptionDetails() : FString();
+	return OptionInstance ? OptionInstance->GetDescriptionDetails() : FText::GetEmpty();
 }
 
 void UEnvironmentQueryGraphNode_Option::AddSubNode(UEnvironmentQueryGraphNode_Test* NodeTemplate, class UEdGraph* ParentGraph)
diff --git a/Engine/Source/Editor/EnvironmentQueryEditor/Private/EnvironmentQueryGraphNode_Root.cpp b/Engine/Source/Editor/EnvironmentQueryEditor/Private/EnvironmentQueryGraphNode_Root.cpp
--- a/Engine/Source/Editor/EnvironmentQueryEditor/Private/EnvironmentQueryGraphNode_Root.cpp
+++ b/Engine/Source/Editor/EnvironmentQueryEditor/Private/EnvironmentQueryGraphNode_Root.cpp
@@ -11,7 +11,7 @@ void UEnvironmentQueryGraphNode_Root::AllocateDefaultPins()
 	UEdGraphPin* Outputs = CreatePin(EGPD_Output, TEXT("Transition"), TEXT(""), NULL, false, false, TEXT("In"));
 }
 
-FString UEnvironmentQueryGraphNode_Root::GetNodeTitle(ENodeTitleType::Type TitleType) const
+FText UEnvironmentQueryGraphNode_Root::GetNodeTitle(ENodeTitleType::Type TitleType) const
 {
-	return FString("ROOT");
+	return NSLOCTEXT("EnvironmentQueryGraphNode_Root", "Root", "ROOT");
 }
diff --git a/Engine/Source/Editor/EnvironmentQueryEditor/Private/EnvironmentQueryGraphNode_Test.cpp b/Engine/Source/Editor/EnvironmentQueryEditor/Private/EnvironmentQueryGraphNode_Test.cpp
--- a/Engine/Source/Editor/EnvironmentQueryEditor/Private/EnvironmentQueryGraphNode_Test.cpp
+++ b/Engine/Source/Editor/EnvironmentQueryEditor/Private/EnvironmentQueryGraphNode_Test.cpp
@@ -35,16 +35,16 @@ void UEnvironmentQueryGraphNode_Test::DestroyNode()
 	Super::DestroyNode();
 }
 
-FString UEnvironmentQueryGraphNode_Test::GetNodeTitle(ENodeTitleType::Type TitleType) const
+FText UEnvironmentQueryGraphNode_Test::GetNodeTitle(ENodeTitleType::Type TitleType) const
 {
 	UEnvQueryTest* TestInstance = Cast<UEnvQueryTest>(NodeInstance);
-	return TestInstance ? TestInstance->GetDescriptionTitle() : FString();
+	return TestInstance ? FText::FromString(TestInstance->GetDescriptionTitle()) : FText::GetEmpty();
 }
 
-FString UEnvironmentQueryGraphNode_Test::GetDescription() const
+FText UEnvironmentQueryGraphNode_Test::GetDescription() const
 {
 	UEnvQueryTest* TestInstance = Cast<UEnvQueryTest>(NodeInstance);
-	return TestInstance ? TestInstance->GetDescriptionDetails() : FString();
+	return TestInstance ? TestInstance->GetDescriptionDetails() : FText::GetEmpty();
 }
 
 void UEnvironmentQueryGraphNode_Test::SetDisplayedWeight(float Pct, bool bNamed)
diff --git a/Engine/Source/Editor/EnvironmentQueryEditor/Private/SGraphNode_EnvironmentQuery.cpp b/Engine/Source/Editor/EnvironmentQueryEditor/Private/SGraphNode_EnvironmentQuery.cpp
--- a/Engine/Source/Editor/EnvironmentQueryEditor/Private/SGraphNode_EnvironmentQuery.cpp
+++ b/Engine/Source/Editor/EnvironmentQueryEditor/Private/SGraphNode_EnvironmentQuery.cpp
@@ -333,10 +333,10 @@ void SGraphNode_EnvironmentQuery::UpdateGraphNode()
 	CreatePinWidgets();
 }
 
-FString	SGraphNode_EnvironmentQuery::GetDescription() const
+FText SGraphNode_EnvironmentQuery::GetDescription() const
 {
 	UEnvironmentQueryGraphNode* StateNode = CastChecked<UEnvironmentQueryGraphNode>(GraphNode);
-	return StateNode ? StateNode->GetDescription() : FString();
+	return StateNode ? StateNode->GetDescription() : FText::GetEmpty();
 }
 
 void SGraphNode_EnvironmentQuery::CreatePinWidgets()
diff --git a/Engine/Source/Editor/EnvironmentQueryEditor/Private/SGraphNode_EnvironmentQuery.h b/Engine/Source/Editor/EnvironmentQueryEditor/Private/SGraphNode_EnvironmentQuery.h
--- a/Engine/Source/Editor/EnvironmentQueryEditor/Private/SGraphNode_EnvironmentQuery.h
+++ b/Engine/Source/Editor/EnvironmentQueryEditor/Private/SGraphNode_EnvironmentQuery.h
@@ -60,7 +60,7 @@ protected:
 
 	FSlateColor GetBorderBackgroundColor() const;
 	FSlateColor GetBackgroundColor() const;
-	FString GetDescription() const;
+	FText GetDescription() const;
 
 	virtual FString GetPreviewCornerText() const;
 	virtual const FSlateBrush* GetNameIcon() const;
diff --git a/Engine/Source/Editor/GameplayTagsEditor/Private/GameplayTagsK2Node_LiteralGameplayTag.cpp b/Engine/Source/Editor/GameplayTagsEditor/Private/GameplayTagsK2Node_LiteralGameplayTag.cpp
--- a/Engine/Source/Editor/GameplayTagsEditor/Private/GameplayTagsK2Node_LiteralGameplayTag.cpp
+++ b/Engine/Source/Editor/GameplayTagsEditor/Private/GameplayTagsK2Node_LiteralGameplayTag.cpp
@@ -25,9 +25,9 @@ FLinearColor UGameplayTagsK2Node_LiteralGameplayTag::GetNodeTitleColor() const
 	return FLinearColor(1.0f, 0.51f, 0.0f);
 }
 
-FString UGameplayTagsK2Node_LiteralGameplayTag::GetNodeTitle(ENodeTitleType::Type TitleType) const
+FText UGameplayTagsK2Node_LiteralGameplayTag::GetNodeTitle(ENodeTitleType::Type TitleType) const
 {
-	return NSLOCTEXT("K2Node", "LiteralGameplayTag", "Make Literal GameplayTagContainer").ToString();
+	return NSLOCTEXT("K2Node", "LiteralGameplayTag", "Make Literal GameplayTagContainer");
 }
 
 bool UGameplayTagsK2Node_LiteralGameplayTag::CanCreateUnderSpecifiedSchema( const UEdGraphSchema* Schema ) const
@@ -44,12 +44,13 @@ void UGameplayTagsK2Node_LiteralGameplayTag::GetMenuEntries( FGraphContextMenuBu
 	UK2Node* EnumNodeTemplate = ContextMenuBuilder.CreateTemplateNode<UGameplayTagsK2Node_LiteralGameplayTag>();
 
 	const FString Category = FunctionCategory + TEXT("|Game");
-	const FString MenuDesc = EnumNodeTemplate->GetNodeTitle(ENodeTitleType::ListView);
+	const FText MenuDesc = EnumNodeTemplate->GetNodeTitle(ENodeTitleType::ListView);
 	const FString Tooltip = EnumNodeTemplate->GetTooltip();
 	const FString Keywords = EnumNodeTemplate->GetKeywords();
 
 	TSharedPtr<FEdGraphSchemaAction_K2NewNode> NodeAction = FK2ActionMenuBuilder::AddNewNodeAction(ContextMenuBuilder, Category, MenuDesc, Tooltip, 0, Keywords);
 	NodeAction->NodeTemplate = EnumNodeTemplate;
+	NodeAction->SearchTitle = EnumNodeTemplate->GetNodeSearchTitle();
 }
 
 void UGameplayTagsK2Node_LiteralGameplayTag::ExpandNode(class FKismetCompilerContext& CompilerContext, UEdGraph* SourceGraph)
diff --git a/Engine/Source/Editor/GraphEditor/Private/ConnectionDrawingPolicy.cpp b/Engine/Source/Editor/GraphEditor/Private/ConnectionDrawingPolicy.cpp
--- a/Engine/Source/Editor/GraphEditor/Private/ConnectionDrawingPolicy.cpp
+++ b/Engine/Source/Editor/GraphEditor/Private/ConnectionDrawingPolicy.cpp
@@ -731,7 +731,7 @@ void FKismetConnectionDrawingPolicy::BuildExecutionRoadmap()
 			else 
 			{
 				UE_LOG(LogConnectionDrawingPolicy, Warning, TEXT("Looks like a wire-trace was not injected before the jump from '%s' to '%s'."), 
-					*CurNode->GetNodeTitle(ENodeTitleType::FullTitle), *NextNode->GetNodeTitle(ENodeTitleType::FullTitle));
+					*CurNode->GetNodeTitle(ENodeTitleType::FullTitle).ToString(), *NextNode->GetNodeTitle(ENodeTitleType::FullTitle).ToString());
 			}
 
 			// clear the exec-pin (we're moving to a new node and want to find 
diff --git a/Engine/Source/Editor/GraphEditor/Private/DragConnection.cpp b/Engine/Source/Editor/GraphEditor/Private/DragConnection.cpp
--- a/Engine/Source/Editor/GraphEditor/Private/DragConnection.cpp
+++ b/Engine/Source/Editor/GraphEditor/Private/DragConnection.cpp
@@ -98,7 +98,7 @@ void FDragConnection::HoverTargetChanged()
 		SetSimpleFeedbackMessage(
 			FEditorStyle::GetBrush(TEXT("Graph.ConnectorFeedback.NewNode")),
 			FLinearColor::White,
-			NSLOCTEXT("GraphEditor.Feedback", "PlaceNewNode", "Place a new node.").ToString());
+			NSLOCTEXT("GraphEditor.Feedback", "PlaceNewNode", "Place a new node."));
 	}
 	else
 	{
diff --git a/Engine/Source/Editor/GraphEditor/Private/DragNode.cpp b/Engine/Source/Editor/GraphEditor/Private/DragNode.cpp
--- a/Engine/Source/Editor/GraphEditor/Private/DragNode.cpp
+++ b/Engine/Source/Editor/GraphEditor/Private/DragNode.cpp
@@ -80,7 +80,7 @@ void FDragNode::HoverTargetChanged()
 		SetSimpleFeedbackMessage(
 			FEditorStyle::GetBrush(TEXT("Graph.ConnectorFeedback.Error")),
 			FLinearColor::White,
-			NSLOCTEXT("GraphEditor.Feedback", "DragNode", "This node cannot be placed here.").ToString());
+			NSLOCTEXT("GraphEditor.Feedback", "DragNode", "This node cannot be placed here."));
 	}
 	else
 	{
diff --git a/Engine/Source/Editor/GraphEditor/Private/GraphActionNode.h b/Engine/Source/Editor/GraphEditor/Private/GraphActionNode.h
--- a/Engine/Source/Editor/GraphEditor/Private/GraphActionNode.h
+++ b/Engine/Source/Editor/GraphEditor/Private/GraphActionNode.h
@@ -107,7 +107,7 @@ public:
 						TSharedPtr<FEdGraphSchemaAction> CurrentAction = NodeToAdd->Actions[0];
 
 						bFoundMatch = CurrentAction->IsParentable()
-							&& NextCategory == CurrentChild->Actions[0]->MenuDescription;
+							&& NextCategory == CurrentChild->Actions[0]->MenuDescription.ToString();
 					}
 				}
 				
@@ -270,7 +270,7 @@ public:
 				}
 				else
 				{
-					return (NodeA->Actions[0]->MenuDescription <= NodeB->Actions[0]->MenuDescription);
+					return (NodeA->Actions[0]->MenuDescription.CompareTo(NodeB->Actions[0]->MenuDescription) <= 0);
 				}
 			}
 			else
@@ -314,7 +314,7 @@ public:
 				}
 				else
 				{
-					return (NodeA->Actions[0]->MenuDescription <= NodeB->Actions[0]->MenuDescription);
+					return (NodeA->Actions[0]->MenuDescription.CompareTo(NodeB->Actions[0]->MenuDescription) <= 0);
 				}
 			}
 			else
diff --git a/Engine/Source/Editor/GraphEditor/Private/GraphDiffControl.cpp b/Engine/Source/Editor/GraphEditor/Private/GraphDiffControl.cpp
--- a/Engine/Source/Editor/GraphEditor/Private/GraphDiffControl.cpp
+++ b/Engine/Source/Editor/GraphEditor/Private/GraphDiffControl.cpp
@@ -63,10 +63,10 @@ static bool IsNodeMatch(class UEdGraphNode* Node1, class UEdGraphNode* Node2, TA
 
 	// the name hashes won't match for nodes from separate graph assets, so we 
 	// need to look for some kind of semblance between the two... title? (what's displayed to the user)
-	FString Title1 = Node1->GetNodeTitle(ENodeTitleType::FullTitle);
-	FString Title2 = Node2->GetNodeTitle(ENodeTitleType::FullTitle);
+	FText Title1 = Node1->GetNodeTitle(ENodeTitleType::FullTitle);
+	FText Title2 = Node2->GetNodeTitle(ENodeTitleType::FullTitle);
 
-	return (Title1 == Title2);
+	return Title1.CompareTo(Title2) == 0;
 }
 
 /** Diff result when a node was added to the graph */
@@ -77,7 +77,10 @@ static void DiffR_NodeAdded( FDiffResults& Results, UEdGraphNode* Node )
 		FDiffSingleResult Diff;
 		Diff.Diff = EDiffType::NODE_ADDED;
 		Diff.Node1 = Node;
-		Diff.ToolTip =  FString::Printf(*LOCTEXT("DIF_AddNode", "Added Node '%s'").ToString(), *Node->GetNodeTitle(ENodeTitleType::ListView));
+
+		FFormatNamedArguments Args;
+		Args.Add(TEXT("NodeTitle"), Node->GetNodeTitle(ENodeTitleType::ListView));
+		Diff.ToolTip =  FText::Format(LOCTEXT("DIF_AddNode", "Added Node '{NodeTitle}'"), Args).ToString();
 		Diff.DisplayString = Diff.ToolTip;
 		Diff.DisplayColor = FLinearColor(0.3f,1.0f,0.4f);
 		Results.Add(Diff);
@@ -93,7 +96,11 @@ static void DiffR_NodeCommentChanged(FDiffResults& Results, class UEdGraphNode*
 		Diff.Diff = EDiffType::NODE_COMMENT;
 		Diff.Node1 = Node;
 		Diff.Node2 = Node2;
-		Diff.ToolTip =  FString::Printf(*LOCTEXT("DIF_CommentModified", "Comment Modified Node '%s'").ToString(), *Node->GetNodeTitle(ENodeTitleType::ListView));
+
+		FFormatNamedArguments Args;
+		Args.Add(TEXT("NodeTitle"), Node->GetNodeTitle(ENodeTitleType::ListView));
+		Diff.ToolTip =  FText::Format(LOCTEXT("DIF_CommentModified", "Comment Modified Node '{NodeTitle}'"), Args).ToString();
+
 		Diff.DisplayString = Diff.ToolTip;
 		Diff.DisplayColor = FLinearColor(0.25f,0.4f,0.5f);
 		Results.Add(Diff);
@@ -109,7 +116,11 @@ static void DiffR_NodeMoved(FDiffResults& Results, UEdGraphNode* Node, class UEd
 		Diff.Diff = EDiffType::NODE_MOVED;
 		Diff.Node1 = Node;
 		Diff.Node2 = OtherNode;
-		Diff.ToolTip = FString::Printf(*LOCTEXT("DIF_MoveNode", "Moved Node '%s'").ToString(), *Node->GetNodeTitle(ENodeTitleType::ListView));
+
+		FFormatNamedArguments Args;
+		Args.Add(TEXT("NodeTitle"), Node->GetNodeTitle(ENodeTitleType::ListView));
+		Diff.ToolTip = FText::Format(LOCTEXT("DIF_MoveNode", "Moved Node '{NodeTitle}'"), Args).ToString();
+
 		Diff.DisplayString = Diff.ToolTip;
 		Diff.DisplayColor = FLinearColor(0.9f, 0.84f, 0.43f);
 		Results.Add(Diff);
@@ -221,11 +232,17 @@ static void DiffR_LinkedToNode(FDiffResults& Results, class UEdGraphPin* Pin1, c
 		Diff.Node1 = Node1;
 		Diff.Node2 = Node2;
 
-		FString Node1Name = *Node1->GetNodeTitle(ENodeTitleType::ListView);
-		FString Node2Name = *Node2->GetNodeTitle(ENodeTitleType::ListView);
-		Diff.ToolTip =  FString::Printf(*LOCTEXT("DIF_PinLinkMovedToolTip", "Pin '%s' was linked to Node '%s', but is now linked to Node '%s'").ToString(), *Pin1->PinName,  *Node1Name, *Node2Name);
+		FText Node1Name = Node1->GetNodeTitle(ENodeTitleType::ListView);
+		FText Node2Name = Node2->GetNodeTitle(ENodeTitleType::ListView);
+
+		FFormatNamedArguments Args;
+		Args.Add(TEXT("PinNameForNode1"), FText::FromString(Pin1->PinName));
+		Args.Add(TEXT("NodeName1"), Node1Name);
+		Args.Add(TEXT("NodeName2"), Node2Name);
+		Diff.ToolTip =  FText::Format(LOCTEXT("DIF_PinLinkMovedToolTip", "Pin '{PinNameForNode1}' was linked to Node '{NodeName1}', but is now linked to Node '{NodeName2}'"), Args).ToString();
+
 		Diff.DisplayColor = FLinearColor(0.85f,0.71f,0.25f);
-		Diff.DisplayString =   FString::Printf(*LOCTEXT("DIF_PinLinkMoved", "Link Moved  '%s' ['%s' -> '%s']").ToString(), *Pin1->PinName,  *Node1Name, *Node2Name);
+		Diff.DisplayString = FText::Format(LOCTEXT("DIF_PinLinkMoved", "Link Moved  '{PinName}' ['{NodeName1}' -> '{NodeName2}']"), Args).ToString();
 		Results.Add(Diff);
 	}
 }
@@ -239,11 +256,14 @@ static void DiffR_PinDefaultValueChanged(FDiffResults& Results, class UEdGraphPi
 		Diff.Diff = EDiffType::PIN_DEFAULT_VALUE;
 		Diff.Pin1 = Pin1;
 		Diff.Pin2 = Pin2;
-		Diff.ToolTip =  FString::Printf(*LOCTEXT("DIF_PinDefaultValueToolTip", "Pin '%s' Default Value was '%s', but is now '%s").ToString(), *Pin2->PinName, 
-			*Pin1->GetDefaultAsString(), *Pin2->GetDefaultAsString());
+
+		FFormatNamedArguments Args;
+		Args.Add(TEXT("PinNameForValue1"), FText::FromString(Pin2->PinName));
+		Args.Add(TEXT("PinValue1"), FText::FromString(Pin1->GetDefaultAsString()));
+		Args.Add(TEXT("PinValue2"), FText::FromString(Pin2->GetDefaultAsString()));
+		Diff.ToolTip =  FText::Format(LOCTEXT("DIF_PinDefaultValueToolTip", "Pin '{PinNameForValue1}' Default Value was '{PinValue1}', but is now '{PinValue2}"), Args).ToString();
 		Diff.DisplayColor = FLinearColor(0.665f,0.13f,0.455f);
-		Diff.DisplayString =   FString::Printf(*LOCTEXT("DIF_PinDefaultValue", "Pin Default '%s' ['%s' -> '%s']").ToString(), *Pin2->PinName, *Pin1->GetDefaultAsString(), 
-			*Pin2->GetDefaultAsString());
+		Diff.DisplayString = FText::Format(LOCTEXT("DIF_PinDefaultValue", "Pin Default '{PinNameForValue1}' '{PinValue1}' -> '{PinValue2}']"), Args).ToString();
 		Results.Add(Diff);
 	}
 }
@@ -261,17 +281,22 @@ static void DiffR_NodePinCount(FDiffResults& Results, class UEdGraphNode* Node,
 {
 	if(Results)
 	{
-		FString NodeName = Node->GetNodeTitle(ENodeTitleType::ListView);
+		FText NodeName = Node->GetNodeTitle(ENodeTitleType::ListView);
 		int32 OriginalCount = Pins2.Num();
 		int32 NewCount = Pins1.Num();
 		FDiffSingleResult Diff;
 		Diff.Diff = EDiffType::NODE_PIN_COUNT;
 		Diff.Node1 = Node;
 		Diff.Node2 = Node2;
-		Diff.ToolTip =  FString::Printf(*LOCTEXT("DIF_NodePinCountChangedToolTip", "Node '%s' had %i Pins, now has %i Pins").ToString(), *NodeName, OriginalCount, NewCount);
+
+		FFormatNamedArguments Args;
+		Args.Add(TEXT("NodeName"), NodeName);
+		Args.Add(TEXT("OriginalCount"), OriginalCount);
+		Args.Add(TEXT("NewCount"), NewCount);
+		Diff.ToolTip =  FText::Format(LOCTEXT("DIF_NodePinCountChangedToolTip", "Node '{NodeName}' had {OriginalCount} Pins, now has {NewCount} Pins"), Args).ToString();
 		Diff.DisplayColor = FLinearColor(0.45f,0.4f,0.4f);
-		Diff.DisplayString =  (OriginalCount < NewCount) ?  (FString::Printf(*LOCTEXT("DIF_NodePinCountIncreased", "Add Pin '%s'").ToString(), *NodeName)) : 
-			(FString::Printf(*LOCTEXT("DIF_NodePinCountDecreased", "Removed Pin '%s'").ToString(), *NodeName));
+		Diff.DisplayString =  (OriginalCount < NewCount) ?  (FText::Format(LOCTEXT("DIF_NodePinCountIncreased", "Add Pin '{NodeName'"), Args).ToString()) : 
+			(FText::Format(LOCTEXT("DIF_NodePinCountDecreased", "Removed Pin '{NodeName}'"), Args).ToString());
 		Results.Add(Diff);
 	}
 }
@@ -545,7 +570,11 @@ bool FGraphDiffControl::DiffGraphs(UEdGraph* const LhsGraph, UEdGraph* const Rhs
 						// flip the "add" to a "removed" because all the diffs are
 						// listed in a context of "what has changed from lhs to rhs?"
 						Difference.Diff	         = EDiffType::NODE_REMOVED;
-						Difference.ToolTip       = FString::Printf(*LOCTEXT("DIF_RemoveNode", "Removed Node '%s'").ToString(), *Difference.Node1->GetNodeTitle(ENodeTitleType::ListView));
+
+						FFormatNamedArguments Args;
+						Args.Add(TEXT("NodeTitle"), Difference.Node1->GetNodeTitle(ENodeTitleType::ListView));
+						Difference.ToolTip       = FText::Format(LOCTEXT("DIF_RemoveNode", "Removed Node '{NodeTitle}'"), Args).ToString();
+
 						Difference.DisplayString = Difference.ToolTip;
 						Difference.DisplayColor  = FLinearColor(1.f,0.4f,0.4f);
 
diff --git a/Engine/Source/Editor/GraphEditor/Private/GraphEditorDragDropAction.cpp b/Engine/Source/Editor/GraphEditor/Private/GraphEditorDragDropAction.cpp
--- a/Engine/Source/Editor/GraphEditor/Private/GraphEditorDragDropAction.cpp
+++ b/Engine/Source/Editor/GraphEditor/Private/GraphEditorDragDropAction.cpp
@@ -110,7 +110,7 @@ void FGraphEditorDragDropAction::SetFeedbackMessage(const TSharedPtr<SWidget>& M
 	}
 }
 
-void FGraphEditorDragDropAction::SetSimpleFeedbackMessage(const FSlateBrush* Icon, const FSlateColor& IconColor, const FString& Message)
+void FGraphEditorDragDropAction::SetSimpleFeedbackMessage(const FSlateBrush* Icon, const FSlateColor& IconColor, const FText& Message)
 {
 	// Let the user know the status of making this connection.
 	SetFeedbackMessage(
diff --git a/Engine/Source/Editor/GraphEditor/Private/KismetNodes/SGraphNodeK2Base.cpp b/Engine/Source/Editor/GraphEditor/Private/KismetNodes/SGraphNodeK2Base.cpp
--- a/Engine/Source/Editor/GraphEditor/Private/KismetNodes/SGraphNodeK2Base.cpp
+++ b/Engine/Source/Editor/GraphEditor/Private/KismetNodes/SGraphNodeK2Base.cpp
@@ -153,11 +153,12 @@ TSharedPtr<SToolTip> SGraphNodeK2Base::GetComplexTooltip()
 	};
 	TSharedRef<SGraphNodeK2Base> ThisRef = SharedThis(this);
 
+	TSharedPtr< SVerticalBox > VerticalBoxWidget;
 	SAssignNew(NodeToolTip, SToolTip)
 		.Visibility_Static(&LocalUtils::IsToolTipVisible, ThisRef)
 		.IsInteractive_Static(&LocalUtils::IsInteractive)
 	[
-		SNew(SVerticalBox)
+		SAssignNew(VerticalBoxWidget, SVerticalBox)
 		// heading container
 		+SVerticalBox::Slot()
 		[
@@ -189,6 +190,28 @@ TSharedPtr<SToolTip> SGraphNodeK2Base::GetComplexTooltip()
 		]
 	];
 
+	// English speakers have no real need to know this exists.
+	if(FInternationalization::GetCurrentCulture()->GetTwoLetterISOLanguageName() != TEXT("en"))
+	{
+		struct Local
+		{
+			static EVisibility GetNativeNodeNameVisibility()
+			{
+				return FSlateApplication::Get().GetModifierKeys().IsAltDown()? EVisibility::Collapsed : EVisibility::Visible;
+			}
+		};
+
+		VerticalBoxWidget->AddSlot()
+			.AutoHeight()
+			.HAlign( HAlign_Right )
+			[
+				SNew( STextBlock )
+				.ColorAndOpacity( FSlateColor::UseSubduedForeground() )
+				.Text( LOCTEXT( "NativeNodeName", "hold (Alt) for native node name" ) )
+				.TextStyle( &FEditorStyle::GetWidgetStyle<FTextBlockStyle>(TEXT("Documentation.SDocumentationTooltip")) )
+				.Visibility_Static(&Local::GetNativeNodeNameVisibility)
+			];
+	}
 	return NodeToolTip;
 }
 
@@ -227,7 +250,7 @@ bool SGraphNodeK2Base::RequiresSecondPassLayout() const
 	return bBeadMode;
 }
 
-FString SGraphNodeK2Base::GetNodeCompactTitle() const
+FText SGraphNodeK2Base::GetNodeCompactTitle() const
 {
 	UK2Node* K2Node = CastChecked<UK2Node>(GraphNode);
 	return K2Node->GetCompactNodeTitle();
diff --git a/Engine/Source/Editor/GraphEditor/Private/KismetNodes/SGraphNodeK2Base.h b/Engine/Source/Editor/GraphEditor/Private/KismetNodes/SGraphNodeK2Base.h
--- a/Engine/Source/Editor/GraphEditor/Private/KismetNodes/SGraphNodeK2Base.h
+++ b/Engine/Source/Editor/GraphEditor/Private/KismetNodes/SGraphNodeK2Base.h
@@ -27,7 +27,7 @@ protected :
 	void UpdateCompactNode();
 
 	/** Get title in compact mode */
-	FString GetNodeCompactTitle() const;
+	FText GetNodeCompactTitle() const;
 
 	/** Retrieves text to tack on to the top of the tooltip (above the standard text) */
 	FText GetToolTipHeading() const;
diff --git a/Engine/Source/Editor/GraphEditor/Private/KismetNodes/SGraphNodeK2Var.cpp b/Engine/Source/Editor/GraphEditor/Private/KismetNodes/SGraphNodeK2Var.cpp
--- a/Engine/Source/Editor/GraphEditor/Private/KismetNodes/SGraphNodeK2Var.cpp
+++ b/Engine/Source/Editor/GraphEditor/Private/KismetNodes/SGraphNodeK2Var.cpp
@@ -31,7 +31,7 @@ void SGraphNodeK2Var::UpdateGraphNode()
 	RightNodeBox.Reset();
 	LeftNodeBox.Reset();
 
-	FString TitleText;
+	FText TitleText;
 	bool bPadTitle = false;
 	float HorizontalTitleMargin = 0.0f;
 	float VerticalTitleMargin = 8.0f;
@@ -44,27 +44,31 @@ void SGraphNodeK2Var::UpdateGraphNode()
 		UK2Node_VariableSet *SetNode = Cast<UK2Node_VariableSet>(GraphNode);
 		if(SetNode->HasLocalRepNotify())
 		{
-			TitleText = NSLOCTEXT("GraphEditor", "VariableSetWithNotify", "SET w/ Notify").ToString();
+			TitleText = NSLOCTEXT("GraphEditor", "VariableSetWithNotify", "SET w/ Notify");
 			//+ FString::Printf(TEXT("\n(%s)"), *SetNode->GetRepNotifyName().ToString());
 		}
 		else
 		{
-			TitleText = NSLOCTEXT("GraphEditor", "VariableSet", "SET").ToString();
+			TitleText = NSLOCTEXT("GraphEditor", "VariableSet", "SET");
 		}
 	}
 	else if (UK2Node_StructOperation* StructOp = Cast<UK2Node_StructOperation>(GraphNode))
 	{
+		FFormatNamedArguments Args;
+		Args.Add(TEXT("VariableName"), StructOp->GetVarNameText());
 		if (GraphNode->IsA(UK2Node_StructMemberGet::StaticClass()))
 		{
-			TitleText = FString::Printf(*NSLOCTEXT("GraphEditor", "StructMemberGet", "Get in %s").ToString(), *(StructOp->GetVarNameString()));
+			TitleText = FText::Format(NSLOCTEXT("GraphEditor", "StructMemberGet", "Get in {VariableName}"), Args);
 		}
 		else if (GraphNode->IsA(UK2Node_StructMemberSet::StaticClass()))
 		{
-			TitleText = FString::Printf(*NSLOCTEXT("GraphEditor", "StructMemberSet", "Set in %s").ToString(), *(StructOp->GetVarNameString()));
+			FFormatNamedArguments Args;
+			Args.Add(TEXT("VariableName"), StructOp->GetVarNameText());
+			TitleText = FText::Format(NSLOCTEXT("GraphEditor", "StructMemberSet", "Set in {VariableName}"), Args);
 		}
 		else if (GraphNode->IsA(UK2Node_MakeStruct::StaticClass()))
 		{
-			TitleText = FString::Printf(*NSLOCTEXT("GraphEditor", "MakeStruct", "Make %s").ToString(), *(StructOp->GetVarNameString()));
+			TitleText = FText::Format(NSLOCTEXT("GraphEditor", "MakeStruct", "Make {VariableName}"), Args);
 		}
 		else
 		{
@@ -76,25 +80,25 @@ void SGraphNodeK2Var::UpdateGraphNode()
 	}
 	else if (UK2Node_Literal* LiteralRef = Cast<UK2Node_Literal>(GraphNode))
 	{
-		FString SubTitleText;
+		FText SubTitleText;
 
 		// Get the name of the level the object is in.
 		if(AActor* Actor = Cast<AActor>(LiteralRef->GetObjectRef()))
 		{
-			FString LevelName;
+			FText LevelName;
 
 			if(ULevel* Level = Actor->GetLevel())
 			{
 				if ( Level->IsPersistentLevel() )
 				{
-					LevelName = NSLOCTEXT("GraphEditor", "PersistentTag", "Persistent Level").ToString();
+					LevelName = NSLOCTEXT("GraphEditor", "PersistentTag", "Persistent Level");
 				}
 				else
 				{
-					LevelName = FPaths::GetCleanFilename(Actor->GetOutermost()->GetName());
+					LevelName = FText::FromString(FPaths::GetCleanFilename(Actor->GetOutermost()->GetName()));
 				}
 			}
-			SubTitleText = FText::Format(NSLOCTEXT("GraphEditor", "ActorRef", "from {0}"), FText::FromString(LevelName)).ToString();
+			SubTitleText = FText::Format(NSLOCTEXT("GraphEditor", "ActorRef", "from {0}"), LevelName);
 		}
 
 		TitleText = GraphNode->GetNodeTitle(ENodeTitleType::ListView);
diff --git a/Engine/Source/Editor/GraphEditor/Private/SGraphActionMenu.cpp b/Engine/Source/Editor/GraphEditor/Private/SGraphActionMenu.cpp
--- a/Engine/Source/Editor/GraphEditor/Private/SGraphActionMenu.cpp
+++ b/Engine/Source/Editor/GraphEditor/Private/SGraphActionMenu.cpp
@@ -25,7 +25,7 @@ namespace GraphActionMenuHelpers
 		bCheck |= (InGraphAction->GetTypeId() == FEdGraphSchemaAction_K2Delegate::StaticGetTypeId() &&
 			((FEdGraphSchemaAction_K2Delegate*)InGraphAction)->GetDelegateName() == ItemName);
 		bCheck |= (InGraphAction->GetTypeId() == FEdGraphSchemaAction_K2TargetNode::StaticGetTypeId() &&
-			((FEdGraphSchemaAction_K2TargetNode*)InGraphAction)->NodeTemplate->GetNodeTitle(ENodeTitleType::EditableTitle) == ItemName.ToString());
+			((FEdGraphSchemaAction_K2TargetNode*)InGraphAction)->NodeTemplate->GetNodeTitle(ENodeTitleType::EditableTitle).ToString() == ItemName.ToString());
 
 		return bCheck;
 	}
@@ -244,7 +244,7 @@ void SGraphActionMenu::RefreshAllActions(bool bPreserveExpansion, bool bHandleOn
 		for (int32 i = 0; i < GraphNodes.Num(); ++i)
 		{
 			TSharedPtr<FEdGraphSchemaAction> GraphAction = GraphNodes[i]->Actions[0];
-			if (GraphAction.IsValid() && GraphAction->MenuDescription == SelectedAction->MenuDescription)
+			if (GraphAction.IsValid() && GraphAction->MenuDescription.ToString() == SelectedAction->MenuDescription.ToString())
 			{
 				// Clear the selection (if this node is already selected then setting it will have no effect)
 				TreeView->ClearSelection();
@@ -443,7 +443,7 @@ static bool CompareGraphActionNode(TSharedPtr<FGraphActionNode> A, TSharedPtr<FG
 
 	if(A->Actions[0].IsValid() && B->Actions[0].IsValid())
 	{
-		return(A->Actions[0]->MenuDescription == B->Actions[0]->MenuDescription);
+		return A->Actions[0]->MenuDescription.CompareTo(B->Actions[0]->MenuDescription) == 0;
 	}
 	else if(!A->Actions[0].IsValid() && !B->Actions[0].IsValid())
 	{
@@ -524,7 +524,7 @@ void SGraphActionMenu::GenerateFilteredItems(bool bPreserveExpansion)
 			// Combine the actions string, separate with \n so terms don't run into each other, and remove the spaces (incase the user is searching for a variable)
 			// In the case of groups containing multiple actions, they will have been created and added at the same place in the code, using the same description
 			// and keywords, so we only need to use the first one for filtering.
-			FString SearchText = CurrentAction.Actions[0]->MenuDescription + LINE_TERMINATOR + CurrentAction.Actions[0]->Keywords + LINE_TERMINATOR +CurrentAction.Actions[0]->Category;
+			FString SearchText = CurrentAction.Actions[0]->MenuDescription.ToString() + LINE_TERMINATOR + CurrentAction.Actions[0]->GetSearchTitle().ToString() + LINE_TERMINATOR + CurrentAction.Actions[0]->Keywords + LINE_TERMINATOR +CurrentAction.Actions[0]->Category;
 			SearchText = SearchText.Replace( TEXT( " " ), TEXT( "" ) );
 			// Get the 'weight' of this in relation to the filter
 			EachWeight = GetActionFilteredWeight( CurrentAction, FilterTerms, SanitizedFilterTerms );
@@ -613,14 +613,8 @@ int32 SGraphActionMenu::GetActionFilteredWeight( const FGraphActionListBuilderBa
 	int32 WholeMatchWeightMultiplier = 2;
 	int32 DescriptionWeight = 5;
 	int32 CategoryWeight = 3;
-	if(InCurrentAction.Actions[0]->MenuDescription == TEXT( "EventActorBeginCursorOver" ) )
-	{
-		int stop = 1;
-	}
-	if(InCurrentAction.Actions[0]->MenuDescription == TEXT( "Event Actor Begin Cursor Over" ) )
-	{
-		int stop = 1;
-	}
+	int32 NodeTitleWeight = 3;
+
 	// Helper array
 	struct FArrayWithWeight
 	{
@@ -638,7 +632,7 @@ int32 SGraphActionMenu::GetActionFilteredWeight( const FGraphActionListBuilderBa
 		// Combine the actions string, separate with \n so terms don't run into each other, and remove the spaces (incase the user is searching for a variable)
 		// In the case of groups containing multiple actions, they will have been created and added at the same place in the code, using the same description
 		// and keywords, so we only need to use the first one for filtering.
-		FString SearchText = InCurrentAction.Actions[Action]->MenuDescription + LINE_TERMINATOR + InCurrentAction.Actions[Action]->Keywords + LINE_TERMINATOR +InCurrentAction.Actions[Action]->Category;
+		FString SearchText = InCurrentAction.Actions[0]->MenuDescription.ToString() + LINE_TERMINATOR + InCurrentAction.Actions[Action]->GetSearchTitle().ToString() + LINE_TERMINATOR + InCurrentAction.Actions[Action]->Keywords + LINE_TERMINATOR +InCurrentAction.Actions[Action]->Category;
 		SearchText = SearchText.Replace( TEXT( " " ), TEXT( "" ) );
 
 		// First the keywords
@@ -647,10 +641,15 @@ int32 SGraphActionMenu::GetActionFilteredWeight( const FGraphActionListBuilderBa
 		WeightedArrayList.Add( EachEntry );
 
 		// The description
-		InCurrentAction.Actions[Action]->MenuDescription.ParseIntoArray( &EachEntry.Array, TEXT(" "), true );
+		InCurrentAction.Actions[Action]->MenuDescription.ToString().ParseIntoArray( &EachEntry.Array, TEXT(" "), true );
 		EachEntry.Weight = DescriptionWeight;
 		WeightedArrayList.Add( EachEntry );
 
+		// The node search title weight
+		InCurrentAction.Actions[Action]->GetSearchTitle().ToString().ParseIntoArray( &EachEntry.Array, TEXT(" "), true );
+		EachEntry.Weight = NodeTitleWeight;
+		WeightedArrayList.Add( EachEntry );
+
 		// The category
 		InCurrentAction.Actions[Action]->Category.ParseIntoArray( &EachEntry.Array, TEXT(" "), true );
 		EachEntry.Weight = CategoryWeight;
diff --git a/Engine/Source/Editor/GraphEditor/Private/SGraphNode.cpp b/Engine/Source/Editor/GraphEditor/Private/SGraphNode.cpp
--- a/Engine/Source/Editor/GraphEditor/Private/SGraphNode.cpp
+++ b/Engine/Source/Editor/GraphEditor/Private/SGraphNode.cpp
@@ -16,7 +16,6 @@
 void SNodeTitle::Construct(const FArguments& InArgs, UEdGraphNode* InNode)
 {
 	GraphNode = InNode;
-	CachedString = FString();
 
 	ExtraLineStyle = InArgs._ExtraLineStyle;
 
@@ -27,9 +26,9 @@ void SNodeTitle::Construct(const FArguments& InArgs, UEdGraphNode* InNode)
 	}
 	else
 	{
-		TitleText = TAttribute<FString>(this, &SNodeTitle::GetNodeTitle);
+		TitleText = TAttribute<FText>(this, &SNodeTitle::GetNodeTitle);
 	}
-	CachedString = TitleText.Get();
+	CachedTitle = TitleText.Get();
 	RebuildWidget();
 }
 
@@ -38,24 +37,23 @@ void SNodeTitle::Tick(const FGeometry& AllottedGeometry, const double InCurrentT
 	SCompoundWidget::Tick(AllottedGeometry, InCurrentTime, InDeltaTime);
 
 	// Checks to see if the cached string is valid, and if not, updates it.
-	FString CurrentNodeTitle = TitleText.Get();
-	if (CurrentNodeTitle != CachedString)
+	if (TitleText.Get().CompareTo(CachedTitle) == 0)
 	{
-		CachedString = CurrentNodeTitle;
+		CachedTitle = TitleText.Get();
 		RebuildWidget();
 	}
 }
 
-FString SNodeTitle::GetNodeTitle() const
+FText SNodeTitle::GetNodeTitle() const
 {
 	return (GraphNode != NULL)
 		? GraphNode->GetNodeTitle(ENodeTitleType::FullTitle)
-		: NSLOCTEXT("GraphEditor", "NullNode", "Null Node").ToString();
+		: NSLOCTEXT("GraphEditor", "NullNode", "Null Node");
 }
 
 FText SNodeTitle::GetHeadTitle() const
 {
-	return GraphNode->bCanRenameNode? FText::FromString(GraphNode->GetNodeTitle(ENodeTitleType::EditableTitle)) : CachedHeadTitle;
+	return GraphNode->bCanRenameNode? GraphNode->GetNodeTitle(ENodeTitleType::EditableTitle) : CachedHeadTitle;
 }
 
 void SNodeTitle::RebuildWidget()
@@ -69,7 +67,7 @@ void SNodeTitle::RebuildWidget()
 
 	// Break the title into lines
 	TArray<FString> Lines;
-	CachedString.ParseIntoArray(&Lines, TEXT("\n"), false);
+	CachedTitle.ToString().ParseIntoArray(&Lines, TEXT("\n"), false);
 
 	if(Lines.Num())
 	{
@@ -429,13 +427,13 @@ FString SGraphNode::GetEditableNodeTitle() const
 
 	if(GraphNode)
 	{
-		return GraphNode->GetNodeTitle(ENodeTitleType::EditableTitle);
+		return GraphNode->GetNodeTitle(ENodeTitleType::EditableTitle).ToString();
 	}
 	return NSLOCTEXT("GraphEditor", "NullNode", "Null Node").ToString();
 
 	// Get the portion of the node that is actually editable text (may be a subsection of the title, or something else entirely)
 	return (GraphNode != NULL)
-		? GraphNode->GetNodeTitle(ENodeTitleType::EditableTitle)
+		? GraphNode->GetNodeTitle(ENodeTitleType::EditableTitle).ToString()
 		: NSLOCTEXT("GraphEditor", "NullNode", "Null Node").ToString();
 }
 
@@ -472,6 +470,12 @@ FText SGraphNode::GetNodeTooltip() const
 {
 	if (GraphNode != NULL)
 	{
+		// Display the native title of the node when alt is held
+		if(FSlateApplication::Get().GetModifierKeys().IsAltDown())
+		{
+			return FText::FromString(GraphNode->GetNodeNativeTitle(ENodeTitleType::ListView));
+		}
+
 		FText TooltipText = FText::FromString(GraphNode->GetTooltip());
 
 		if (UEdGraph* Graph = GraphNode->GetGraph())
@@ -479,13 +483,16 @@ FText SGraphNode::GetNodeTooltip() const
 			// If the node resides in an intermediate graph, show the UObject name for debug purposes
 			if (Graph->HasAnyFlags(RF_Transient))
 			{
-				TooltipText = FText::FromString(FString::Printf(TEXT("%s\n\n%s"), *GraphNode->GetName(), *TooltipText.ToString()));
+				FFormatNamedArguments Args;
+				Args.Add(TEXT("NodeName"), FText::FromString(GraphNode->GetName()));
+				Args.Add(TEXT("TooltipText"), TooltipText);
+				TooltipText = FText::Format(NSLOCTEXT("GraphEditor", "GraphNodeTooltip", "{NodeName}\n\n{TooltipText}"), Args);
 			}
 		}
 
 		if (TooltipText.IsEmpty())
 		{
-			TooltipText =  FText::FromString(GraphNode->GetNodeTitle(ENodeTitleType::FullTitle));
+			TooltipText =  GraphNode->GetNodeTitle(ENodeTitleType::FullTitle);
 		}
 
 		return TooltipText;
diff --git a/Engine/Source/Editor/GraphEditor/Private/SGraphPalette.cpp b/Engine/Source/Editor/GraphEditor/Private/SGraphPalette.cpp
--- a/Engine/Source/Editor/GraphEditor/Private/SGraphPalette.cpp
+++ b/Engine/Source/Editor/GraphEditor/Private/SGraphPalette.cpp
@@ -156,12 +156,12 @@ void SGraphPaletteItem::OnNameTextCommitted(const FText& NewText, ETextCommit::T
 
 FText SGraphPaletteItem::GetDisplayText() const
 {
-	return FText::FromString(ActionPtr.Pin()->MenuDescription);
+	return ActionPtr.Pin()->MenuDescription;
 }
 
 FString SGraphPaletteItem::GetItemTooltip() const
 {
-	return ActionPtr.Pin()->MenuDescription;
+	return ActionPtr.Pin()->MenuDescription.ToString();
 }
 
 
diff --git a/Engine/Source/Editor/GraphEditor/Public/GraphEditorDragDropAction.h b/Engine/Source/Editor/GraphEditor/Public/GraphEditorDragDropAction.h
--- a/Engine/Source/Editor/GraphEditor/Public/GraphEditorDragDropAction.h
+++ b/Engine/Source/Editor/GraphEditor/Public/GraphEditorDragDropAction.h
@@ -26,7 +26,7 @@ public:
 	
 protected:
 	void SetFeedbackMessage(const TSharedPtr<SWidget>& Message);
-	void SetSimpleFeedbackMessage(const FSlateBrush* Icon, const FSlateColor& IconColor, const FString& Message);
+	void SetSimpleFeedbackMessage(const FSlateBrush* Icon, const FSlateColor& IconColor, const FText& Message);
 
 	UEdGraphPin* GetHoveredPin() const;
 	UEdGraphNode* GetHoveredNode() const;
diff --git a/Engine/Source/Editor/GraphEditor/Public/SGraphNode.h b/Engine/Source/Editor/GraphEditor/Public/SGraphNode.h
--- a/Engine/Source/Editor/GraphEditor/Public/SGraphNode.h
+++ b/Engine/Source/Editor/GraphEditor/Public/SGraphNode.h
@@ -27,7 +27,7 @@ public:
 		SLATE_ARGUMENT(FName, ExtraLineStyle)
 
 		// Title text to display, auto-binds to get the title if not set externally
-		SLATE_ATTRIBUTE(FString, Text)
+		SLATE_ATTRIBUTE(FText, Text)
 	SLATE_END_ARGS()
 
 	void Construct(const FArguments& InArgs, UEdGraphNode* InNode);
@@ -41,18 +41,18 @@ public:
 
 protected:
 	UEdGraphNode* GraphNode;
-	FString CachedString;
+	FText CachedTitle;
 	FName ExtraLineStyle;
 
 	/** The cached head title to return */
 	FText CachedHeadTitle;
 
 	/** The title text to use, auto-binds to get the title if not set externally */
-	TAttribute< FString > TitleText;
+	TAttribute< FText > TitleText;
 
 protected:
 	// Gets the expected node title
-	FString GetNodeTitle() const;
+	FText GetNodeTitle() const;
 
 	// Rebuilds the widget if needed
 	void RebuildWidget();
