#TTP 316058 - BP: On Vector math nodes, the symbols are mostly obscured
#Branch UE4
#Proj GraphEditor, EditorStyle, Slate

#Summary Modified the vector text entry boxes to match other math nodes numerical entry styles and modified the K2 base graph node to vertically centre the text box denoting its function.

#Change Modified SGraphNodeK2Base::UpdateCompactNode to vertically centre the overlay with the text denoting the noes function ( in this case the operator )
#Change Modified SVectorTextBox::Construct to change the style to match other math operation nodes ( white border, white text and no background )
#Change Added VectorEditableTextBox to slate editor styles.
#Change Modified SNumericEntryBox to expose the background and foreground colors as slate arguments so they can be customised. Default behaviour should be as before.

diff --git a/Engine/Source/Editor/EditorStyle/Private/SlateEditorStyle.cpp b/Engine/Source/Editor/EditorStyle/Private/SlateEditorStyle.cpp
--- a/Engine/Source/Editor/EditorStyle/Private/SlateEditorStyle.cpp
+++ b/Engine/Source/Editor/EditorStyle/Private/SlateEditorStyle.cpp
@@ -3100,6 +3100,18 @@ void FSlateEditorStyle::FStyle::SetupGraphEditorStyles()
 					);
 			}
 
+			// VectorEditableTextBox
+			{
+				Set( "Graph.VectorEditableTextBox", FEditableTextBoxStyle()
+					.SetBackgroundImageNormal( BOX_BRUSH( "Graph/CommonWidgets/TextBox", FMargin(4.0f/16.0f) ) )
+					.SetBackgroundImageHovered( BOX_BRUSH( "Graph/CommonWidgets/TextBox_Hovered", FMargin(4.0f/16.0f) ) )
+					.SetBackgroundImageFocused( BOX_BRUSH( "Graph/CommonWidgets/TextBox_Hovered", FMargin(4.0f/16.0f) ) )
+					.SetBackgroundImageReadOnly( BOX_BRUSH( "Graph/CommonWidgets/TextBox", FMargin(4.0f/16.0f) ) )
+					.SetForegroundColor( FLinearColor::White )
+					.SetBackgroundColor( FLinearColor::Blue )
+					);
+			}
+
 			// Check Box
 			{
 				/* Set images for various SCheckBox states of style Graph.Checkbox ... */
diff --git a/Engine/Source/Editor/GraphEditor/Private/KismetNodes/SGraphNodeK2Base.cpp b/Engine/Source/Editor/GraphEditor/Private/KismetNodes/SGraphNodeK2Base.cpp
--- a/Engine/Source/Editor/GraphEditor/Private/KismetNodes/SGraphNodeK2Base.cpp
+++ b/Engine/Source/Editor/GraphEditor/Private/KismetNodes/SGraphNodeK2Base.cpp
@@ -76,9 +76,7 @@ void SGraphNodeK2Base::UpdateCompactNode()
 			]
 			+SOverlay::Slot()
 			.HAlign(HAlign_Center)
-			// Pad out so that the node circles do not
-			// overlap the background text.
-			.Padding( FMargin(23,0) )
+			.VAlign(VAlign_Center)
 			[
 				// MIDDLE
 				SNew(STextBlock)
@@ -90,7 +88,7 @@ void SGraphNodeK2Base::UpdateCompactNode()
 			[
 				SNew(SHorizontalBox)
 				+SHorizontalBox::Slot()
-				.Padding( FMargin(0,0,5,0) )
+				.Padding(0)
 				.HAlign(HAlign_Left)
 				.VAlign(VAlign_Center)
 				.FillWidth(1.0f)
@@ -100,7 +98,7 @@ void SGraphNodeK2Base::UpdateCompactNode()
 				]
 				+SHorizontalBox::Slot()
 				.AutoWidth()
-				.Padding( FMargin(5,0,0,0) )
+				.Padding(0)
 				.HAlign(HAlign_Right)
 				.VAlign(VAlign_Center)
 				[
diff --git a/Engine/Source/Editor/GraphEditor/Private/KismetPins/SGraphPinVector.cpp b/Engine/Source/Editor/GraphEditor/Private/KismetPins/SGraphPinVector.cpp
--- a/Engine/Source/Editor/GraphEditor/Private/KismetPins/SGraphPinVector.cpp
+++ b/Engine/Source/Editor/GraphEditor/Private/KismetPins/SGraphPinVector.cpp
@@ -28,12 +28,14 @@ public:
 		VisibleText_0 = InArgs._VisibleText_0;
 		VisibleText_1 = InArgs._VisibleText_1;
 		VisibleText_2 = InArgs._VisibleText_2;
+		const FLinearColor LabelClr = FLinearColor( 1.f, 1.f, 1.f, 0.4f );
 
 		this->ChildSlot
 		[
 			SNew(SVerticalBox)
 			+SVerticalBox::Slot()
-			.AutoHeight() .Padding(5)
+			.AutoHeight() 
+			.Padding(0)
 			[
 				SNew(SHorizontalBox)
 				+SHorizontalBox::Slot()
@@ -45,14 +47,18 @@ public:
 					.Label()
 					[
 						SNew( STextBlock )
-						.Font( FEditorStyle::GetFontStyle( "Graph.EditableTextBox" ) )
+						.Font( FEditorStyle::GetFontStyle( "Graph.VectorEditableTextBox" ) )
 						.Text( LOCTEXT("VectorNodeXAxisValueLabel", "X") )
+						.ColorAndOpacity( LabelClr )
 					]
 					.Value( this, &SVectorTextBox::GetTypeInValue_0 )
 					.OnValueCommitted( InArgs._OnFloatCommitted_Box_0 )
-					.Font( FEditorStyle::GetFontStyle( "Graph.EditableTextBox" ) )
+					.Font( FEditorStyle::GetFontStyle( "Graph.VectorEditableTextBox" ) )
 					.UndeterminedString( LOCTEXT("MultipleValues", "Multiple Values") )
 					.ToolTipText(LOCTEXT("VectorNodeXAxisValueLabel_ToolTip", "X value"))
+					.EditableTextBoxStyle( &FEditorStyle::GetWidgetStyle<FEditableTextBoxStyle>( "Graph.VectorEditableTextBox" ))
+					.BorderForegroundColor( FLinearColor::White )
+					.BorderBackgroundColor( FLinearColor::White )
 				]
 				+SHorizontalBox::Slot()
 				.AutoWidth().Padding(2) .HAlign(HAlign_Fill)
@@ -63,14 +69,18 @@ public:
 					.Label()
 					[
 						SNew( STextBlock )
-						.Font( FEditorStyle::GetFontStyle( "Graph.EditableTextBox" ) )
+						.Font( FEditorStyle::GetFontStyle( "Graph.VectorEditableTextBox" ) )
 						.Text( LOCTEXT("VectorNodeYAxisValueLabel", "Y") )
+						.ColorAndOpacity( LabelClr )
 					]
 					.Value( this, &SVectorTextBox::GetTypeInValue_1 )
 					.OnValueCommitted( InArgs._OnFloatCommitted_Box_1 )
-					.Font( FEditorStyle::GetFontStyle( "Graph.EditableTextBox" ) )
+					.Font( FEditorStyle::GetFontStyle( "Graph.VectorEditableTextBox" ) )
 					.UndeterminedString( LOCTEXT("MultipleValues", "Multiple Values") )
 					.ToolTipText(LOCTEXT("VectorNodeYAxisValueLabel_ToolTip", "Y value"))
+					.EditableTextBoxStyle( &FEditorStyle::GetWidgetStyle<FEditableTextBoxStyle>( "Graph.VectorEditableTextBox" ))
+					.BorderForegroundColor( FLinearColor::White )
+					.BorderBackgroundColor( FLinearColor::White )
 				]
 				+SHorizontalBox::Slot()
 				.AutoWidth().Padding(2) .HAlign(HAlign_Fill)
@@ -81,14 +91,18 @@ public:
 					.Label()
 					[
 						SNew( STextBlock )
-						.Font( FEditorStyle::GetFontStyle( "Graph.EditableTextBox" ) )
+						.Font( FEditorStyle::GetFontStyle( "Graph.VectorEditableTextBox" ) )
 						.Text( LOCTEXT("VectorNodeZAxisValueLabel", "Z") )
+						.ColorAndOpacity( LabelClr )
 					]
 					.Value( this, &SVectorTextBox::GetTypeInValue_2 )
 					.OnValueCommitted( InArgs._OnFloatCommitted_Box_2 )
-					.Font( FEditorStyle::GetFontStyle( "Graph.EditableTextBox" ) )
+					.Font( FEditorStyle::GetFontStyle( "Graph.VectorEditableTextBox" ) )
 					.UndeterminedString( LOCTEXT("MultipleValues", "Multiple Values") )
 					.ToolTipText(LOCTEXT("VectorNodeZAxisValueLabel_ToolTip", "Z value"))
+					.EditableTextBoxStyle( &FEditorStyle::GetWidgetStyle<FEditableTextBoxStyle>( "Graph.VectorEditableTextBox" ))
+					.BorderForegroundColor( FLinearColor::White )
+					.BorderBackgroundColor( FLinearColor::White )
 				]
 			]
 		];
diff --git a/Engine/Source/Runtime/Slate/Public/SNumericEntryBox.h b/Engine/Source/Runtime/Slate/Public/SNumericEntryBox.h
--- a/Engine/Source/Runtime/Slate/Public/SNumericEntryBox.h
+++ b/Engine/Source/Runtime/Slate/Public/SNumericEntryBox.h
@@ -36,6 +36,8 @@ public:
 		, _MinSliderValue(0)
 		, _MaxSliderValue(100)
 		, _SliderExponent(1.f)
+		, _BorderForegroundColor(FCoreStyle::Get().GetSlateColor("InvertedForeground"))
+		, _BorderBackgroundColor(FLinearColor::White)
 	{}		
 
 		/** Style to use for the editable text box within this widget */
@@ -47,6 +49,10 @@ public:
 		SLATE_ARGUMENT( EVerticalAlignment, LabelVAlign )
 		/** Padding around the label content */
 		SLATE_ARGUMENT( FMargin, LabelPadding )
+		/** Border Foreground Color */
+		SLATE_ARGUMENT( FSlateColor, BorderForegroundColor )
+		/** Border Background Color */
+		SLATE_ARGUMENT( FSlateColor, BorderBackgroundColor )
 		/** The value that should be displayed.  This value is optional in the case where a value cannot be determined */
 		SLATE_ATTRIBUTE( TOptional<NumericType>, Value )
 		/** The string to display if the value cannot be determined */
@@ -169,8 +175,8 @@ public:
 			[
 				SNew( SBorder )
 				.BorderImage( this, &SNumericEntryBox<NumericType>::GetBorderImage )
-				.BorderBackgroundColor( FLinearColor::White )
-				.ForegroundColor( FCoreStyle::Get().GetSlateColor("InvertedForeground") )
+				.BorderBackgroundColor( InArgs._BorderBackgroundColor )
+				.ForegroundColor( InArgs._BorderForegroundColor )
 				.Padding(0)
 				[
 					HorizontalBox
