#summary Fixed node titles not refreshing when they the text changes.

#change Old and new titles were being checked for being the same to allow for the update in text, instead of being different.

#ttp 332815 - Critical: Regression: EDITOR: MATERIAL: Values on Constant nodes not updating

diff --git a/Engine/Source/Editor/GraphEditor/Private/SGraphNode.cpp b/Engine/Source/Editor/GraphEditor/Private/SGraphNode.cpp
--- a/Engine/Source/Editor/GraphEditor/Private/SGraphNode.cpp
+++ b/Engine/Source/Editor/GraphEditor/Private/SGraphNode.cpp
@@ -37,7 +37,7 @@ void SNodeTitle::Tick(const FGeometry& AllottedGeometry, const double InCurrentT
 	SCompoundWidget::Tick(AllottedGeometry, InCurrentTime, InDeltaTime);
 
 	// Checks to see if the cached string is valid, and if not, updates it.
-	if (TitleText.Get().CompareTo(CachedTitle) == 0)
+	if (TitleText.Get().CompareTo(CachedTitle) != 0)
 	{
 		CachedTitle = TitleText.Get();
 		RebuildWidget();
