#ue4
- Disabled LogConnectionDrawingPolicy log spam

diff --git a/Engine/Source/Editor/GraphEditor/Private/ConnectionDrawingPolicy.cpp b/Engine/Source/Editor/GraphEditor/Private/ConnectionDrawingPolicy.cpp
--- a/Engine/Source/Editor/GraphEditor/Private/ConnectionDrawingPolicy.cpp
+++ b/Engine/Source/Editor/GraphEditor/Private/ConnectionDrawingPolicy.cpp
@@ -730,7 +730,7 @@ void FKismetConnectionDrawingPolicy::BuildExecutionRoadmap()
 			}
 			else 
 			{
-				UE_LOG(LogConnectionDrawingPolicy, Warning, TEXT("Looks like a wire-trace was not injected before the jump from '%s' to '%s'."), 
+				UE_LOG(LogConnectionDrawingPolicy, Verbose, TEXT("Looks like a wire-trace was not injected before the jump from '%s' to '%s'."), 
 					*CurNode->GetNodeTitle(ENodeTitleType::FullTitle).ToString(), *NextNode->GetNodeTitle(ENodeTitleType::FullTitle).ToString());
 			}
 
