TTP #333069/#332291: Removing newly added transactions that were holding onto too many assets and presenting undo actions on editor startup (fallout from CL 2038794).

diff --git a/Engine/Source/Editor/BlueprintGraph/Private/K2Node.cpp b/Engine/Source/Editor/BlueprintGraph/Private/K2Node.cpp
--- a/Engine/Source/Editor/BlueprintGraph/Private/K2Node.cpp
+++ b/Engine/Source/Editor/BlueprintGraph/Private/K2Node.cpp
@@ -3,7 +3,6 @@
 
 #include "BlueprintGraphPrivatePCH.h"
 #include "KismetDebugUtilities.h" // for HasDebuggingData(), GetWatchText()
-#include "ScopedTransaction.h"
 
 #define LOCTEXT_NAMESPACE "K2Node"
 
@@ -325,12 +324,6 @@ void UK2Node::RemovePinsFromOldPins(TArray<UEdGraphPin*>& OldPins, int32 Removed
 
 void UK2Node::ReconstructNode()
 {
-	FFormatNamedArguments Args;
-	Args.Add(TEXT("NodeTitle"), GetNodeTitle(ENodeTitleType::FullTitle));
-
-	FText TransactionDesc = FText::Format(LOCTEXT("ReconstructNodeTransaction", "'{NodeTitle}' Node Refreshed"), Args);
-	FScopedTransaction const Transaction(TransactionDesc);
-
 	Modify();
 
 	UBlueprint* Blueprint = GetBlueprint();
diff --git a/Engine/Source/Editor/GraphEditor/Private/SGraphEditorImpl.cpp b/Engine/Source/Editor/GraphEditor/Private/SGraphEditorImpl.cpp
--- a/Engine/Source/Editor/GraphEditor/Private/SGraphEditorImpl.cpp
+++ b/Engine/Source/Editor/GraphEditor/Private/SGraphEditorImpl.cpp
@@ -437,12 +437,15 @@ bool SGraphEditorImpl::CanBreakPinLinks() const
 void SGraphEditorImpl::ReconstructNodes()
 {
 	const UEdGraphSchema* Schema = this->EdGraphObj->GetSchema();
-
-	for (FGraphPanelSelectionSet::TConstIterator NodeIt( GraphPanel->SelectionManager.GetSelectedNodes() ); NodeIt; ++NodeIt)
 	{
-		if (UEdGraphNode* Node = Cast<UEdGraphNode>(*NodeIt))
+		FScopedTransaction const Transaction(LOCTEXT("ReconstructNodeTransaction", "Refresh Node(s)"));
+
+		for (FGraphPanelSelectionSet::TConstIterator NodeIt( GraphPanel->SelectionManager.GetSelectedNodes() ); NodeIt; ++NodeIt)
 		{
-			Schema->ReconstructNode(*Node);
+			if (UEdGraphNode* Node = Cast<UEdGraphNode>(*NodeIt))
+			{
+				Schema->ReconstructNode(*Node);
+			}
 		}
 	}
 	NotifyGraphChanged();
diff --git a/Engine/Source/Editor/UnrealEd/Private/Kismet2/BlueprintEditorUtils.cpp b/Engine/Source/Editor/UnrealEd/Private/Kismet2/BlueprintEditorUtils.cpp
--- a/Engine/Source/Editor/UnrealEd/Private/Kismet2/BlueprintEditorUtils.cpp
+++ b/Engine/Source/Editor/UnrealEd/Private/Kismet2/BlueprintEditorUtils.cpp
@@ -246,9 +246,6 @@ void FBlueprintEditorUtils::RefreshAllNodes(UBlueprint* Blueprint)
 		AllNodes.Sort(FCompareNodePriority());
 	}
 
-	// scoped transaction (so users don't have to undo every single ReconstructNode)
-	FScopedTransaction const Transaction(LOCTEXT("ReconstructNodes", "Refresh All Nodes"));
-
 	bool bLastChangesStructure = (AllNodes.Num() > 0) ? AllNodes[0]->NodeCausesStructuralBlueprintChange() : true;
 	for( TArray<UK2Node*>::TIterator NodeIt(AllNodes); NodeIt; ++NodeIt )
 	{
