#TTP 316058 - BP: On Vector math nodes, the symbols are mostly obscured
#Branch UE4
#Proj GraphEditor

#Summary The 2D Vector node style changes ( for readability ) were missed from the original changelist 2044998 , just adding them with this changelist.

diff --git a/Engine/Source/Editor/GraphEditor/Private/KismetPins/SGraphPinVector2D.cpp b/Engine/Source/Editor/GraphEditor/Private/KismetPins/SGraphPinVector2D.cpp
--- a/Engine/Source/Editor/GraphEditor/Private/KismetPins/SGraphPinVector2D.cpp
+++ b/Engine/Source/Editor/GraphEditor/Private/KismetPins/SGraphPinVector2D.cpp
@@ -23,12 +23,14 @@ public:
 	{
 		VisibleText_X = InArgs._VisibleText_X;
 		VisibleText_Y = InArgs._VisibleText_Y;
+		const FLinearColor LabelClr = FLinearColor( 1.f, 1.f, 1.f, 0.4f );
 
 		this->ChildSlot
 		[
 			SNew(SVerticalBox)
 			+SVerticalBox::Slot()
-			.AutoHeight() .Padding(5)
+			.AutoHeight()
+			.Padding(0)
 			[
 				SNew(SHorizontalBox)
 				+SHorizontalBox::Slot()
@@ -40,14 +42,18 @@ public:
 					.Label()
 					[
 						SNew( STextBlock )
-						.Font( FEditorStyle::GetFontStyle( "Graph.EditableTextBox" ) )
+						.Font( FEditorStyle::GetFontStyle( "Graph.VectorEditableTextBox" ) )
 						.Text( LOCTEXT("VectorNodeXAxisValueLabel", "X") )
+						.ColorAndOpacity( LabelClr )
 					]
 					.Value( this, &SVector2DTextBox::GetTypeInValue_X )
 					.OnValueCommitted( InArgs._OnFloatCommitted_Box_X )
-					.Font( FEditorStyle::GetFontStyle( "Graph.EditableTextBox" ) )
+					.Font( FEditorStyle::GetFontStyle( "Graph.VectorEditableTextBox" ) )
 					.UndeterminedString( LOCTEXT("MultipleValues", "Multiple Values") )
-					.ToolTipText(LOCTEXT("VectorNodeXAxisValueLabel_ToolTip", "X value"))
+					.ToolTipText(LOCTEXT("VectorNodeXAxisValueLabel_ToolTip", "X value") )
+					.EditableTextBoxStyle( &FEditorStyle::GetWidgetStyle<FEditableTextBoxStyle>( "Graph.VectorEditableTextBox" ))
+					.BorderForegroundColor( FLinearColor::White )
+					.BorderBackgroundColor( FLinearColor::White )
 				]
 				+SHorizontalBox::Slot()
 				.AutoWidth().Padding(2) .HAlign(HAlign_Fill)
@@ -58,14 +64,18 @@ public:
 					.Label()
 					[
 						SNew( STextBlock )
-						.Font( FEditorStyle::GetFontStyle( "Graph.EditableTextBox" ) )
+						.Font( FEditorStyle::GetFontStyle( "Graph.VectorEditableTextBox" ) )
 						.Text( LOCTEXT("VectorNodeYAxisValueLabel", "Y") )
+						.ColorAndOpacity( LabelClr )
 					]
 					.Value( this, &SVector2DTextBox::GetTypeInValue_Y )
 					.OnValueCommitted( InArgs._OnFloatCommitted_Box_Y )
-					.Font( FEditorStyle::GetFontStyle( "Graph.EditableTextBox" ) )
+					.Font( FEditorStyle::GetFontStyle( "Graph.VectorEditableTextBox" ) )
 					.UndeterminedString( LOCTEXT("MultipleValues", "Multiple Values") )
 					.ToolTipText(LOCTEXT("VectorNodeYAxisValueLabel_ToolTip", "Y value"))
+					.EditableTextBoxStyle( &FEditorStyle::GetWidgetStyle<FEditableTextBoxStyle>( "Graph.VectorEditableTextBox" ))
+					.BorderForegroundColor( FLinearColor::White )
+					.BorderBackgroundColor( FLinearColor::White )
 				]
 			]
 		];
