Only warning about uninjected wire-traces when the two nodes are connected (accounts for sequence nodes, where the end of one chain jumps back in execution).
(TTP #333145)

diff --git a/Engine/Source/Editor/GraphEditor/Private/ConnectionDrawingPolicy.cpp b/Engine/Source/Editor/GraphEditor/Private/ConnectionDrawingPolicy.cpp
--- a/Engine/Source/Editor/GraphEditor/Private/ConnectionDrawingPolicy.cpp
+++ b/Engine/Source/Editor/GraphEditor/Private/ConnectionDrawingPolicy.cpp
@@ -728,9 +728,14 @@ void FKismetConnectionDrawingPolicy::BuildExecutionRoadmap()
 					ExecTiming.PredExecTime = CurNodeTime;
 				}
 			}
-			else 
+			// if the nodes aren't graphically connected how could they be 
+			// executed back-to-back? well, this could be a pop back to a 
+			// sequence node from the end of one thread of execution, etc.
+			else if (AreNodesGraphicallySequential(CurNode, NextNode))
 			{
-				UE_LOG(LogConnectionDrawingPolicy, Verbose, TEXT("Looks like a wire-trace was not injected before the jump from '%s' to '%s'."), 
+				// only warn when the nodes are directly connected (this is all
+				// for execution flow visualization after all)
+				UE_LOG(LogConnectionDrawingPolicy, Warning, TEXT("Looks like a wire-trace was not injected before the jump from '%s' to '%s'."), 
 					*CurNode->GetNodeTitle(ENodeTitleType::FullTitle).ToString(), *NextNode->GetNodeTitle(ENodeTitleType::FullTitle).ToString());
 			}
 
@@ -773,6 +778,32 @@ bool FKismetConnectionDrawingPolicy::TreatWireAsExecutionPin(UEdGraphPin* InputP
 	return (InputPin != NULL) && (Schema->IsExecPin(*OutputPin));
 }
 
+bool FKismetConnectionDrawingPolicy::AreNodesGraphicallySequential(UEdGraphNode* InputNode, UEdGraphNode* OutputNode) const
+{
+	for (UEdGraphPin* Pin : InputNode->Pins)
+	{
+		if (Pin->Direction != EGPD_Output)
+		{
+			continue;
+		}
+
+		for (UEdGraphPin* Connection : Pin->LinkedTo)
+		{
+			if (!TreatWireAsExecutionPin(Pin, Connection))
+			{
+				continue;
+			}
+
+			if (Connection->GetOwningNode() == OutputNode)
+			{
+				return true;
+			}
+		}
+	}
+
+	return false;
+}
+
 void FKismetConnectionDrawingPolicy::DetermineStyleOfExecWire(float& Thickness, FLinearColor& WireColor, bool& bDrawBubbles, const FTimePair& Times)
 {
 	// It's a followed link, make it strong and yellowish but fading over time
diff --git a/Engine/Source/Editor/GraphEditor/Private/ConnectionDrawingPolicy.h b/Engine/Source/Editor/GraphEditor/Private/ConnectionDrawingPolicy.h
--- a/Engine/Source/Editor/GraphEditor/Private/ConnectionDrawingPolicy.h
+++ b/Engine/Source/Editor/GraphEditor/Private/ConnectionDrawingPolicy.h
@@ -153,6 +153,17 @@ public:
 
 	virtual void SetIncompatiblePinDrawState(const TSharedPtr<SGraphPin>& StartPin, const TSet< TSharedRef<SWidget> >& VisiblePins);
 	virtual void ResetIncompatiblePinDrawState(const TSet< TSharedRef<SWidget> >& VisiblePins);
+
+private: 
+	/**
+	* Checks that the specified nodes are connected in the graph directly by an 
+	* exec pin.
+	* 
+	* @param  InputNode	The node whose exec output(s) should be connected to OutputNode.
+	* @param  OutputNode	The node whose exec input should be coming from the InputNode.
+	* @return True if the two nodes are connected in the graph, otherwise false.
+	*/
+	bool AreNodesGraphicallySequential(UEdGraphNode* InputNode, UEdGraphNode* OutputNode) const;
 };
 
 /////////////////////////////////////////////////////
