Non unity fixes

diff --git a/Engine/Source/Editor/AnimGraph/Private/AnimGraphNode_Root.cpp b/Engine/Source/Editor/AnimGraph/Private/AnimGraphNode_Root.cpp
--- a/Engine/Source/Editor/AnimGraph/Private/AnimGraphNode_Root.cpp
+++ b/Engine/Source/Editor/AnimGraph/Private/AnimGraphNode_Root.cpp
@@ -1,6 +1,7 @@
 // Copyright 1998-2014 Epic Games, Inc. All Rights Reserved.
 
 #include "AnimGraphPrivatePCH.h"
+#include "AnimGraphNode_Root.h"
 
 /////////////////////////////////////////////////////
 // FPoseLinkMappingRecord
diff --git a/Engine/Source/Editor/GraphEditor/Private/AnimationNodes/SGraphNodeSequencePlayer.cpp b/Engine/Source/Editor/GraphEditor/Private/AnimationNodes/SGraphNodeSequencePlayer.cpp
--- a/Engine/Source/Editor/GraphEditor/Private/AnimationNodes/SGraphNodeSequencePlayer.cpp
+++ b/Engine/Source/Editor/GraphEditor/Private/AnimationNodes/SGraphNodeSequencePlayer.cpp
@@ -6,6 +6,7 @@
 #include "SGraphPreviewer.h"
 #include "Editor/UnrealEd/Public/Kismet2/KismetDebugUtilities.h"
 #include "Editor/UnrealEd/Public/Kismet2/BlueprintEditorUtils.h"
+#include "AnimGraphNode_SequencePlayer.h"
 
 /////////////////////////////////////////////////////
 // SGraphNodeSequencePlayer
diff --git a/Engine/Source/Editor/GraphEditor/Private/AnimationStateNodes/SGraphNodeAnimTransition.cpp b/Engine/Source/Editor/GraphEditor/Private/AnimationStateNodes/SGraphNodeAnimTransition.cpp
--- a/Engine/Source/Editor/GraphEditor/Private/AnimationStateNodes/SGraphNodeAnimTransition.cpp
+++ b/Engine/Source/Editor/GraphEditor/Private/AnimationStateNodes/SGraphNodeAnimTransition.cpp
@@ -14,6 +14,9 @@
 
 #include "AnimationTransitionGraph.h"
 #include "AnimGraphNode_TransitionResult.h"
+#include "AnimGraphNode_StateMachineBase.h"
+#include "AnimStateTransitionNode.h"
+#include "AnimationStateMachineGraph.h"
 
 #define LOCTEXT_NAMESPACE "TransitionNodes"
 
