Fixed preview connectors sometimes disappearing when CTRL+Dragging pins in graphs

This fixes ttp 331811 - The control+left drag click functionality in the material editor broke recently.
I've changed the preview connector pins to be stored as a handle (NodeGuid + PinName) so that we can reference them persistently across rebuilds of the graph

diff --git a/Engine/Source/Editor/GraphEditor/Private/SGraphPanel.cpp b/Engine/Source/Editor/GraphEditor/Private/SGraphPanel.cpp
--- a/Engine/Source/Editor/GraphEditor/Private/SGraphPanel.cpp
+++ b/Engine/Source/Editor/GraphEditor/Private/SGraphPanel.cpp
@@ -20,6 +20,47 @@
 
 DEFINE_LOG_CATEGORY_STATIC(LogGraphPanel, Log, All);
 
+
+/** A handle to a pin, defined by its owning node's GUID, and the pin's name. Used to reference a pin without referring to its widget */
+struct FGraphPinHandle
+{
+	/** The GUID of the node to which this pin belongs */
+	FGuid NodeGuid;
+
+	/** The name of the pin we are referencing */
+	FString PinName;
+
+	/** Constructor */
+	FGraphPinHandle(UEdGraphPin* InPin)
+	{
+		if (auto* Node = InPin->GetOwningNode())
+		{
+			PinName = InPin->PinName;
+			NodeGuid = Node->NodeGuid;
+		}
+	}
+
+	/** Find a pin widget in the specified panel from this handle */
+	TSharedPtr<SGraphPin> FindInGraphPanel(const SGraphPanel& InPanel) const
+	{
+		// First off, find the node
+		TSharedPtr<SGraphNode> GraphNode = InPanel.GetNodeWidgetFromGuid(NodeGuid);
+		if (GraphNode.IsValid())
+		{
+			UEdGraphNode* Node = GraphNode->GetNodeObj();
+			UEdGraphPin* Pin = Node->FindPin(PinName);
+
+			if (Pin)
+			{
+				return GraphNode->FindWidgetForPin(Pin);
+			}
+		}
+
+		return TSharedPtr<SGraphPin>();
+	}
+};
+
+
 /**
  * Construct a widget
  *
@@ -241,7 +282,7 @@ int32 SGraphPanel::OnPaint( const FGeometry& AllottedGeometry, const FSlateRect&
 					const bool bNodeIsDifferent = (!GraphObjToDiff || NodeMatch.Diff());
 
 					/* When dragging off a pin, we want to duck the alpha of some nodes */
-					TSharedPtr< SGraphPin > OnlyStartPin = (1 == PreviewConnectorFromPins.Num()) ? PreviewConnectorFromPins[0] : TSharedPtr< SGraphPin >();
+					TSharedPtr< SGraphPin > OnlyStartPin = (1 == PreviewConnectorFromPins.Num()) ? PreviewConnectorFromPins[0].FindInGraphPanel(*this) : TSharedPtr< SGraphPin >();
 					const bool bNodeIsNotUsableInCurrentContext = Schema->FadeNodeWhenDraggingOffPin(NodeObj, OnlyStartPin.IsValid() ? OnlyStartPin.Get()->GetPinObj() : NULL);
 					const FWidgetStyle& NodeStyleToUse = (bNodeIsDifferent && !bNodeIsNotUsableInCurrentContext)? InWidgetStyle : FadedStyle;
 
@@ -326,7 +367,17 @@ int32 SGraphPanel::OnPaint( const FGeometry& AllottedGeometry, const FSlateRect&
 				}
 			}
 		}
-		ConnectionDrawingPolicy->SetHoveredPins(CurrentHoveredPins, PreviewConnectorFromPins, TimeSinceMouseEnteredPin);
+
+		TArray<TSharedPtr<SGraphPin>> OverridePins;
+		for (const FGraphPinHandle& Handle : PreviewConnectorFromPins)
+		{
+			TSharedPtr<SGraphPin> Pin = Handle.FindInGraphPanel(*this);
+			if (Pin.IsValid())
+			{
+				OverridePins.Add(Pin);
+			}
+		}
+		ConnectionDrawingPolicy->SetHoveredPins(CurrentHoveredPins, OverridePins, TimeSinceMouseEnteredPin);
 		ConnectionDrawingPolicy->SetMarkedPin(MarkedPin);
 
 		// Get the set of pins for all children and synthesize geometry for culled out pins so lines can be drawn to them.
@@ -373,9 +424,13 @@ int32 SGraphPanel::OnPaint( const FGeometry& AllottedGeometry, const FSlateRect&
 		// Draw preview connections (only connected on one end)
 		if (PreviewConnectorFromPins.Num() > 0)
 		{
-			for (TArray< TSharedPtr<SGraphPin> >::TConstIterator StartPinIterator(PreviewConnectorFromPins); StartPinIterator; ++StartPinIterator)
+			for (const FGraphPinHandle& Handle : PreviewConnectorFromPins)
 			{
-				TSharedPtr< SGraphPin > CurrentStartPin = *StartPinIterator;
+				TSharedPtr< SGraphPin > CurrentStartPin = Handle.FindInGraphPanel(*this);
+				if (!CurrentStartPin.IsValid())
+				{
+					continue;
+				}
 				const FArrangedWidget* PinGeometry = PinGeometries.Find( CurrentStartPin.ToSharedRef() );
 
 				if (PinGeometry != NULL)
@@ -590,8 +645,8 @@ TSharedPtr<SWidget> SGraphPanel::OnSummonContextMenu(const FGeometry& MyGeometry
 
 bool SGraphPanel::OnHandleLeftMouseRelease(const FGeometry& MyGeometry, const FPointerEvent& MouseEvent)
 {
-	const bool bIsMakingConnection = (PreviewConnectorFromPins.Num() > 0) && PreviewConnectorFromPins[0].IsValid() && IsEditable.Get();
-	if (bIsMakingConnection)
+	TSharedPtr<SGraphPin> PreviewConnectionPin = PreviewConnectorFromPins.Num() > 0 ? PreviewConnectorFromPins[0].FindInGraphPanel(*this) : nullptr;
+	if (PreviewConnectionPin.IsValid() && IsEditable.Get())
 	{
 		TSet< TSharedRef<SWidget> > AllConnectors;
 		for (int32 ChildIndex = 0; ChildIndex < Children.Num(); ++ChildIndex)
@@ -613,10 +668,10 @@ bool SGraphPanel::OnHandleLeftMouseRelease(const FGeometry& MyGeometry, const FP
 			{
 				SGraphPin& TargetPin = static_cast<SGraphPin&>( PinWidgetGeometry.Widget.Get() );
 
-				if (PreviewConnectorFromPins[0]->TryHandlePinConnection(TargetPin))
+				if (PreviewConnectionPin->TryHandlePinConnection(TargetPin))
 				{
 					NodeList.Add(TargetPin.GetPinObj()->GetOwningNode());
-					NodeList.Add(PreviewConnectorFromPins[0]->GetPinObj()->GetOwningNode());
+					NodeList.Add(PreviewConnectionPin->GetPinObj()->GetOwningNode());
 				}
 				bHandledDrop = true;
 			}
@@ -745,7 +800,10 @@ FReply SGraphPanel::OnDrop( const FGeometry& MyGeometry, const FDragDropEvent& D
 
 void SGraphPanel::OnBeginMakingConnection( const TSharedRef<SGraphPin>& InOriginatingPin )
 {
-	PreviewConnectorFromPins.Add(InOriginatingPin);
+	if (auto* PinObj = InOriginatingPin->GetPinObj())
+	{
+		PreviewConnectorFromPins.Add(PinObj);
+	}
 }
 
 void SGraphPanel::OnStopMakingConnection(bool bForceStop)
@@ -769,6 +827,11 @@ void SGraphPanel::AddGraphNode( const TSharedRef<SNodePanel::SNode>& NodeToAdd )
 	GraphNode->SetOwner( SharedThis(this) );
 
 	const UEdGraphNode* Node = GraphNode->GetNodeObj();
+	if (Node)
+	{
+		NodeGuidMap.Add(Node->NodeGuid, GraphNode);
+	}
+
 	if (Node && Node->IsA( UEdGraphNode_Comment::StaticClass()))
 	{
 		SNodePanel::AddGraphNodeToBack(NodeToAdd);
@@ -781,6 +844,7 @@ void SGraphPanel::AddGraphNode( const TSharedRef<SNodePanel::SNode>& NodeToAdd )
 
 void SGraphPanel::RemoveAllNodes()
 {
+	NodeGuidMap.Empty();
 	CurrentHoveredPins.Empty();
 	SNodePanel::RemoveAllNodes();
 }
@@ -861,6 +925,11 @@ void SGraphPanel::AddNode (UEdGraphNode* Node)
 	}
 }
 
+TSharedPtr<SGraphNode> SGraphPanel::GetNodeWidgetFromGuid(FGuid Guid) const
+{
+	return NodeGuidMap.FindRef(Guid).Pin();
+}
+
 void SGraphPanel::Update()
 {
 	// Add widgets for all the nodes that don't have one.
@@ -998,4 +1067,4 @@ void SGraphPanel::AddReferencedObjects( FReferenceCollector& Collector )
 {
 	Collector.AddReferencedObject( GraphObj );
 	Collector.AddReferencedObject( GraphObjToDiff );
-}
+}
\ No newline at end of file
diff --git a/Engine/Source/Editor/GraphEditor/Private/SGraphPanel.h b/Engine/Source/Editor/GraphEditor/Private/SGraphPanel.h
--- a/Engine/Source/Editor/GraphEditor/Private/SGraphPanel.h
+++ b/Engine/Source/Editor/GraphEditor/Private/SGraphPanel.h
@@ -4,7 +4,7 @@
 
 #include "GraphEditor.h"
 
-class SGraphPanel;
+struct FGraphPinHandle;
 class SGraphNode;
 class SGraphPin;
 class UEdGraphNode;
@@ -135,6 +135,14 @@ public:
 	/** Pin marked via shift-clicking */
 	TWeakPtr<SGraphPin> MarkedPin;
 
+	/** Get a graph node widget from the specified GUID, if it applies to any nodes in this graph */
+	TSharedPtr<SGraphNode> GetNodeWidgetFromGuid(FGuid Guid) const;
+
+private:
+
+	/** A map of guid -> graph nodes */
+	TMap<FGuid, TWeakPtr<SGraphNode>> NodeGuidMap;
+
 protected:
 	UEdGraph* GraphObj;
 	UEdGraph* GraphObjToDiff;//if it exists, this is 
@@ -156,7 +164,7 @@ protected:
 	double TimeSinceMouseLeftPin;
 
 	/** Sometimes the panel draws a preview connector; e.g. when the user is connecting pins */
-	TArray< TSharedPtr<SGraphPin> > PreviewConnectorFromPins;
+	TArray< FGraphPinHandle > PreviewConnectorFromPins;
 	FVector2D PreviewConnectorEndpoint;
 
 	/** Invoked when we need to summon a context menu */
