Unity compilation fix: FGraphPinHandle

diff --git a/Engine/Source/Editor/GraphEditor/Private/SGraphPanel.cpp b/Engine/Source/Editor/GraphEditor/Private/SGraphPanel.cpp
--- a/Engine/Source/Editor/GraphEditor/Private/SGraphPanel.cpp
+++ b/Engine/Source/Editor/GraphEditor/Private/SGraphPanel.cpp
@@ -21,45 +21,32 @@
 DEFINE_LOG_CATEGORY_STATIC(LogGraphPanel, Log, All);
 
 
-/** A handle to a pin, defined by its owning node's GUID, and the pin's name. Used to reference a pin without referring to its widget */
-struct FGraphPinHandle
+SGraphPanel::FGraphPinHandle::FGraphPinHandle(UEdGraphPin* InPin)
 {
-	/** The GUID of the node to which this pin belongs */
-	FGuid NodeGuid;
-
-	/** The name of the pin we are referencing */
-	FString PinName;
-
-	/** Constructor */
-	FGraphPinHandle(UEdGraphPin* InPin)
+	if (auto* Node = InPin->GetOwningNode())
 	{
-		if (auto* Node = InPin->GetOwningNode())
-		{
-			PinName = InPin->PinName;
-			NodeGuid = Node->NodeGuid;
-		}
+		PinName = InPin->PinName;
+		NodeGuid = Node->NodeGuid;
 	}
+}
 
-	/** Find a pin widget in the specified panel from this handle */
-	TSharedPtr<SGraphPin> FindInGraphPanel(const SGraphPanel& InPanel) const
+TSharedPtr<SGraphPin> SGraphPanel::FGraphPinHandle::FindInGraphPanel(const SGraphPanel& InPanel) const
+{
+	// First off, find the node
+	TSharedPtr<SGraphNode> GraphNode = InPanel.GetNodeWidgetFromGuid(NodeGuid);
+	if (GraphNode.IsValid())
 	{
-		// First off, find the node
-		TSharedPtr<SGraphNode> GraphNode = InPanel.GetNodeWidgetFromGuid(NodeGuid);
-		if (GraphNode.IsValid())
-		{
-			UEdGraphNode* Node = GraphNode->GetNodeObj();
-			UEdGraphPin* Pin = Node->FindPin(PinName);
+		UEdGraphNode* Node = GraphNode->GetNodeObj();
+		UEdGraphPin* Pin = Node->FindPin(PinName);
 
-			if (Pin)
-			{
-				return GraphNode->FindWidgetForPin(Pin);
-			}
+		if (Pin)
+		{
+			return GraphNode->FindWidgetForPin(Pin);
 		}
-
-		return TSharedPtr<SGraphPin>();
 	}
-};
 
+	return TSharedPtr<SGraphPin>();
+}
 
 /**
  * Construct a widget
diff --git a/Engine/Source/Editor/GraphEditor/Private/SGraphPanel.h b/Engine/Source/Editor/GraphEditor/Private/SGraphPanel.h
--- a/Engine/Source/Editor/GraphEditor/Private/SGraphPanel.h
+++ b/Engine/Source/Editor/GraphEditor/Private/SGraphPanel.h
@@ -4,7 +4,6 @@
 
 #include "GraphEditor.h"
 
-struct FGraphPinHandle;
 class SGraphNode;
 class SGraphPin;
 class UEdGraphNode;
@@ -50,6 +49,22 @@ public:
 		//SLATE_ATTRIBUTE( FGraphAppearanceInfo, Appearance )
 	SLATE_END_ARGS()
 
+	/** A handle to a pin, defined by its owning node's GUID, and the pin's name. Used to reference a pin without referring to its widget */
+	struct FGraphPinHandle
+	{
+		/** The GUID of the node to which this pin belongs */
+		FGuid NodeGuid;
+
+		/** The name of the pin we are referencing */
+		FString PinName;
+
+		/** Constructor */
+		FGraphPinHandle(UEdGraphPin* InPin);
+
+		/** Find a pin widget in the specified panel from this handle */
+		TSharedPtr<SGraphPin> FindInGraphPanel(const SGraphPanel& InPanel) const;
+	};
+
 	/**
 	 * Construct a widget
 	 *
