Remove old reference to a NoRedist folder

diff --git a/Engine/Source/Editor/GraphEditor/GraphEditor.Build.cs b/Engine/Source/Editor/GraphEditor/GraphEditor.Build.cs
--- a/Engine/Source/Editor/GraphEditor/GraphEditor.Build.cs
+++ b/Engine/Source/Editor/GraphEditor/GraphEditor.Build.cs
@@ -15,7 +15,6 @@ public class GraphEditor : ModuleRules
 				"Editor/GraphEditor/Private/KismetPins",
 				"Editor/GraphEditor/Private/MaterialNodes",
 				"Editor/GraphEditor/Private/MaterialPins",
-				"Editor/GraphEditor/Private/NoRedist",
 				"Editor/GraphEditor/Private/SoundNodes",
 			}
 		);
