Added a movement filter to the node movement call, which now adds any successfully moved nodes. This prevents duplicate movement if two selected comments try to move the same node.
#TTP 314475 - k2: comment boxes within other comment boxes do not move if the outer most box is moved. Also, if two nested comment boxes are moved, the nodes within move 2x as much as they should (video)
#Branch UE4
#Proj GraphEditor, BehaviourTreeEditor

#Add SNodePanel::SNode::FNodeSet typedef for conveniance and clarity
#Change Changed SNodePanel::MoveTo to take an FNodeSet as an exclusion filter on the moved nodes, once moved nodes are appended into this set.

diff --git a/Engine/Source/Editor/BehaviorTreeEditor/Private/EdGraphSchema_BehaviorTree.cpp b/Engine/Source/Editor/BehaviorTreeEditor/Private/EdGraphSchema_BehaviorTree.cpp
--- a/Engine/Source/Editor/BehaviorTreeEditor/Private/EdGraphSchema_BehaviorTree.cpp
+++ b/Engine/Source/Editor/BehaviorTreeEditor/Private/EdGraphSchema_BehaviorTree.cpp
@@ -72,13 +72,14 @@ void AutoArrangeNodes(UBehaviorTreeGraphNode* ParentNode, FBBNode& BBoxTree, flo
 		if ( ParentNode->Pins[i]->Direction == EGPD_Output)
 		{
 			UEdGraphPin* Pin =  ParentNode->Pins[i];
+			SGraphNode::FNodeSet NodeFilter;
 			for (int32 Index=0; Index < Pin->LinkedTo.Num(); ++Index)
 			{
 				UBehaviorTreeGraphNode* GraphNode = Cast<UBehaviorTreeGraphNode>(Pin->LinkedTo[Index]->GetOwningNode());
 				if (GraphNode != NULL && BBoxTree.Children.Num() > 0)
 				{
 					AutoArrangeNodes(GraphNode, *BBoxTree.Children[BBoxIndex], PosX, PosY + GraphNode->NodeWidget.Pin()->GetDesiredSize().Y * 2.5f);
-					GraphNode->NodeWidget.Pin()->MoveTo(FVector2D(BBoxTree.Children[BBoxIndex]->SubGraphBBox.X /2 - GraphNode->NodeWidget.Pin()->GetDesiredSize().X /2 + PosX, PosY));
+					GraphNode->NodeWidget.Pin()->MoveTo(FVector2D(BBoxTree.Children[BBoxIndex]->SubGraphBBox.X /2 - GraphNode->NodeWidget.Pin()->GetDesiredSize().X /2 + PosX, PosY), NodeFilter);
 					PosX += BBoxTree.Children[BBoxIndex]->SubGraphBBox.X + 20;
 				}
 				BBoxIndex++;
diff --git a/Engine/Source/Editor/GraphEditor/Private/AnimationStateNodes/SGraphNodeAnimTransition.cpp b/Engine/Source/Editor/GraphEditor/Private/AnimationStateNodes/SGraphNodeAnimTransition.cpp
--- a/Engine/Source/Editor/GraphEditor/Private/AnimationStateNodes/SGraphNodeAnimTransition.cpp
+++ b/Engine/Source/Editor/GraphEditor/Private/AnimationStateNodes/SGraphNodeAnimTransition.cpp
@@ -33,7 +33,7 @@ void SGraphNodeAnimTransition::GetNodeInfoPopups(FNodeInfoContext* Context, TArr
 {
 }
 
-void SGraphNodeAnimTransition::MoveTo(const FVector2D& NewPosition)
+void SGraphNodeAnimTransition::MoveTo(const FVector2D& NewPosition, FNodeSet& NodeFilter)
 {
 	// Ignored; position is set by the location of the attached state nodes
 }
diff --git a/Engine/Source/Editor/GraphEditor/Private/AnimationStateNodes/SGraphNodeAnimTransition.h b/Engine/Source/Editor/GraphEditor/Private/AnimationStateNodes/SGraphNodeAnimTransition.h
--- a/Engine/Source/Editor/GraphEditor/Private/AnimationStateNodes/SGraphNodeAnimTransition.h
+++ b/Engine/Source/Editor/GraphEditor/Private/AnimationStateNodes/SGraphNodeAnimTransition.h
@@ -12,7 +12,7 @@ public:
 
 	// SNodePanel::SNode interface
 	virtual void GetNodeInfoPopups(FNodeInfoContext* Context, TArray<FGraphInformationPopupInfo>& Popups) const OVERRIDE;
-	virtual void MoveTo(const FVector2D& NewPosition) OVERRIDE;
+	virtual void MoveTo(const FVector2D& NewPosition, FNodeSet& NodeFilter) OVERRIDE;
 	virtual bool RequiresSecondPassLayout() const OVERRIDE;
 	virtual void PerformSecondPassLayout(const TMap< UObject*, TSharedRef<SNode> >& NodeToWidgetLookup) const OVERRIDE;
 	// End of SNodePanel::SNode interface
diff --git a/Engine/Source/Editor/GraphEditor/Private/MaterialNodes/SGraphNodeMaterialBase.cpp b/Engine/Source/Editor/GraphEditor/Private/MaterialNodes/SGraphNodeMaterialBase.cpp
--- a/Engine/Source/Editor/GraphEditor/Private/MaterialNodes/SGraphNodeMaterialBase.cpp
+++ b/Engine/Source/Editor/GraphEditor/Private/MaterialNodes/SGraphNodeMaterialBase.cpp
@@ -243,9 +243,9 @@ void SGraphNodeMaterialBase::CreatePinWidgets()
 	}
 }
 
-void SGraphNodeMaterialBase::MoveTo(const FVector2D& NewPosition)
+void SGraphNodeMaterialBase::MoveTo(const FVector2D& NewPosition, FNodeSet& NodeFilter)
 {
-	SGraphNode::MoveTo(NewPosition);
+	SGraphNode::MoveTo(NewPosition, NodeFilter);
 
 	MaterialNode->MaterialExpression->MaterialExpressionEditorX = MaterialNode->NodePosX;
 	MaterialNode->MaterialExpression->MaterialExpressionEditorY = MaterialNode->NodePosY;
diff --git a/Engine/Source/Editor/GraphEditor/Private/MaterialNodes/SGraphNodeMaterialBase.h b/Engine/Source/Editor/GraphEditor/Private/MaterialNodes/SGraphNodeMaterialBase.h
--- a/Engine/Source/Editor/GraphEditor/Private/MaterialNodes/SGraphNodeMaterialBase.h
+++ b/Engine/Source/Editor/GraphEditor/Private/MaterialNodes/SGraphNodeMaterialBase.h
@@ -68,7 +68,7 @@ public:
 	// End of SGraphNode interface
 
 	// SNodePanel::SNode interface
-	virtual void MoveTo(const FVector2D& NewPosition) OVERRIDE;
+	virtual void MoveTo(const FVector2D& NewPosition, FNodeSet& NodeFilter) OVERRIDE;
 	// End of SNodePanel::SNode interface
 
 	UMaterialGraphNode* GetMaterialGraphNode() const {return MaterialNode;}
diff --git a/Engine/Source/Editor/GraphEditor/Private/MaterialNodes/SGraphNodeMaterialComment.cpp b/Engine/Source/Editor/GraphEditor/Private/MaterialNodes/SGraphNodeMaterialComment.cpp
--- a/Engine/Source/Editor/GraphEditor/Private/MaterialNodes/SGraphNodeMaterialComment.cpp
+++ b/Engine/Source/Editor/GraphEditor/Private/MaterialNodes/SGraphNodeMaterialComment.cpp
@@ -11,12 +11,15 @@ void SGraphNodeMaterialComment::Construct(const FArguments& InArgs, class UMater
 	this->CommentNode = InNode;
 }
 
-void SGraphNodeMaterialComment::MoveTo(const FVector2D& NewPosition)
+void SGraphNodeMaterialComment::MoveTo(const FVector2D& NewPosition, FNodeSet& NodeFilter)
 {
-	SGraphNodeComment::MoveTo(NewPosition);
+	if( !NodeFilter.Find( SharedThis( this ) ))
+	{
+		SGraphNodeComment::MoveTo(NewPosition, NodeFilter);
 
-	CommentNode->MaterialExpressionComment->MaterialExpressionEditorX = CommentNode->NodePosX;
-	CommentNode->MaterialExpressionComment->MaterialExpressionEditorY = CommentNode->NodePosY;
-	CommentNode->MaterialExpressionComment->MarkPackageDirty();
-	CommentNode->MaterialDirtyDelegate.ExecuteIfBound();
-}
+		CommentNode->MaterialExpressionComment->MaterialExpressionEditorX = CommentNode->NodePosX;
+		CommentNode->MaterialExpressionComment->MaterialExpressionEditorY = CommentNode->NodePosY;
+		CommentNode->MaterialExpressionComment->MarkPackageDirty();
+		CommentNode->MaterialDirtyDelegate.ExecuteIfBound();
+	}
+}
\ No newline at end of file
diff --git a/Engine/Source/Editor/GraphEditor/Private/MaterialNodes/SGraphNodeMaterialComment.h b/Engine/Source/Editor/GraphEditor/Private/MaterialNodes/SGraphNodeMaterialComment.h
--- a/Engine/Source/Editor/GraphEditor/Private/MaterialNodes/SGraphNodeMaterialComment.h
+++ b/Engine/Source/Editor/GraphEditor/Private/MaterialNodes/SGraphNodeMaterialComment.h
@@ -11,7 +11,7 @@ public:
 	void Construct(const FArguments& InArgs, class UMaterialGraphNode_Comment* InNode);
 
 	// SNodePanel::SNode interface
-	virtual void MoveTo(const FVector2D& NewPosition) OVERRIDE;
+	virtual void MoveTo(const FVector2D& NewPosition, FNodeSet& NodeFilter) OVERRIDE;
 	// End of SNodePanel::SNode interface
 
 private:
diff --git a/Engine/Source/Editor/GraphEditor/Private/MaterialNodes/SGraphNodeMaterialResult.cpp b/Engine/Source/Editor/GraphEditor/Private/MaterialNodes/SGraphNodeMaterialResult.cpp
--- a/Engine/Source/Editor/GraphEditor/Private/MaterialNodes/SGraphNodeMaterialResult.cpp
+++ b/Engine/Source/Editor/GraphEditor/Private/MaterialNodes/SGraphNodeMaterialResult.cpp
@@ -48,9 +48,9 @@ void SGraphNodeMaterialResult::CreatePinWidgets()
 	}
 }
 
-void SGraphNodeMaterialResult::MoveTo(const FVector2D& NewPosition)
+void SGraphNodeMaterialResult::MoveTo(const FVector2D& NewPosition, FNodeSet& NodeFilter)
 {
-	SGraphNode::MoveTo(NewPosition);
+	SGraphNode::MoveTo(NewPosition, NodeFilter);
 
 	RootNode->Material->EditorX = RootNode->NodePosX;
 	RootNode->Material->EditorY = RootNode->NodePosY;
diff --git a/Engine/Source/Editor/GraphEditor/Private/MaterialNodes/SGraphNodeMaterialResult.h b/Engine/Source/Editor/GraphEditor/Private/MaterialNodes/SGraphNodeMaterialResult.h
--- a/Engine/Source/Editor/GraphEditor/Private/MaterialNodes/SGraphNodeMaterialResult.h
+++ b/Engine/Source/Editor/GraphEditor/Private/MaterialNodes/SGraphNodeMaterialResult.h
@@ -15,7 +15,7 @@ public:
 	// End of SGraphNode interface
 
 	// SNodePanel::SNode interface
-	virtual void MoveTo(const FVector2D& NewPosition) OVERRIDE;
+	virtual void MoveTo(const FVector2D& NewPosition, FNodeSet& NodeFilter) OVERRIDE;
 	// End of SNodePanel::SNode interface
 
 private:
diff --git a/Engine/Source/Editor/GraphEditor/Private/SGraphNode.cpp b/Engine/Source/Editor/GraphEditor/Private/SGraphNode.cpp
--- a/Engine/Source/Editor/GraphEditor/Private/SGraphNode.cpp
+++ b/Engine/Source/Editor/GraphEditor/Private/SGraphNode.cpp
@@ -397,14 +397,14 @@ void SGraphNode::SetOwner( const TSharedRef<SGraphPanel>& OwnerPanel )
 }
 
 /** @param NewPosition  The Node should be relocated to this position in the graph panel */
-void SGraphNode::MoveTo( const FVector2D& NewPosition )
+void SGraphNode::MoveTo( const FVector2D& NewPosition, FNodeSet& NodeFilter )
 {
-	if (GraphNode)
+	if ( !NodeFilter.Find( SharedThis( this )))
 	{
-		if (!RequiresSecondPassLayout())
+		if (GraphNode && !RequiresSecondPassLayout())
 		{
+			NodeFilter.Add( SharedThis( this ) );
 			GraphNode->Modify();
-
 			GraphNode->NodePosX = NewPosition.X;
 			GraphNode->NodePosY = NewPosition.Y;
 		}
diff --git a/Engine/Source/Editor/GraphEditor/Private/SGraphNodeComment.cpp b/Engine/Source/Editor/GraphEditor/Private/SGraphNodeComment.cpp
--- a/Engine/Source/Editor/GraphEditor/Private/SGraphNodeComment.cpp
+++ b/Engine/Source/Editor/GraphEditor/Private/SGraphNodeComment.cpp
@@ -163,8 +163,8 @@ FReply SGraphNodeComment::OnMouseMove(const FGeometry& MyGeometry, const FPointe
 			UserSize.X = SCommentNodeDefs::MinWidth;
 			DeltaNodePos.X = GetCorrectedNodePosition().X - GetPosition().X;
 		}
-
-		SGraphNode::MoveTo( GetPosition() + DeltaNodePos );
+		SGraphNode::FNodeSet NodeFilter;
+		SGraphNode::MoveTo( GetPosition() + DeltaNodePos, NodeFilter );
 	}
 	else
 	{
@@ -549,12 +549,10 @@ bool SGraphNodeComment::OnHitTest( const FGeometry& MyGeometry, FVector2D InAbso
 	return (InMouseZone != CWZ_NotInWindow) && (InMouseZone != CWZ_InWindow);
 }
 
-void SGraphNodeComment::MoveTo( const FVector2D& NewPosition ) 
+void SGraphNodeComment::MoveTo( const FVector2D& NewPosition, FNodeSet& NodeFilter ) 
 {
 	FVector2D PositionDelta = NewPosition - GetPosition();
-	SGraphNode::MoveTo(NewPosition);
-
-
+	SGraphNode::MoveTo(NewPosition, NodeFilter);
 	// Don't drag note content if either of the shift keys are down.
 	FModifierKeysState KeysState = FSlateApplication::Get().GetModifierKeys();
 	if(!KeysState.IsShiftDown())
@@ -570,8 +568,9 @@ void SGraphNodeComment::MoveTo( const FVector2D& NewPosition )
 			{
 				if (UEdGraphNode* Node = Cast<UEdGraphNode>(*NodeIt))
 				{
-					if ( !Panel->SelectionManager.IsNodeSelected(Node))
+					if ( !Panel->SelectionManager.IsNodeSelected(Node) && !NodeFilter.Find( Node->NodeWidget.Pin() ))
 					{
+						NodeFilter.Add( Node->NodeWidget.Pin() );
 						Node->NodePosX += PositionDelta.X;
 						Node->NodePosY += PositionDelta.Y;
 					}
diff --git a/Engine/Source/Editor/GraphEditor/Private/SGraphPanel.cpp b/Engine/Source/Editor/GraphEditor/Private/SGraphPanel.cpp
--- a/Engine/Source/Editor/GraphEditor/Private/SGraphPanel.cpp
+++ b/Engine/Source/Editor/GraphEditor/Private/SGraphPanel.cpp
@@ -506,7 +506,8 @@ void SGraphPanel::UpdateSelectedNodesPositions (FVector2D PositionIncrement)
 		if (pWidget != NULL)
 		{
 			SNode& Widget = pWidget->Get();
-			Widget.MoveTo(Widget.GetPosition() + PositionIncrement);
+			SNode::FNodeSet NodeFilter;
+			Widget.MoveTo(Widget.GetPosition() + PositionIncrement, NodeFilter);
 		}
 	}
 }
diff --git a/Engine/Source/Editor/GraphEditor/Private/SNodePanel.cpp b/Engine/Source/Editor/GraphEditor/Private/SNodePanel.cpp
--- a/Engine/Source/Editor/GraphEditor/Private/SNodePanel.cpp
+++ b/Engine/Source/Editor/GraphEditor/Private/SNodePanel.cpp
@@ -702,9 +702,11 @@ FReply SNodePanel::OnMouseMove(const FGeometry& MyGeometry, const FPointerEvent&
 							}
 
 							// 2. Move selected nodes to new positions
+							SNode::FNodeSet NodeFilter;
+
 							for (int32 NodeIdx = 0; NodeIdx < DefferedNodesToMove.Num(); ++NodeIdx)
 							{
-								DefferedNodesToMove[NodeIdx].Node->MoveTo(DefferedNodesToMove[NodeIdx].NewPosition);
+								DefferedNodesToMove[NodeIdx].Node->MoveTo( DefferedNodesToMove[NodeIdx].NewPosition, NodeFilter );
 							}
 						}
 					}
diff --git a/Engine/Source/Editor/GraphEditor/Public/SGraphNode.h b/Engine/Source/Editor/GraphEditor/Public/SGraphNode.h
--- a/Engine/Source/Editor/GraphEditor/Public/SGraphNode.h
+++ b/Engine/Source/Editor/GraphEditor/Public/SGraphNode.h
@@ -81,7 +81,7 @@ public:
 	// End of SWidget interface
 
 	// SNodePanel::SNode interface
-	virtual void MoveTo( const FVector2D& NewPosition ) OVERRIDE;
+	virtual void MoveTo( const FVector2D& NewPosition, FNodeSet& NodeFilter ) OVERRIDE;
 	virtual FVector2D GetPosition() const OVERRIDE;
 	virtual FString GetNodeComment() const OVERRIDE;
 	virtual UObject* GetObjectBeingDisplayed() const OVERRIDE;
diff --git a/Engine/Source/Editor/GraphEditor/Public/SGraphNodeComment.h b/Engine/Source/Editor/GraphEditor/Public/SGraphNodeComment.h
--- a/Engine/Source/Editor/GraphEditor/Public/SGraphNodeComment.h
+++ b/Engine/Source/Editor/GraphEditor/Public/SGraphNodeComment.h
@@ -78,7 +78,7 @@ protected:
 	void HandleSelection(bool bIsSelected, bool bUpdateNodesUnderComment = false) const;
 
 	/** called when user is moving the comment node */
-	virtual void MoveTo(const FVector2D& NewPosition) OVERRIDE;
+	virtual void MoveTo(const FVector2D& NewPosition, FNodeSet& NodeFilter ) OVERRIDE;
 
 private:
 
diff --git a/Engine/Source/Editor/GraphEditor/Public/SNodePanel.h b/Engine/Source/Editor/GraphEditor/Public/SNodePanel.h
--- a/Engine/Source/Editor/GraphEditor/Public/SNodePanel.h
+++ b/Engine/Source/Editor/GraphEditor/Public/SNodePanel.h
@@ -252,8 +252,14 @@ public:
 	class SNode : public SBorder
 	{
 	public:
-		/** @param NewPosition  The Node should be relocated to this position in the graph panel */
-		virtual void MoveTo(const FVector2D& NewPosition)
+
+		typedef TSet<TWeakPtr<SNodePanel::SNode>> FNodeSet;
+
+		/**
+		* @param NewPosition	The Node should be relocated to this position in the graph panel
+		* @param NodeFilter		Set of nodes to prevent movement on, after moving successfully a node is added to this set.
+		*/
+		virtual void MoveTo(const FVector2D& NewPosition, FNodeSet& NodeFilter )
 		{
 		}
 
