Fix for a display string.

diff --git a/Engine/Source/Editor/GraphEditor/Private/GraphDiffControl.cpp b/Engine/Source/Editor/GraphEditor/Private/GraphDiffControl.cpp
--- a/Engine/Source/Editor/GraphEditor/Private/GraphDiffControl.cpp
+++ b/Engine/Source/Editor/GraphEditor/Private/GraphDiffControl.cpp
@@ -295,7 +295,7 @@ static void DiffR_NodePinCount(FDiffResults& Results, class UEdGraphNode* Node,
 		Args.Add(TEXT("NewCount"), NewCount);
 		Diff.ToolTip =  FText::Format(LOCTEXT("DIF_NodePinCountChangedToolTip", "Node '{NodeName}' had {OriginalCount} Pins, now has {NewCount} Pins"), Args).ToString();
 		Diff.DisplayColor = FLinearColor(0.45f,0.4f,0.4f);
-		Diff.DisplayString =  (OriginalCount < NewCount) ?  (FText::Format(LOCTEXT("DIF_NodePinCountIncreased", "Add Pin '{NodeName'"), Args).ToString()) : 
+		Diff.DisplayString =  (OriginalCount < NewCount) ?  (FText::Format(LOCTEXT("DIF_NodePinCountIncreased", "Add Pin '{NodeName}'"), Args).ToString()) : 
 			(FText::Format(LOCTEXT("DIF_NodePinCountDecreased", "Removed Pin '{NodeName}'"), Args).ToString());
 		Results.Add(Diff);
 	}
