Added trackpad gestures support to blueprint and material editors

diff --git a/Engine/Source/Editor/GraphEditor/Private/SNodePanel.cpp b/Engine/Source/Editor/GraphEditor/Private/SNodePanel.cpp
--- a/Engine/Source/Editor/GraphEditor/Private/SNodePanel.cpp
+++ b/Engine/Source/Editor/GraphEditor/Private/SNodePanel.cpp
@@ -914,6 +914,27 @@ void SNodePanel::OnKeyboardFocusLost( const FKeyboardFocusEvent& InKeyboardFocus
 	LastKeyGestureDetected.bShift = false;
 }
 
+FReply SNodePanel::OnTouchGesture( const FGeometry& MyGeometry, const FPointerEvent& GestureEvent )
+{
+	const EGestureEvent::Type GestureType = GestureEvent.GetGestureType();
+	const FVector2D& GestureDelta = GestureEvent.GetGestureDelta();
+	if (GestureType == EGestureEvent::Magnify)
+	{
+		// We want to zoom into this point; i.e. keep it the same fraction offset into the panel
+		const FVector2D WidgetSpaceCursorPos = MyGeometry.AbsoluteToLocal(GestureEvent.GetScreenSpacePosition());
+		const int32 ZoomLevelDelta = FMath::FloorToInt(GestureDelta.X * 10);
+		ChangeZoomLevel(ZoomLevelDelta, WidgetSpaceCursorPos, GestureEvent.IsControlDown());
+		return FReply::Handled();
+	}
+	else if (GestureType == EGestureEvent::Scroll)
+	{
+		this->bIsPanning = true;
+		ViewOffset -= GestureDelta / GetZoomAmount();
+		return FReply::Handled();
+	}
+	return FReply::Unhandled();
+}
+
 void SNodePanel::FindNodesAffectedByMarquee( FGraphPanelSelectionSet& OutAffectedNodes ) const
 {
 	OutAffectedNodes.Empty();
diff --git a/Engine/Source/Editor/GraphEditor/Public/SNodePanel.h b/Engine/Source/Editor/GraphEditor/Public/SNodePanel.h
--- a/Engine/Source/Editor/GraphEditor/Public/SNodePanel.h
+++ b/Engine/Source/Editor/GraphEditor/Public/SNodePanel.h
@@ -354,6 +354,7 @@ public:
 	virtual FReply OnKeyDown( const FGeometry& MyGeometry, const FKeyboardEvent& InKeyboardEvent ) OVERRIDE;
 	virtual FReply OnKeyUp( const FGeometry& MyGeometry, const FKeyboardEvent& InKeyboardEvent ) OVERRIDE;
 	virtual void OnKeyboardFocusLost( const FKeyboardFocusEvent& InKeyboardFocusEvent ) OVERRIDE;
+	virtual FReply OnTouchGesture( const FGeometry& MyGeometry, const FPointerEvent& GestureEvent ) OVERRIDE;
 
 	// End of SWidget interface
 public:
