GraphEditor: Removed the arrow head offset for input-side wires when no arrow head is being rendered

diff --git a/Engine/Source/Editor/GraphEditor/Private/ConnectionDrawingPolicy.cpp b/Engine/Source/Editor/GraphEditor/Private/ConnectionDrawingPolicy.cpp
--- a/Engine/Source/Editor/GraphEditor/Private/ConnectionDrawingPolicy.cpp
+++ b/Engine/Source/Editor/GraphEditor/Private/ConnectionDrawingPolicy.cpp
@@ -532,6 +532,7 @@ FKismetConnectionDrawingPolicy::FKismetConnectionDrawingPolicy(int32 InBackLayer
 
 	// Don't want to draw ending arrowheads
 	ArrowImage = nullptr;
+	ArrowRadius = FVector2D::ZeroVector;
 
 	// But we do want to draw midpoint arrowheads
 	if (GetDefault<UEditorExperimentalSettings>()->bDrawMidpointArrowsInBlueprints)
@@ -1210,6 +1211,7 @@ FSoundCueGraphConnectionDrawingPolicy::FSoundCueGraphConnectionDrawingPolicy(int
 
 	// Don't want to draw ending arrowheads
 	ArrowImage = nullptr;
+	ArrowRadius = FVector2D::ZeroVector;
 }
 
 void FSoundCueGraphConnectionDrawingPolicy::Draw(TMap<TSharedRef<SWidget>, FArrangedWidget>& PinGeometries, FArrangedChildren& ArrangedNodes)
@@ -1347,6 +1349,7 @@ FMaterialGraphConnectionDrawingPolicy::FMaterialGraphConnectionDrawingPolicy(int
 {
 	// Don't want to draw ending arrowheads
 	ArrowImage = nullptr;
+	ArrowRadius = FVector2D::ZeroVector;
 }
 
 void FMaterialGraphConnectionDrawingPolicy::DetermineWiringStyle(UEdGraphPin* OutputPin, UEdGraphPin* InputPin, /*inout*/ float& Thickness, /*inout*/ FLinearColor& WireColor, /*inout*/bool& bDrawBubbles, /*inout*/ bool& bBidirectional)
