Blueprints: Add reroute nodes to help organize wires
- Clean up FDragConnection construction and add a factory function in SGraphPin to extend it for some pins
- Clean up arguments to FOnGetContextMenuFor in SGraphPanel
- Various modifications to SGraphPin to permit implementation of wire knots/reroute nodes

diff --git a/Engine/Source/Editor/BlueprintGraph/Classes/K2Node_Knot.h b/Engine/Source/Editor/BlueprintGraph/Classes/K2Node_Knot.h
new file mode 100644
--- /dev/null
+++ b/Engine/Source/Editor/BlueprintGraph/Classes/K2Node_Knot.h
@@ -0,0 +1,40 @@
+// Copyright 1998-2014 Epic Games, Inc. All Rights Reserved.
+
+#pragma once
+
+#include "K2Node_Knot.generated.h"
+
+UCLASS(MinimalAPI)
+class UK2Node_Knot : public UK2Node
+{
+	GENERATED_UCLASS_BODY()
+
+public:
+	// UEdGraphNode interface
+	virtual void AllocateDefaultPins() OVERRIDE;
+	virtual FString GetTooltip() const OVERRIDE;
+	virtual FText GetNodeTitle(ENodeTitleType::Type TitleType) const OVERRIDE;
+	virtual void GetMenuEntries(struct FGraphContextMenuBuilder& ContextMenuBuilder) const OVERRIDE;
+	virtual bool ShouldOverridePinNames() const OVERRIDE;
+	virtual FString GetPinNameOverride(const UEdGraphPin& Pin) const OVERRIDE;
+	virtual void OnRenameNode(const FString& NewName) OVERRIDE;
+	virtual TSharedPtr<class INameValidatorInterface> MakeNameValidator() const OVERRIDE;
+	// End of UEdGraphNode interface
+
+	// UK2Node interface
+	virtual bool IsNodeSafeToIgnore() const OVERRIDE;
+	virtual void ExpandNode(class FKismetCompilerContext& CompilerContext, UEdGraph* SourceGraph) OVERRIDE;
+	virtual void NotifyPinConnectionListChanged(UEdGraphPin* Pin) OVERRIDE;
+	virtual void PostReconstructNode() OVERRIDE;
+	// End of UK2Node interface
+
+	UEdGraphPin* GetInputPin() const
+	{
+		return Pins[0];
+	}
+
+	UEdGraphPin* GetOutputPin() const
+	{
+		return Pins[1];
+	}
+};
diff --git a/Engine/Source/Editor/BlueprintGraph/Private/EdGraphSchema_K2.cpp b/Engine/Source/Editor/BlueprintGraph/Private/EdGraphSchema_K2.cpp
--- a/Engine/Source/Editor/BlueprintGraph/Private/EdGraphSchema_K2.cpp
+++ b/Engine/Source/Editor/BlueprintGraph/Private/EdGraphSchema_K2.cpp
@@ -4260,9 +4260,7 @@ void UEdGraphSchema_K2::CombineTwoPinNetsAndRemoveOldPins(UEdGraphPin* InPinA, U
 {
 	check(InPinA != NULL);
 	check(InPinB != NULL);
-
 	ensure(InPinA->Direction != InPinB->Direction);
-	ensure(InPinA->GetOwningNode() != InPinB->GetOwningNode());
 
 	if ((InPinA->LinkedTo.Num() == 0) && (InPinA->Direction == EGPD_Input))
 	{
diff --git a/Engine/Source/Editor/BlueprintGraph/Private/K2Node_Knot.cpp b/Engine/Source/Editor/BlueprintGraph/Private/K2Node_Knot.cpp
new file mode 100644
--- /dev/null
+++ b/Engine/Source/Editor/BlueprintGraph/Private/K2Node_Knot.cpp
@@ -0,0 +1,160 @@
+// Copyright 1998-2014 Epic Games, Inc. All Rights Reserved.
+
+#include "BlueprintGraphPrivatePCH.h"
+#include "K2ActionMenuBuilder.h"
+
+#define LOCTEXT_NAMESPACE "K2Node_Knot"
+
+/////////////////////////////////////////////////////
+// UK2Node_Knot
+
+UK2Node_Knot::UK2Node_Knot(const class FPostConstructInitializeProperties& PCIP)
+	: Super(PCIP)
+{
+	bCanRenameNode = true;
+}
+
+void UK2Node_Knot::AllocateDefaultPins()
+{
+	const UEdGraphSchema_K2* Schema = GetDefault<UEdGraphSchema_K2>();
+
+	const FString InputPinName(TEXT("InputPin"));
+	const FString OutputPinName(TEXT("OutputPin"));
+
+	UEdGraphPin* MyInputPin = CreatePin(EGPD_Input, Schema->PC_Wildcard, FString(), nullptr, /*bIsArray=*/ false, /*bIsReference=*/ false, InputPinName);
+	MyInputPin->bDefaultValueIsIgnored = true;
+
+	UEdGraphPin* MyOutputPin = CreatePin(EGPD_Output, Schema->PC_Wildcard, FString(), nullptr, /*bIsArray=*/ false, /*bIsReference=*/ false, OutputPinName);
+}
+
+FString UK2Node_Knot::GetTooltip() const
+{
+	//@TODO: Should pull the tooltip from the source pin
+	return LOCTEXT("KnotTooltip", "Reroute Node (reroutes wires)").ToString();
+}
+
+FText UK2Node_Knot::GetNodeTitle(ENodeTitleType::Type TitleType) const
+{
+	if (TitleType == ENodeTitleType::EditableTitle)
+	{
+		return FText::FromString(NodeComment);
+	}
+	else
+	{
+		return LOCTEXT("KnotTitle", "Reroute Node");
+	}
+}
+
+void UK2Node_Knot::GetMenuEntries(struct FGraphContextMenuBuilder& ContextMenuBuilder) const
+{
+	if (ContextMenuBuilder.FromPin != nullptr)
+	{
+		UK2Node_Knot* TemplateNode = NewObject<UK2Node_Knot>(GetTransientPackage(), GetClass());
+
+		FString EmptyCategory;
+		FText MenuDesc = LOCTEXT("KnotMenuDescription", "Add Reroute Node...");
+		FString Tooltip = TemplateNode->GetTooltip();//@TODO: Make this work +LOCTEXT("KnotMenuExtraTooltip", "\nYou can also create one by Shift+Dragging off a pin").ToString();
+		FString Keywords = TemplateNode->GetKeywords();
+
+		TSharedPtr<FEdGraphSchemaAction_K2NewNode> NodeAction = FK2ActionMenuBuilder::AddNewNodeAction(ContextMenuBuilder, EmptyCategory, MenuDesc, Tooltip, 0, Keywords);
+		NodeAction->NodeTemplate = TemplateNode;
+		NodeAction->SearchTitle = TemplateNode->GetNodeSearchTitle();
+	}
+}
+
+void UK2Node_Knot::ExpandNode(class FKismetCompilerContext& CompilerContext, UEdGraph* SourceGraph)
+{
+	const UEdGraphSchema_K2* K2Schema = GetDefault<UEdGraphSchema_K2>();
+
+	UEdGraphPin* MyInputPin = GetInputPin();
+	UEdGraphPin* MyOutputPin = GetOutputPin();
+
+	K2Schema->CombineTwoPinNetsAndRemoveOldPins(MyInputPin, MyOutputPin);
+}
+
+bool UK2Node_Knot::IsNodeSafeToIgnore() const
+{
+	return true;
+}
+
+void UK2Node_Knot::NotifyPinConnectionListChanged(UEdGraphPin* Pin)
+{
+	const UEdGraphSchema_K2* K2Schema = GetDefault<UEdGraphSchema_K2>();
+	UEdGraphPin* MyInputPin = GetInputPin();
+	UEdGraphPin* MyOutputPin = GetOutputPin();
+
+	const int32 NumLinks = MyInputPin->LinkedTo.Num() + MyOutputPin->LinkedTo.Num();
+
+	if (Pin->LinkedTo.Num() > 0)
+	{
+		// Just made a connection, was it the first?
+		if (NumLinks == 1)
+		{
+			UEdGraphPin* TypeSource = Pin->LinkedTo[0];
+
+			MyInputPin->PinType = TypeSource->PinType;
+			MyOutputPin->PinType = TypeSource->PinType;
+		}
+	}
+	else
+	{
+		// Just broke a connection, was it the last?
+		if (NumLinks == 0)
+		{
+			// Revert to wildcard
+			MyInputPin->BreakAllPinLinks();
+			MyInputPin->PinType.ResetToDefaults();
+			MyInputPin->PinType.PinCategory = K2Schema->PC_Wildcard;
+
+			MyOutputPin->BreakAllPinLinks();
+			MyOutputPin->PinType.ResetToDefaults();
+			MyOutputPin->PinType.PinCategory = K2Schema->PC_Wildcard;
+		}
+	}
+}
+
+void UK2Node_Knot::PostReconstructNode()
+{
+	UEdGraphPin* MyInputPin = GetInputPin();
+	UEdGraphPin* MyOutputPin = GetOutputPin();
+
+	// Find a pin that has connections to use to jumpstart the wildcard process
+	for (int32 PinIndex = 0; PinIndex < Pins.Num(); ++PinIndex)
+	{
+		if (Pins[PinIndex]->LinkedTo.Num() > 0)
+		{
+			// The pin is linked, continue to use its type as the type for all pins.
+			UEdGraphPin* TypeSource = Pins[PinIndex]->LinkedTo[0];
+
+			MyInputPin->PinType = TypeSource->PinType;
+			MyOutputPin->PinType = TypeSource->PinType;
+			break;
+		}
+	}
+}
+
+bool UK2Node_Knot::ShouldOverridePinNames() const
+{
+	return true;
+}
+
+FString UK2Node_Knot::GetPinNameOverride(const UEdGraphPin& Pin) const
+{
+	// Keep the pin size tiny
+	return FString();
+}
+
+void UK2Node_Knot::OnRenameNode(const FString& NewName)
+{
+	NodeComment = NewName;
+}
+
+TSharedPtr<class INameValidatorInterface> UK2Node_Knot::MakeNameValidator() const
+{
+	// Comments can be duplicated, etc...
+	return MakeShareable(new FDummyNameValidator(EValidatorResult::Ok));
+}
+
+/////////////////////////////////////////////////////
+
+#undef LOCTEXT_NAMESPACE
diff --git a/Engine/Source/Editor/GraphEditor/Private/DragConnection.cpp b/Engine/Source/Editor/GraphEditor/Private/DragConnection.cpp
--- a/Engine/Source/Editor/GraphEditor/Private/DragConnection.cpp
+++ b/Engine/Source/Editor/GraphEditor/Private/DragConnection.cpp
@@ -4,32 +4,11 @@
 #include "GraphEditorCommon.h"
 #include "ScopedTransaction.h"
 
-TSharedRef<FDragConnection> FDragConnection::New(const TSharedRef<SGraphPanel>& InGraphPanel, const TSharedRef<SGraphPin>& InStartingPin)
-{
-	TSharedRef<FDragConnection> Operation = MakeShareable(new FDragConnection);
-	Operation->GraphPanel = InGraphPanel;
-	Operation->StartingPins.Add(InStartingPin);
-
-	UEdGraphPin* PinObj = InStartingPin->GetPinObj();
-
-	// adjust the decorator away from the current mouse location a small amount based on cursor size
-	Operation->DecoratorAdjust =  (PinObj->Direction == EGPD_Input)
-		? FSlateApplication::Get().GetCursorSize() * FVector2D(-1.0f, 1.0f)
-		: FSlateApplication::Get().GetCursorSize();
-
-	Operation->Init(Operation->GraphPanel.ToSharedRef(), Operation->StartingPins);
-	Operation->Construct();
-	return Operation;
-}
 
-TSharedRef<FDragConnection> FDragConnection::New(const TSharedRef<SGraphPanel>& InGraphPanel, const TArray< TSharedRef<SGraphPin> >& InStartingPins, bool bIsMovingLinks)
+TSharedRef<FDragConnection> FDragConnection::New(const TSharedRef<SGraphPanel>& InGraphPanel, const TArray< TSharedRef<SGraphPin> >& InStartingPins, bool bInShiftOperation)
 {
-	TSharedRef<FDragConnection> Operation = MakeShareable(new FDragConnection);
-	Operation->GraphPanel = InGraphPanel;
-	Operation->StartingPins = InStartingPins;
-	Operation->DecoratorAdjust = FSlateApplication::Get().GetCursorSize();
-	Operation->Init(Operation->GraphPanel.ToSharedRef(), Operation->StartingPins);
-	Operation->Construct();
+	TSharedRef<FDragConnection> Operation = MakeShareable(new FDragConnection(InGraphPanel, InStartingPins, bInShiftOperation));
+
 	return Operation;
 }
 
@@ -40,11 +19,6 @@ void FDragConnection::OnDrop( bool bDropWasHandled, const FPointerEvent& MouseEv
 	Super::OnDrop(bDropWasHandled, MouseEvent);
 }
 
-UEdGraphPin* FDragConnection::GetGraphPinForSPin(TSharedRef<SGraphPin>& SPin)
-{
-	return SPin->GetPinObj();
-}
-
 void FDragConnection::OnDragged(const class FDragDropEvent& DragDropEvent)
 {
 	FVector2D TargetPosition = DragDropEvent.GetScreenSpacePosition();
@@ -63,12 +37,12 @@ void FDragConnection::HoverTargetChanged()
 
 	if (TargetPinObj != NULL)
 	{
-		ValidateGraphPinList();
+		TArray<UEdGraphPin*> ValidSourcePins;
+		ValidateGraphPinList(/*out*/ ValidSourcePins);
 
 		// Check the schema for connection responses
-		for (TArray< TSharedRef<SGraphPin> >::TIterator PinIterator(StartingPins); PinIterator; ++PinIterator)
+		for (UEdGraphPin* StartingPinObj : ValidSourcePins)
 		{
-			UEdGraphPin* StartingPinObj = GetGraphPinForSPin(*PinIterator);
 			if (TargetPinObj != StartingPinObj)
 			{
 				// The Graph object in which the pins reside.
@@ -152,27 +126,44 @@ void FDragConnection::HoverTargetChanged()
 	}
 }
 
-void FDragConnection::Init( const TSharedRef<SGraphPanel>& InGraphPanel, const TArray< TSharedRef<SGraphPin> >& InStartingPins )
+FDragConnection::FDragConnection(const TSharedRef<SGraphPanel>& InGraphPanel, const TArray< TSharedRef<SGraphPin> >& InStartingPins, bool bInShiftOperation)
+	: bShiftOperation(bInShiftOperation)
 {
-	for (TArray< TSharedRef<SGraphPin> >::TConstIterator PinIterator(InStartingPins); PinIterator; ++PinIterator)
+	GraphPanel = InGraphPanel;
+	StartingPins = InStartingPins;
+
+	// adjust the decorator away from the current mouse location a small amount based on cursor size
+	DecoratorAdjust = FSlateApplication::Get().GetCursorSize();
+	if (StartingPins.Num() > 0)
 	{
-		InGraphPanel->OnBeginMakingConnection( *PinIterator );
+		UEdGraphPin* PinObj = StartingPins[0]->GetPinObj();
+
+		DecoratorAdjust = (PinObj->Direction == EGPD_Input)
+			? FSlateApplication::Get().GetCursorSize() * FVector2D(-1.0f, 1.0f)
+			: FSlateApplication::Get().GetCursorSize();
 	}
+
+	for (const TSharedRef<SGraphPin> PinRef : StartingPins)
+	{
+		InGraphPanel->OnBeginMakingConnection(PinRef->GetPinObj());
+	}
+
+	Construct();
 }
 
 FReply FDragConnection::DroppedOnPin(FVector2D ScreenPosition, FVector2D GraphPosition)
 {
-	ValidateGraphPinList();
+	TArray<UEdGraphPin*> ValidSourcePins;
+	ValidateGraphPinList(/*out*/ ValidSourcePins);
 
 	const FScopedTransaction Transaction( NSLOCTEXT("UnrealEd", "GraphEd_CreateConnection", "Create Pin Link") );
 
 	UEdGraphPin* PinA = GetHoveredPin();
 	bool bError = false;
 	TSet<UEdGraphNode*> NodeList;
-	for (TArray< TSharedRef<SGraphPin> >::TIterator PinIterator(StartingPins); !bError && PinIterator; ++PinIterator)
-	{
-		UEdGraphPin* PinB = GetGraphPinForSPin(*PinIterator);
 
+	for (UEdGraphPin* PinB : ValidSourcePins)
+	{
 		if ((PinA != NULL) && (PinB != NULL))
 		{
 			UEdGraph* MyGraphObj = PinA->GetOwningNode()->GetGraph();
@@ -206,18 +197,12 @@ FReply FDragConnection::DroppedOnPin(FVector2D ScreenPosition, FVector2D GraphPo
 
 FReply FDragConnection::DroppedOnPanel( const TSharedRef< SWidget >& Panel, FVector2D ScreenPosition, FVector2D GraphPosition, UEdGraph& Graph)
 {
-	ValidateGraphPinList();
-
 	// Gather any source drag pins
 	TArray<UEdGraphPin*> PinObjects;
-	for (auto PinIt = StartingPins.CreateConstIterator(); PinIt; ++PinIt)
-	{
-		UEdGraphPin* GraphPin = (*PinIt)->GetPinObj();
-		PinObjects.Add(GraphPin);
-	}
+	ValidateGraphPinList(/*out*/ PinObjects);
 
 	// Create a context menu
-	TSharedPtr<SWidget> WidgetToFocus = GraphPanel->SummonContextMenu(ScreenPosition, GraphPosition, NULL, NULL, PinObjects);
+	TSharedPtr<SWidget> WidgetToFocus = GraphPanel->SummonContextMenu(ScreenPosition, GraphPosition, NULL, NULL, PinObjects, bShiftOperation);
 
 	// Give the context menu focus
 	return (WidgetToFocus.IsValid())
@@ -226,30 +211,34 @@ FReply FDragConnection::DroppedOnPanel( const TSharedRef< SWidget >& Panel, FVec
 }
 
 
-void FDragConnection::ValidateGraphPinList( )
+void FDragConnection::ValidateGraphPinList(TArray<UEdGraphPin*>& OutValidPins)
 {
+	OutValidPins.Empty(StartingPins.Num());
+
 	for (TArray< TSharedRef<SGraphPin> >::TIterator PinIterator(StartingPins); PinIterator; ++PinIterator)
 	{
-		UEdGraphPin* StartingPinObj = GetGraphPinForSPin(*PinIterator);
-
-		if( StartingPinObj )
+		if (UEdGraphPin* StartingPinObj = (*PinIterator)->GetPinObj())
 		{
 			//Check whether the list contains updated pin object references by checking its outer class type
-			if( StartingPinObj->GetOuter() == NULL || !StartingPinObj->GetOuter()->IsA( UEdGraphNode::StaticClass() ) )
+			if ((StartingPinObj->GetOuter() == NULL) || !StartingPinObj->GetOuter()->IsA(UEdGraphNode::StaticClass()))
 			{
 				//This pin object reference is old. So remove it from the list.
 				TSharedRef<SGraphPin> PinPtr = *PinIterator;
 				StartingPins.Remove( PinPtr );
 			}
+			else
+			{
+				OutValidPins.Add(StartingPinObj);
+			}
 		}
 	}
 }
 
-void FDragConnection::OnDragBegin( const TSharedRef<class SGraphPin>& InPin)
+void FDragConnection::OnDragBegin(const TSharedRef<class SGraphPin>& InPin)
 {
-	if( !StartingPins.Contains( InPin ) )
+	if (!StartingPins.Contains(InPin))
 	{
-		StartingPins.Add( InPin );
-		GraphPanel->OnBeginMakingConnection( InPin );
+		StartingPins.Add(InPin);
+		GraphPanel->OnBeginMakingConnection(InPin->GetPinObj());
 	}
 }
diff --git a/Engine/Source/Editor/GraphEditor/Private/DragConnection.h b/Engine/Source/Editor/GraphEditor/Private/DragConnection.h
--- a/Engine/Source/Editor/GraphEditor/Private/DragConnection.h
+++ b/Engine/Source/Editor/GraphEditor/Private/DragConnection.h
@@ -7,8 +7,7 @@ class FDragConnection : public FGraphEditorDragDropAction
 public:
 	DRAG_DROP_OPERATOR_TYPE(FDragConnection, FGraphEditorDragDropAction)
 
-	static TSharedRef<FDragConnection> New(const TSharedRef<SGraphPanel>& InGraphPanel, const TSharedRef<SGraphPin>& InStartingPin);
-	static TSharedRef<FDragConnection> New(const TSharedRef<SGraphPanel>& InGraphPanel, const TArray< TSharedRef<SGraphPin> >& InStartingPins, bool bIsMovingLinks);
+	static TSharedRef<FDragConnection> New(const TSharedRef<SGraphPanel>& InGraphPanel, const TArray< TSharedRef<SGraphPin> >& InStartingPins, bool bInShiftOperation);
 	
 	// FDragDropOperation interface
 	virtual void OnDrop( bool bDropWasHandled, const FPointerEvent& MouseEvent ) OVERRIDE;
@@ -17,26 +16,28 @@ public:
 	// FGraphEditorDragDropAction interface
 	virtual void HoverTargetChanged() OVERRIDE;
 	virtual FReply DroppedOnPin(FVector2D ScreenPosition, FVector2D GraphPosition) OVERRIDE;
-	virtual FReply DroppedOnPanel( const TSharedRef< SWidget >& Panel, FVector2D ScreenPosition, FVector2D GraphPosition, UEdGraph& Graph) OVERRIDE;
-	virtual void OnDragBegin( const TSharedRef<class SGraphPin>& InPin) OVERRIDE;
-	virtual void OnDragged (const class FDragDropEvent& DragDropEvent ) OVERRIDE;
+	virtual FReply DroppedOnPanel(const TSharedRef< SWidget >& Panel, FVector2D ScreenPosition, FVector2D GraphPosition, UEdGraph& Graph) OVERRIDE;
+	virtual void OnDragBegin(const TSharedRef<class SGraphPin>& InPin) OVERRIDE;
+	virtual void OnDragged(const class FDragDropEvent& DragDropEvent) OVERRIDE;
 	// End of FGraphEditorDragDropAction interface
-	
-protected:
-	typedef FGraphEditorDragDropAction Super;
-
-	UEdGraphPin* GetGraphPinForSPin(TSharedRef<SGraphPin>& SPin);
 
 	/*
 	 *	Function to check validity of graph pins in the StartPins list. This check helps to prevent processing graph pins which are outdated.
 	 */
-	void ValidateGraphPinList( );
-	
-	void Init( const TSharedRef<SGraphPanel>& InGraphPanel, const TArray< TSharedRef<SGraphPin> >& InStartingPins );
+	virtual void ValidateGraphPinList(TArray<UEdGraphPin*>& OutValidPins);
+
+protected:
+	typedef FGraphEditorDragDropAction Super;
+
+	FDragConnection(const TSharedRef<SGraphPanel>& InGraphPanel, const TArray< TSharedRef<SGraphPin> >& InStartingPins, bool bInShiftOperation);
 
 protected:
 	TSharedPtr<SGraphPanel> GraphPanel;
 	TArray< TSharedRef<SGraphPin> > StartingPins;
+
 	/** Offset information for the decorator widget */
-	FVector2D	DecoratorAdjust;
+	FVector2D DecoratorAdjust;
+
+	/** Was shift pressed when the drag started? */
+	bool bShiftOperation;
 };
diff --git a/Engine/Source/Editor/GraphEditor/Private/GraphEditorDragDropAction.cpp b/Engine/Source/Editor/GraphEditor/Private/GraphEditorDragDropAction.cpp
--- a/Engine/Source/Editor/GraphEditor/Private/GraphEditorDragDropAction.cpp
+++ b/Engine/Source/Editor/GraphEditor/Private/GraphEditorDragDropAction.cpp
@@ -5,7 +5,7 @@
 
 UEdGraphPin* FGraphEditorDragDropAction::GetHoveredPin() const
 {
-	return HoveredPin.IsValid() ? HoveredPin->GetPinObj() : NULL;
+	return HoveredPin.Get();
 }
 
 UEdGraphNode* FGraphEditorDragDropAction::GetHoveredNode() const
@@ -33,7 +33,7 @@ UEdGraph* FGraphEditorDragDropAction::GetHoveredGraph() const
 	return NULL;
 }
 
-void FGraphEditorDragDropAction::SetHoveredPin(const TSharedPtr<SGraphPin>& InPin)
+void FGraphEditorDragDropAction::SetHoveredPin(UEdGraphPin* InPin)
 {
 	if (HoveredPin != InPin)
 	{
diff --git a/Engine/Source/Editor/GraphEditor/Private/NodeFactory.cpp b/Engine/Source/Editor/GraphEditor/Private/NodeFactory.cpp
--- a/Engine/Source/Editor/GraphEditor/Private/NodeFactory.cpp
+++ b/Engine/Source/Editor/GraphEditor/Private/NodeFactory.cpp
@@ -9,6 +9,7 @@
 
 #include "SGraphNodeDefault.h"
 #include "SGraphNodeComment.h"
+#include "SGraphNodeKnot.h"
 
 #include "KismetNodes/SGraphNodeK2Base.h"
 #include "KismetNodes/SGraphNodeK2Default.h"
@@ -188,6 +189,10 @@ TSharedPtr<SGraphNode> FNodeFactory::CreateNodeWidget(UEdGraphNode* InNode)
 		{
 			return SNew(SGraphNodeK2ArrayFunction, CallFunction);
 		}
+		else if (UK2Node_Knot* Knot = Cast<UK2Node_Knot>(InNode))
+		{
+			return SNew(SGraphNodeKnot, Knot);
+		}
 		else
 		{
 			return SNew(SGraphNodeK2Default, K2Node);
diff --git a/Engine/Source/Editor/GraphEditor/Private/SGraphEditorImpl.cpp b/Engine/Source/Editor/GraphEditor/Private/SGraphEditorImpl.cpp
--- a/Engine/Source/Editor/GraphEditor/Private/SGraphEditorImpl.cpp
+++ b/Engine/Source/Editor/GraphEditor/Private/SGraphEditorImpl.cpp
@@ -357,7 +357,7 @@ bool SGraphEditorImpl::GraphEd_OnGetGraphEnabled() const
 	return !bTitleBarOnly;
 }
 
-FActionMenuContent SGraphEditorImpl::GraphEd_OnGetContextMenuFor( const FVector2D& NodeAddPosition, UEdGraphNode* InGraphNode, UEdGraphPin* InGraphPin, const TArray<UEdGraphPin*>& InDragFromPins )
+FActionMenuContent SGraphEditorImpl::GraphEd_OnGetContextMenuFor(const FGraphContextMenuArguments& SpawnInfo)
 {
 	if (EdGraphObj != NULL)
 	{
@@ -365,9 +365,9 @@ FActionMenuContent SGraphEditorImpl::GraphEd_OnGetContextMenuFor( const FVector2
 		check(Schema);
 			
 		// Cache the pin this menu is being brought up for
-		GraphPinForMenu = InGraphPin;
+		GraphPinForMenu = SpawnInfo.GraphPin;
 
-		if ((InGraphPin != NULL) || (InGraphNode != NULL))
+		if ((SpawnInfo.GraphPin != NULL) || (SpawnInfo.GraphNode != NULL))
 		{
 			// Get all menu extenders for this context menu from the graph editor module
 			FGraphEditorModule& GraphEditorModule = FModuleManager::GetModuleChecked<FGraphEditorModule>( TEXT("GraphEditor") );
@@ -378,7 +378,7 @@ FActionMenuContent SGraphEditorImpl::GraphEd_OnGetContextMenuFor( const FVector2
 			{
 				if (MenuExtenderDelegates[i].IsBound())
 				{
-					Extenders.Add(MenuExtenderDelegates[i].Execute(this->Commands.ToSharedRef(), EdGraphObj, InGraphNode, InGraphPin, !IsEditable.Get()));
+					Extenders.Add(MenuExtenderDelegates[i].Execute(this->Commands.ToSharedRef(), EdGraphObj, SpawnInfo.GraphNode, SpawnInfo.GraphPin, !IsEditable.Get()));
 				}
 			}
 			TSharedPtr<FExtender> MenuExtender = FExtender::Combine(Extenders);
@@ -386,7 +386,7 @@ FActionMenuContent SGraphEditorImpl::GraphEd_OnGetContextMenuFor( const FVector2
 			// Show the menu for the pin or node under the cursor
 			const bool bShouldCloseAfterAction = true;
 			FMenuBuilder MenuBuilder( bShouldCloseAfterAction, this->Commands, MenuExtender );
-			Schema->GetContextMenuActions(EdGraphObj, InGraphNode, InGraphPin, &MenuBuilder, !IsEditable.Get());
+			Schema->GetContextMenuActions(EdGraphObj, SpawnInfo.GraphNode, SpawnInfo.GraphPin, &MenuBuilder, !IsEditable.Get());
 
 			return FActionMenuContent(MenuBuilder.MakeWidget());
 		}
@@ -400,8 +400,8 @@ FActionMenuContent SGraphEditorImpl::GraphEd_OnGetContextMenuFor( const FVector2
 				{
 					Content = OnCreateActionMenu.Execute(
 						EdGraphObj, 
-						NodeAddPosition, 
-						InDragFromPins, 
+						SpawnInfo.NodeAddPosition,
+						SpawnInfo.DragFromPins,
 						bAutoExpandActionMenu, 
 						SGraphEditor::FActionMenuClosed::CreateSP(this, &SGraphEditorImpl::OnClosedActionMenu)
 						);
@@ -411,8 +411,8 @@ FActionMenuContent SGraphEditorImpl::GraphEd_OnGetContextMenuFor( const FVector2
 					TSharedRef<SGraphEditorActionMenu> Menu =	
 						SNew(SGraphEditorActionMenu)
 						.GraphObj( EdGraphObj )
-						.NewNodePosition( NodeAddPosition )
-						.DraggedFromPins( InDragFromPins )
+						.NewNodePosition(SpawnInfo.NodeAddPosition)
+						.DraggedFromPins(SpawnInfo.DragFromPins)
 						.AutoExpandActionMenu(bAutoExpandActionMenu)
 						.OnClosedCallback( SGraphEditor::FActionMenuClosed::CreateSP(this, &SGraphEditorImpl::OnClosedActionMenu)
 						);
@@ -420,7 +420,7 @@ FActionMenuContent SGraphEditorImpl::GraphEd_OnGetContextMenuFor( const FVector2
 					Content = FActionMenuContent( Menu, Menu->GetFilterTextBox() );
 				}
 
-				if (InDragFromPins.Num() > 0)
+				if (SpawnInfo.DragFromPins.Num() > 0)
 				{
 					GraphPanel->PreservePinPreviewUntilForced();
 				}
diff --git a/Engine/Source/Editor/GraphEditor/Private/SGraphEditorImpl.h b/Engine/Source/Editor/GraphEditor/Private/SGraphEditorImpl.h
--- a/Engine/Source/Editor/GraphEditor/Private/SGraphEditorImpl.h
+++ b/Engine/Source/Editor/GraphEditor/Private/SGraphEditorImpl.h
@@ -93,7 +93,7 @@ public:
 
 	bool GraphEd_OnGetGraphEnabled() const;
 
-	FActionMenuContent GraphEd_OnGetContextMenuFor( const FVector2D& NodeAddPosition, UEdGraphNode* InGraphNode, UEdGraphPin* InGraphPin, const TArray<UEdGraphPin*>& InDragFromPins );
+	FActionMenuContent GraphEd_OnGetContextMenuFor(const FGraphContextMenuArguments& SpawnInfo);
 
 	void GraphEd_OnPanelUpdated();
 
diff --git a/Engine/Source/Editor/GraphEditor/Private/SGraphNode.cpp b/Engine/Source/Editor/GraphEditor/Private/SGraphNode.cpp
--- a/Engine/Source/Editor/GraphEditor/Private/SGraphNode.cpp
+++ b/Engine/Source/Editor/GraphEditor/Private/SGraphNode.cpp
@@ -53,7 +53,7 @@ FText SNodeTitle::GetNodeTitle() const
 
 FText SNodeTitle::GetHeadTitle() const
 {
-	return GraphNode->bCanRenameNode? GraphNode->GetNodeTitle(ENodeTitleType::EditableTitle) : CachedHeadTitle;
+	return GraphNode->bCanRenameNode ? GraphNode->GetNodeTitle(ENodeTitleType::EditableTitle) : CachedHeadTitle;
 }
 
 void SNodeTitle::RebuildWidget()
@@ -1110,11 +1110,11 @@ bool SGraphNode::IsNameReadOnly() const
 	return !GraphNode->bCanRenameNode;
 }
 
-bool SGraphNode::OnVerifyNameTextChanged ( const FText& InText, FText& OutErrorMessage ) 
+bool SGraphNode::OnVerifyNameTextChanged(const FText& InText, FText& OutErrorMessage)
 {
 	bool bValid(true);
 
-	if( GetEditableNodeTitle() != InText.ToString() && OnVerifyTextCommit.IsBound() )
+	if ((GetEditableNodeTitle() != InText.ToString()) && OnVerifyTextCommit.IsBound())
 	{
 		bValid = OnVerifyTextCommit.Execute(InText, GraphNode);
 	}
@@ -1127,12 +1127,15 @@ bool SGraphNode::OnVerifyNameTextChanged ( const FText& InText, FText& OutErrorM
 	return bValid;
 }
 
-void SGraphNode::OnNameTextCommited ( const FText& InText, ETextCommit::Type CommitInfo ) 
+void SGraphNode::OnNameTextCommited(const FText& InText, ETextCommit::Type CommitInfo)
 {
 	OnTextCommitted.ExecuteIfBound(InText, CommitInfo, GraphNode);
 	
 	UpdateErrorInfo();
-	ErrorReporting->SetError(ErrorMsg);
+	if (ErrorReporting.IsValid())
+	{
+		ErrorReporting->SetError(ErrorMsg);
+	}
 }
 
 void SGraphNode::RequestRename()
diff --git a/Engine/Source/Editor/GraphEditor/Private/SGraphNodeKnot.cpp b/Engine/Source/Editor/GraphEditor/Private/SGraphNodeKnot.cpp
new file mode 100644
--- /dev/null
+++ b/Engine/Source/Editor/GraphEditor/Private/SGraphNodeKnot.cpp
@@ -0,0 +1,361 @@
+// Copyright 1998-2014 Epic Games, Inc. All Rights Reserved.
+
+
+#include "GraphEditorCommon.h"
+#include "SGraphNodeKnot.h"
+
+class FAmbivalentDirectionDragConnection : public FDragConnection
+{
+public:
+	FAmbivalentDirectionDragConnection(UK2Node_Knot* InKnot, const TSharedRef<SGraphPanel>& InGraphPanel, const TArray< TSharedRef<SGraphPin> >& InStartingPins, bool bInShiftOperation)
+		: FDragConnection(InGraphPanel, InStartingPins, bInShiftOperation)
+		, KnotPtr(InKnot)
+		, StartScreenPos(FVector2D::ZeroVector)
+		, MostRecentScreenPos(FVector2D::ZeroVector)
+		, bLatchedStartScreenPos(false)
+	{
+	}
+		
+	// FDragDropOperation interface
+	virtual void OnDragged(const class FDragDropEvent& DragDropEvent) OVERRIDE;
+	// End of FDragDropOperation interface
+
+	UEdGraphPin* GetBestPin() const;
+
+	// FDragConnection interface
+	virtual void ValidateGraphPinList(TArray<UEdGraphPin*>& OutValidPins) OVERRIDE;
+	// End of FDragConnection interface
+
+protected:
+	TWeakObjectPtr<UK2Node_Knot> KnotPtr;
+
+	FVector2D StartScreenPos;
+	FVector2D MostRecentScreenPos;
+	bool bLatchedStartScreenPos;
+};
+
+UEdGraphPin* FAmbivalentDirectionDragConnection::GetBestPin() const
+{
+	if (bLatchedStartScreenPos)
+	{
+		if (UK2Node_Knot* Knot = KnotPtr.Get())
+		{
+			const bool bIsRight = MostRecentScreenPos.X >= StartScreenPos.X;
+			return bIsRight ? Knot->GetOutputPin() : Knot->GetInputPin();
+		}
+	}
+
+	return nullptr;
+}
+
+void FAmbivalentDirectionDragConnection::OnDragged(const class FDragDropEvent& DragDropEvent)
+{
+	if (bLatchedStartScreenPos)
+	{
+		const FVector2D LastScreenPos = MostRecentScreenPos;
+		MostRecentScreenPos = DragDropEvent.GetScreenSpacePosition();
+
+		// Switch directions on the preview connector as we cross from left to right of the starting drag point (or vis versa)
+		const bool bWasRight = LastScreenPos.X >= StartScreenPos.X;
+		const bool bIsRight = MostRecentScreenPos.X >= StartScreenPos.X;
+
+		if (bWasRight ^ bIsRight)
+		{
+			GraphPanel->OnStopMakingConnection(/*bForceStop=*/ true);
+			GraphPanel->OnBeginMakingConnection(GetBestPin());
+		}
+	}
+	else
+	{
+		StartScreenPos = DragDropEvent.GetScreenSpacePosition();
+		MostRecentScreenPos = StartScreenPos;
+		bLatchedStartScreenPos = true;
+	}
+
+	FDragConnection::OnDragged(DragDropEvent);
+}
+
+void FAmbivalentDirectionDragConnection::ValidateGraphPinList(TArray<UEdGraphPin*>& OutValidPins)
+{
+	OutValidPins.Empty(StartingPins.Num());
+
+	if (UK2Node_Knot* Knot = KnotPtr.Get())
+	{
+		bool bUseOutput = true;
+
+		// Pick output or input based on if the drag op is currently to the left or to the right of the starting drag point
+		if (bLatchedStartScreenPos)
+		{
+			bUseOutput = (StartScreenPos.X < MostRecentScreenPos.X);
+		}
+
+		if (UEdGraphPin* TargetPinObj = GetHoveredPin())
+		{
+// 			if (UK2Node_Knot* TargetKnot = Cast<UK2Node_Knot>(TargetPinObj->GetOwningNode()))
+// 			{
+// 				// The visible pin on a knot is always an output, so Rely on the direction matching; since the visible pin on another knot is always an output
+// 			}
+// 			else
+			{
+				// Dragging to another pin, pick the opposite direction as a source to maximize connection chances
+				if (TargetPinObj->Direction == EGPD_Input)
+				{
+					bUseOutput = true;
+				}
+				else
+				{
+					bUseOutput = false;
+				}
+			}
+		}
+
+		// Switch the effective valid pin so it makes sense for the current drag context
+		if (bUseOutput)
+		{
+			OutValidPins.Add(Knot->GetOutputPin());
+		}
+		else
+		{
+			OutValidPins.Add(Knot->GetInputPin());
+		}
+	}
+	else
+	{
+		// Fall back to the default behavior
+		FDragConnection::ValidateGraphPinList(OutValidPins);
+	}
+}
+
+/////////////////////////////////////////////////////
+// SGraphPinKnot
+
+class SGraphPinKnot : public SGraphPin
+{
+public:
+	SLATE_BEGIN_ARGS(SGraphPinKnot) {}
+	SLATE_END_ARGS()
+
+	void Construct(const FArguments& InArgs, UEdGraphPin* InPin);
+
+	// SWidget interface
+	virtual void OnDragEnter(const FGeometry& MyGeometry, const FDragDropEvent& DragDropEvent) OVERRIDE;
+	// End of SWidget interface
+
+protected:
+	// Begin SGraphPin interface
+	virtual TSharedRef<SWidget>	GetDefaultValueWidget() OVERRIDE;
+	virtual TSharedRef<FDragDropOperation> SpawnPinDragEvent(const TSharedRef<SGraphPanel>& InGraphPanel, const TArray< TSharedRef<SGraphPin> >& InStartingPins, bool bInShiftOperation) OVERRIDE;
+	virtual FReply OnPinMouseDown(const FGeometry& SenderGeometry, const FPointerEvent& MouseEvent) OVERRIDE;
+	virtual FSlateColor GetPinColor() const OVERRIDE;
+	// End SGraphPin interface
+};
+
+void SGraphPinKnot::Construct(const FArguments& InArgs, UEdGraphPin* InPin)
+{
+	SGraphPin::Construct(SGraphPin::FArguments().SideToSideMargin(0.0f), InPin);
+}
+
+void SGraphPinKnot::OnDragEnter(const FGeometry& MyGeometry, const FDragDropEvent& DragDropEvent)
+{
+	TSharedPtr<FDragDropOperation> Operation = DragDropEvent.GetOperation();
+	if (Operation.IsValid() && Operation->IsOfType<FDragConnection>())
+	{
+		TSharedPtr<FDragConnection> DragConnectionOp = StaticCastSharedPtr<FDragConnection>(Operation);
+
+		TArray<UEdGraphPin*> ValidPins;
+		DragConnectionOp->ValidateGraphPinList(/*out*/ ValidPins);
+		
+		if (ValidPins.Num() > 0)
+		{
+			if (UK2Node_Knot* Knot = Cast<UK2Node_Knot>(GraphPinObj->GetOwningNode()))
+			{
+				// Dragging to another pin, pick the opposite direction as a source to maximize connection chances
+				UEdGraphPin* PinToHoverOver = (ValidPins[0]->Direction == EGPD_Input) ? Knot->GetOutputPin() : Knot->GetInputPin();
+				DragConnectionOp->SetHoveredPin(PinToHoverOver);
+
+				// Pins treat being dragged over the same as being hovered outside of drag and drop if they know how to respond to the drag action.
+				SBorder::OnMouseEnter(MyGeometry, DragDropEvent);
+
+				return;
+			}
+		}
+	}
+
+	SGraphPin::OnDragEnter(MyGeometry, DragDropEvent);
+}
+
+FSlateColor SGraphPinKnot::GetPinColor() const
+{
+	// Make ourselves transparent if we're the input, since we are underneath the output pin and would double-blend looking ugly
+	return (GetPinObj()->Direction == EEdGraphPinDirection::EGPD_Input) ? FLinearColor::Transparent : SGraphPin::GetPinColor();
+}
+
+TSharedRef<SWidget>	SGraphPinKnot::GetDefaultValueWidget()
+{
+	return SNullWidget::NullWidget;
+}
+
+TSharedRef<FDragDropOperation> SGraphPinKnot::SpawnPinDragEvent(const TSharedRef<SGraphPanel>& InGraphPanel, const TArray< TSharedRef<SGraphPin> >& InStartingPins, bool bInShiftOperation)
+{
+	TSharedRef<FAmbivalentDirectionDragConnection> Operation = MakeShareable(new FAmbivalentDirectionDragConnection(CastChecked<UK2Node_Knot>(GetPinObj()->GetOwningNode()), InGraphPanel, InStartingPins, bInShiftOperation));
+
+	return Operation;
+}
+
+FReply SGraphPinKnot::OnPinMouseDown(const FGeometry& SenderGeometry, const FPointerEvent& MouseEvent)
+{
+	if (MouseEvent.GetEffectingButton() == EKeys::LeftMouseButton)
+	{
+		if (!GraphPinObj->bNotConnectable && IsEditable.Get())
+		{
+			if (MouseEvent.IsAltDown())
+			{
+				// Normally break connections, but overloaded here to delete the node entirely
+				const FScopedTransaction Transaction(FGenericCommands::Get().Delete->GetDescription());
+
+				UK2Node_Knot* NodeToDelete = CastChecked<UK2Node_Knot>(GetPinObj()->GetOwningNode());
+				UBlueprint* OwnerBlueprint = NodeToDelete->GetBlueprint();
+				NodeToDelete->GetGraph()->Modify();
+
+				FBlueprintEditorUtils::RemoveNode(OwnerBlueprint, NodeToDelete, /*bDontRecompile=*/ true);
+				FBlueprintEditorUtils::MarkBlueprintAsModified(OwnerBlueprint);
+
+				return FReply::Handled();
+			}
+			else if (MouseEvent.IsControlDown())
+			{
+				// Normally moves the connections from one pin to another, but moves the node instead since it's really representing a set of connections
+				// Returning unhandled will cause the node behind us to catch it and move us
+				return FReply::Unhandled();
+			}
+		}
+	}
+
+	return SGraphPin::OnPinMouseDown(SenderGeometry, MouseEvent);
+}
+
+//////////////////////////////////////////////////////////////////////////
+// SGraphNodeKnot
+
+void SGraphNodeKnot::Construct(const FArguments& InArgs, UK2Node_Knot* InKnot)
+{
+	SGraphNodeDefault::Construct(SGraphNodeDefault::FArguments().GraphNodeObj(InKnot));
+}
+
+void SGraphNodeKnot::UpdateGraphNode()
+{
+	const FVector2D NodeSpacerSize(42.0f, 24.0f);
+
+	InputPins.Empty();
+	OutputPins.Empty();
+
+	// Reset variables that are going to be exposed, in case we are refreshing an already setup node.
+	RightNodeBox.Reset();
+	LeftNodeBox.Reset();
+
+	//@TODO: Keyboard focus on edit doesn't work unless the node is visible, but the text is just the comment and it's already shown in a bubble, so Transparent black it is...
+	InlineEditableText = SNew(SInlineEditableTextBlock)
+		.ColorAndOpacity(FLinearColor::Transparent)
+		.Style(FEditorStyle::Get(), "Graph.Node.NodeTitleInlineEditableText")
+		.Text(this, &SGraphNodeKnot::GetEditableNodeTitleAsText)
+		.OnVerifyTextChanged(this, &SGraphNodeKnot::OnVerifyNameTextChanged)
+		.OnTextCommitted(this, &SGraphNodeKnot::OnNameTextCommited)
+		.IsReadOnly(this, &SGraphNodeKnot::IsNameReadOnly)
+		.IsSelected(this, &SGraphNodeKnot::IsSelectedExclusively);
+
+	this->ContentScale.Bind( this, &SGraphNode::GetContentScale );
+
+	this->ChildSlot
+		.HAlign(HAlign_Center)
+		.VAlign(VAlign_Center)
+		[
+			SNew(SOverlay)
+			+SOverlay::Slot()
+			[
+				// Grab handle to be able to move the node
+				SNew(SSpacer)
+				.Size(NodeSpacerSize)
+				.Visibility(EVisibility::Visible)
+				.Cursor(EMouseCursor::CardinalCross)
+			]
+
+			+SOverlay::Slot()
+// 			.VAlign(VAlign_Center)
+// 			.HAlign(HAlign_Center)
+			[
+				SNew(SVerticalBox)
+				+SVerticalBox::Slot()
+				.VAlign(VAlign_Top)
+				.HAlign(HAlign_Center)
+				[
+					SNew(SHorizontalBox)
+					+ SHorizontalBox::Slot()
+					.AutoWidth()
+					[
+						SNew(SOverlay)
+						+SOverlay::Slot()
+						[
+							// LEFT
+							SAssignNew(LeftNodeBox, SVerticalBox)
+						]
+						+SOverlay::Slot()
+						[
+							// RIGHT
+							SAssignNew(RightNodeBox, SVerticalBox)
+						]
+					]
+				]
+
+				+SVerticalBox::Slot()
+				.VAlign(VAlign_Fill)
+				.HAlign(HAlign_Fill)
+				[
+					InlineEditableText.ToSharedRef()
+				]
+			]
+		];
+
+	CreatePinWidgets();
+}
+
+const FSlateBrush* SGraphNodeKnot::GetShadowBrush(bool bSelected) const
+{
+	return bSelected ? FEditorStyle::GetBrush(TEXT("Graph.Node.ShadowSelected")) : FEditorStyle::GetNoBrush();
+}
+
+TSharedPtr<SGraphPin> SGraphNodeKnot::CreatePinWidget(UEdGraphPin* Pin) const
+{
+	return SNew(SGraphPinKnot, Pin);
+}
+
+void SGraphNodeKnot::AddPin(const TSharedRef<SGraphPin>& PinToAdd)
+{
+	PinToAdd->SetOwner(SharedThis(this));
+
+	const UEdGraphPin* PinObj = PinToAdd->GetPinObj();
+
+	if (PinToAdd->GetDirection() == EEdGraphPinDirection::EGPD_Input)
+	{
+		LeftNodeBox->AddSlot()
+			.AutoHeight()
+			.HAlign(HAlign_Left)
+			.VAlign(VAlign_Center)
+//			.Padding(10, 4)
+			[
+				PinToAdd
+			];
+		InputPins.Add(PinToAdd);
+	}
+	else
+	{
+		RightNodeBox->AddSlot()
+			.AutoHeight()
+			.HAlign(HAlign_Right)
+			.VAlign(VAlign_Center)
+//			.Padding(10, 4)
+			[
+				PinToAdd
+			];
+		OutputPins.Add(PinToAdd);
+	}
+}
diff --git a/Engine/Source/Editor/GraphEditor/Private/SGraphNodeKnot.h b/Engine/Source/Editor/GraphEditor/Private/SGraphNodeKnot.h
new file mode 100644
--- /dev/null
+++ b/Engine/Source/Editor/GraphEditor/Private/SGraphNodeKnot.h
@@ -0,0 +1,23 @@
+// Copyright 1998-2014 Epic Games, Inc. All Rights Reserved.
+
+#pragma once
+
+#include "SGraphNodeDefault.h"
+#include "K2Node_Knot.h"
+
+class SGraphNodeKnot : public SGraphNodeDefault
+{
+public:
+	SLATE_BEGIN_ARGS(SGraphNodeKnot) {}
+	SLATE_END_ARGS()
+
+	void Construct(const FArguments& InArgs, class UK2Node_Knot* InKnot);
+
+	// SGraphNode interface
+	virtual void UpdateGraphNode() OVERRIDE;
+	virtual const FSlateBrush* GetShadowBrush(bool bSelected) const OVERRIDE;
+	virtual TSharedPtr<SGraphPin> CreatePinWidget(UEdGraphPin* Pin) const OVERRIDE;
+	virtual void AddPin(const TSharedRef<SGraphPin>& PinToAdd) OVERRIDE;
+	virtual void RequestRenameOnSpawn() OVERRIDE { }
+	// End of SGraphNode interface
+};
diff --git a/Engine/Source/Editor/GraphEditor/Private/SGraphPanel.cpp b/Engine/Source/Editor/GraphEditor/Private/SGraphPanel.cpp
--- a/Engine/Source/Editor/GraphEditor/Private/SGraphPanel.cpp
+++ b/Engine/Source/Editor/GraphEditor/Private/SGraphPanel.cpp
@@ -636,7 +636,7 @@ TSharedPtr<SWidget> SGraphPanel::OnSummonContextMenu(const FGeometry& MyGeometry
 		const FVector2D NodeAddPosition = PanelCoordToGraphCoord( MyGeometry.AbsoluteToLocal(MouseEvent.GetScreenSpacePosition()) );
 		TArray<UEdGraphPin*> NoSourcePins;
 
-		return SummonContextMenu( MouseEvent.GetScreenSpacePosition(), NodeAddPosition, NodeUnderCursor, PinUnderCursor, NoSourcePins);
+		return SummonContextMenu(MouseEvent.GetScreenSpacePosition(), NodeAddPosition, NodeUnderCursor, PinUnderCursor, NoSourcePins, MouseEvent.IsShiftDown());
 	}
 
 	return TSharedPtr<SWidget>();
@@ -798,11 +798,11 @@ FReply SGraphPanel::OnDrop( const FGeometry& MyGeometry, const FDragDropEvent& D
 	}
 }
 
-void SGraphPanel::OnBeginMakingConnection( const TSharedRef<SGraphPin>& InOriginatingPin )
+void SGraphPanel::OnBeginMakingConnection(UEdGraphPin* InOriginatingPin)
 {
-	if (auto* PinObj = InOriginatingPin->GetPinObj())
+	if (InOriginatingPin != nullptr)
 	{
-		PreviewConnectorFromPins.Add(PinObj);
+		PreviewConnectorFromPins.Add(InOriginatingPin);
 	}
 }
 
@@ -849,11 +849,18 @@ void SGraphPanel::RemoveAllNodes()
 	SNodePanel::RemoveAllNodes();
 }
 
-TSharedPtr<SWidget> SGraphPanel::SummonContextMenu( const FVector2D& WhereToSummon, const FVector2D& WhereToAddNode, UEdGraphNode* ForNode, UEdGraphPin* ForPin, const TArray<UEdGraphPin*>& DragFromPins )
+TSharedPtr<SWidget> SGraphPanel::SummonContextMenu(const FVector2D& WhereToSummon, const FVector2D& WhereToAddNode, UEdGraphNode* ForNode, UEdGraphPin* ForPin, const TArray<UEdGraphPin*>& DragFromPins, bool bShiftOperation)
 {
 	if (OnGetContextMenuFor.IsBound())
 	{
-		FActionMenuContent FocusedContent = OnGetContextMenuFor.Execute( WhereToAddNode, ForNode, ForPin, DragFromPins );
+		FGraphContextMenuArguments SpawnInfo;
+		SpawnInfo.NodeAddPosition = WhereToAddNode;
+		SpawnInfo.GraphNode = ForNode;
+		SpawnInfo.GraphPin = ForPin;
+		SpawnInfo.DragFromPins = DragFromPins;
+		SpawnInfo.bShiftOperation = bShiftOperation;
+
+		FActionMenuContent FocusedContent = OnGetContextMenuFor.Execute(SpawnInfo);
 
 		TSharedRef<SWidget> MenuContent =
 			SNew( SBorder )
@@ -884,13 +891,13 @@ void SGraphPanel::AttachGraphEvents(TSharedPtr<SGraphNode> CreatedSubNode)
 	CreatedSubNode->SetTextCommittedEvent(OnTextCommitted);
 }
 
-void SGraphPanel::AddNode (UEdGraphNode* Node)
+void SGraphPanel::AddNode(UEdGraphNode* Node)
 {
 	TSharedPtr<SGraphNode> NewNode = FNodeFactory::CreateNodeWidget(Node);
 	check(NewNode.IsValid());
 
 	FEdGraphEditAction* GraphAction = UserAddedNodes.Find(Node);
-	bool bWasUserAdded = GraphAction? true : false;
+	const bool bWasUserAdded = (GraphAction != nullptr) ? true : false;
 
 	NewNode->SetIsEditable(IsEditable);
 	NewNode->SetDoubleClickEvent(OnNodeDoubleClicked);
@@ -911,13 +918,13 @@ void SGraphPanel::AddNode (UEdGraphNode* Node)
 		NewNode->PlaySpawnEffect();
 
 		// Do not Select nodes unless they are marked for selection
-		if((GraphAction->Action & GRAPHACTION_SelectNode) != 0)
+		if ((GraphAction->Action & GRAPHACTION_SelectNode) != 0)
 		{
 			SelectAndCenterObject(Node, false);
 		}
 
 		NewNode->UpdateGraphNode();
-		NewNode->RequestRename();
+		NewNode->RequestRenameOnSpawn();
 	}
 	else
 	{
diff --git a/Engine/Source/Editor/GraphEditor/Private/SGraphPanel.h b/Engine/Source/Editor/GraphEditor/Private/SGraphPanel.h
--- a/Engine/Source/Editor/GraphEditor/Private/SGraphPanel.h
+++ b/Engine/Source/Editor/GraphEditor/Private/SGraphPanel.h
@@ -10,14 +10,30 @@ class UEdGraphNode;
 
 DECLARE_DELEGATE( FOnUpdateGraphPanel )
 
+// Arguments when the graph panel wants to open a context menu
+struct FGraphContextMenuArguments
+{
+	// The endpoint of the drag or the location of the right-click
+	FVector2D NodeAddPosition;
+
+	// The source node if there are any
+	UEdGraphNode* GraphNode;
+
+	// The source pin if there is one
+	UEdGraphPin* GraphPin;
+
+	// 
+	TArray<UEdGraphPin*> DragFromPins;
+
+	// Was shift held down when the operation started?
+	bool bShiftOperation;
+};
+
+
 class GRAPHEDITOR_API SGraphPanel : public SNodePanel, public FGCObject
 {
 public:
-	DECLARE_DELEGATE_RetVal_FourParams(
-		FActionMenuContent,
-		FOnGetContextMenuFor,
-		const FVector2D&, class UEdGraphNode*, class UEdGraphPin*, const TArray<class UEdGraphPin*>& )
-
+	DECLARE_DELEGATE_RetVal_OneParam(FActionMenuContent, FOnGetContextMenuFor, const FGraphContextMenuArguments& /*SpawnInfo*/)
 
 	SLATE_BEGIN_ARGS( SGraphPanel )
 		: _OnGetContextMenuFor()
@@ -97,9 +113,9 @@ public:
 	// End of FGCObject interface.
 
 	void ArrangeChildrenForContextMenuSummon(const FGeometry& AllottedGeometry, FArrangedChildren& ArrangedChildren) const;
-	TSharedPtr<SWidget> SummonContextMenu( const FVector2D& WhereToSummon, const FVector2D& WhereToAddNode, UEdGraphNode* ForNode, UEdGraphPin* ForPin, const TArray<UEdGraphPin*>& DragFromPins );
+	TSharedPtr<SWidget> SummonContextMenu(const FVector2D& WhereToSummon, const FVector2D& WhereToAddNode, UEdGraphNode* ForNode, UEdGraphPin* ForPin, const TArray<UEdGraphPin*>& DragFromPins, bool bShiftOperation);
 
-	void OnBeginMakingConnection( const TSharedRef<SGraphPin>& InOriginatingPin );
+	void OnBeginMakingConnection(UEdGraphPin* InOriginatingPin);
 	void OnStopMakingConnection(bool bForceStop = false);
 	void PreservePinPreviewUntilForced();
 
diff --git a/Engine/Source/Editor/GraphEditor/Private/SGraphPin.cpp b/Engine/Source/Editor/GraphEditor/Private/SGraphPin.cpp
--- a/Engine/Source/Editor/GraphEditor/Private/SGraphPin.cpp
+++ b/Engine/Source/Editor/GraphEditor/Private/SGraphPin.cpp
@@ -9,6 +9,64 @@
 #include "SLevelOfDetailBranchNode.h"
 #include "ScopedTransaction.h"
 
+/////////////////////////////////////////////////////
+// FKnotNetCollector
+
+struct FKnotNetCollector
+{
+	TSet<UEdGraphNode*> VisitedNodes;
+	TSet<UEdGraphPin*> VisitedPins;
+
+	FKnotNetCollector(UEdGraphPin* StartingPin)
+	{
+		TraversePin(StartingPin);
+	}
+
+	void TraversePin(UEdGraphPin* Pin)
+	{
+		if (UK2Node_Knot* Knot = Cast<UK2Node_Knot>(Pin->GetOwningNode()))
+		{
+			TraverseNodes(Knot);
+		}
+		else
+		{
+			VisitedPins.Add(Pin);
+
+			for (UEdGraphPin* OtherPin : Pin->LinkedTo)
+			{
+				UEdGraphNode* OtherNode = OtherPin->GetOwningNode();
+				if (OtherPin->GetOwningNode()->IsA(UK2Node_Knot::StaticClass()))
+				{
+					TraverseNodes(OtherNode);
+				}
+			}
+		}
+	}
+
+	void TraverseNodes(UEdGraphNode* Node)
+	{
+		if (VisitedNodes.Contains(Node))
+		{
+			return;
+		}
+		VisitedNodes.Add(Node);
+
+		for (UEdGraphPin* MyPin : Node->Pins)
+		{
+			VisitedPins.Add(MyPin);
+
+			for (UEdGraphPin* OtherPin : MyPin->LinkedTo)
+			{
+				UEdGraphNode* OtherNode = OtherPin->GetOwningNode();
+				if (OtherNode->IsA(UK2Node_Knot::StaticClass()))
+				{
+					TraverseNodes(OtherNode);
+				}
+			}
+		}
+	}
+};
+
 /////////////////////////////////////////////////////
 // SGraphPin
 
@@ -167,17 +225,22 @@ void SGraphPin::Construct(const FArguments& InArgs, UEdGraphPin* InPin)
 				LabelWidget
 			];
 
-		LabelAndValue->AddSlot()
-			.Padding( bIsInput ? FMargin(5,0,0,0) : FMargin(0,0,5,0) )
-			.VAlign(VAlign_Center)
-			[
-				SNew(SBox)
-				.Padding(0.0f)
-				.IsEnabled(this, &SGraphPin::IsEditingEnabled)
+		TSharedRef<SWidget> ValueWidget = GetDefaultValueWidget();
+
+		if (ValueWidget != SNullWidget::NullWidget)
+		{
+			LabelAndValue->AddSlot()
+				.Padding(bIsInput ? FMargin(InArgs._SideToSideMargin, 0, 0, 0) : FMargin(0, 0, InArgs._SideToSideMargin, 0))
+				.VAlign(VAlign_Center)
 				[
-					GetDefaultValueWidget()
-				]	
-			];
+					SNew(SBox)
+					.Padding(0.0f)
+					.IsEnabled(this, &SGraphPin::IsEditingEnabled)
+					[
+						ValueWidget
+					]
+				];
+		}
 
 		LabelAndValue->AddSlot()
 			.VAlign(VAlign_Center)
@@ -201,7 +264,7 @@ void SGraphPin::Construct(const FArguments& InArgs, UEdGraphPin* InPin)
 			+SHorizontalBox::Slot()
 			.AutoWidth()
 			.VAlign(VAlign_Center)
-			.Padding( 0,0,5,0 )
+			.Padding(0, 0, InArgs._SideToSideMargin, 0)
 			[
 				ActualPinWidget
 			]
@@ -226,7 +289,7 @@ void SGraphPin::Construct(const FArguments& InArgs, UEdGraphPin* InPin)
 			+SHorizontalBox::Slot()
 			.AutoWidth()
 			.VAlign(VAlign_Center)
-			.Padding( 5,0,0,0 )
+			.Padding(InArgs._SideToSideMargin, 0, 0, 0)
 			[
 				ActualPinWidget
 			];
@@ -379,7 +442,7 @@ FReply SGraphPin::OnPinMouseDown( const FGeometry& SenderGeometry, const FPointe
 				{
 					bIsMovingLinks = true;
 
-					return FReply::Handled().BeginDragDrop(FDragConnection::New(OwnerPanelPtr.ToSharedRef(), PinArray, true));
+					return FReply::Handled().BeginDragDrop(SpawnPinDragEvent(OwnerPanelPtr.ToSharedRef(), PinArray, /*bIsShiftOperation=*/ false));
 				}
 				else
 				{
@@ -387,13 +450,14 @@ FReply SGraphPin::OnPinMouseDown( const FGeometry& SenderGeometry, const FPointe
 					return FReply::Handled();
 				}
 			}
-			else if (MouseEvent.IsShiftDown())
-			{
-				return FReply::Handled();
-			}
 			
 			// Start a drag-drop on the pin
-			return FReply::Handled().BeginDragDrop(FDragConnection::New(this->OwnerNodePtr.Pin()->GetOwnerPanel().ToSharedRef(), SharedThis(this)));
+			{
+				TArray<TSharedRef<SGraphPin>> PinArray;
+				PinArray.Add(SharedThis(this));
+
+				return FReply::Handled().BeginDragDrop(SpawnPinDragEvent(this->OwnerNodePtr.Pin()->GetOwnerPanel().ToSharedRef(), PinArray, MouseEvent.IsShiftDown()));
+			}
 		}
 		else
 		{
@@ -450,6 +514,10 @@ TOptional<EMouseCursor::Type> SGraphPin::GetPinCursor() const
 	}
 }
 
+TSharedRef<FDragDropOperation> SGraphPin::SpawnPinDragEvent(const TSharedRef<SGraphPanel>& InGraphPanel, const TArray< TSharedRef<SGraphPin> >& InStartingPins, bool bShiftOperation)
+{
+	return FDragConnection::New(InGraphPanel, InStartingPins, bShiftOperation);
+}
 
 /**
  * The system calls this method to notify the widget that a mouse button was release within it. This event is bubbled.
@@ -487,14 +555,32 @@ FReply SGraphPin::OnMouseButtonUp( const FGeometry& MyGeometry, const FPointerEv
 
 void SGraphPin::OnMouseEnter( const FGeometry& MyGeometry, const FPointerEvent& MouseEvent )
 {
-	OwnerNodePtr.Pin()->GetOwnerPanel()->AddPinToHoverSet(GetPinObj());
-	SCompoundWidget::OnMouseEnter( MyGeometry, MouseEvent );
+	FKnotNetCollector NetCollector(GetPinObj());
+
+	TSharedPtr<SGraphPanel> Panel = OwnerNodePtr.Pin()->GetOwnerPanel();
+	for (UEdGraphPin* PinInNet : NetCollector.VisitedPins)
+	{
+		Panel->AddPinToHoverSet(PinInNet);
+		HoverPinSet.Add(PinInNet);
+	}
+
+	SCompoundWidget::OnMouseEnter(MyGeometry, MouseEvent);
 }
 
 void SGraphPin::OnMouseLeave( const FPointerEvent& MouseEvent )
 {	
-	OwnerNodePtr.Pin()->GetOwnerPanel()->RemovePinFromHoverSet(GetPinObj());
-	SCompoundWidget::OnMouseLeave( MouseEvent );
+	TSharedPtr<SGraphPanel> Panel = OwnerNodePtr.Pin()->GetOwnerPanel();
+
+	for (TWeakObjectPtr<UEdGraphPin> WeakPin : HoverPinSet)
+	{
+		if (UEdGraphPin* PinInNet = WeakPin.Get())
+		{
+			Panel->RemovePinFromHoverSet(PinInNet);
+		}
+	}
+	HoverPinSet.Empty();
+
+	SCompoundWidget::OnMouseLeave(MouseEvent);
 }
 
 void SGraphPin::OnDragEnter( const FGeometry& MyGeometry, const FDragDropEvent& DragDropEvent )
@@ -513,7 +599,7 @@ void SGraphPin::OnDragEnter( const FGeometry& MyGeometry, const FDragDropEvent&
 		{
 			// Inform the Drag and Drop operation that we are hovering over this pin.
 			TSharedPtr<FGraphEditorDragDropAction> DragConnectionOp = StaticCastSharedPtr<FGraphEditorDragDropAction>(Operation);
-			DragConnectionOp->SetHoveredPin( SharedThis(this) );
+			DragConnectionOp->SetHoveredPin(GraphPinObj);
 		}	
 
 		// Pins treat being dragged over the same as being hovered outside of drag and drop if they know how to respond to the drag action.
@@ -551,9 +637,9 @@ void SGraphPin::OnDragLeave( const FDragDropEvent& DragDropEvent )
 	{
 		// Inform the Drag and Drop operation that we are not hovering any pins
 		TSharedPtr<FGraphEditorDragDropAction> DragConnectionOp = StaticCastSharedPtr<FGraphEditorDragDropAction>(Operation);
-		DragConnectionOp->SetHoveredPin( TSharedPtr<SGraphPin>(NULL) );
+		DragConnectionOp->SetHoveredPin(nullptr);
 
-		SBorder::OnMouseLeave( DragDropEvent );
+		SBorder::OnMouseLeave(DragDropEvent);
 	}
 
 	else if (Operation->IsOfType<FAssetDragDropOp>())
diff --git a/Engine/Source/Editor/GraphEditor/Public/GraphEditorDragDropAction.h b/Engine/Source/Editor/GraphEditor/Public/GraphEditorDragDropAction.h
--- a/Engine/Source/Editor/GraphEditor/Public/GraphEditorDragDropAction.h
+++ b/Engine/Source/Editor/GraphEditor/Public/GraphEditorDragDropAction.h
@@ -8,7 +8,7 @@ class GRAPHEDITOR_API FGraphEditorDragDropAction : public FDragDropOperation
 public:
 	DRAG_DROP_OPERATOR_TYPE(FGraphEditorDragDropAction, FDragDropOperation)
 
-	void SetHoveredPin(const TSharedPtr<class SGraphPin>& InPin);
+	void SetHoveredPin(class UEdGraphPin* InPin);
 	void SetHoveredNode(const TSharedPtr<class SGraphNode>& InNode);
 	void SetHoveredGraph(const TSharedPtr<class SGraphPanel>& InGraph);
 	void SetHoveredCategoryName(const FString& InHoverCategoryName);
@@ -37,7 +37,7 @@ protected:
 
 private:
 	// The pin that the drag action is currently hovering over
-	TSharedPtr<class SGraphPin> HoveredPin;
+	TWeakObjectPtr<class UEdGraphPin> HoveredPin;
 
 	// The node that the drag action is currently hovering over
 	TSharedPtr<class SGraphNode> HoveredNode;
diff --git a/Engine/Source/Editor/GraphEditor/Public/SGraphNode.h b/Engine/Source/Editor/GraphEditor/Public/SGraphNode.h
--- a/Engine/Source/Editor/GraphEditor/Public/SGraphNode.h
+++ b/Engine/Source/Editor/GraphEditor/Public/SGraphNode.h
@@ -167,8 +167,16 @@ public:
 	void SetDisallowedPinConnectionEvent(SGraphEditor::FOnDisallowedPinConnection InOnDisallowedPinConnection);
 	/** returns true if a rename is pending on this node */
 	bool IsRenamePending() const { return bRenameIsPending; }
+
+	/** Requests a rename when the node was initially spawned */
+	virtual void RequestRenameOnSpawn()
+	{
+		RequestRename();
+	}
+
 	/** flags node as rename pending if supported */
 	void RequestRename();
+
 	/** Sets node into rename state if supported */
 	void ApplyRename();
 
diff --git a/Engine/Source/Editor/GraphEditor/Public/SGraphPin.h b/Engine/Source/Editor/GraphEditor/Public/SGraphPin.h
--- a/Engine/Source/Editor/GraphEditor/Public/SGraphPin.h
+++ b/Engine/Source/Editor/GraphEditor/Public/SGraphPin.h
@@ -16,10 +16,12 @@ public:
 		: _PinLabelStyle(NAME_DefaultPinLabelStyle)
 		, _HasToolTip(true)
 		, _UsePinColorForText(false)
+		, _SideToSideMargin(5.0f)
 		{}
 		SLATE_ARGUMENT(FName, PinLabelStyle)
 		SLATE_ARGUMENT(bool, HasToolTip)
 		SLATE_ARGUMENT(bool, UsePinColorForText)
+		SLATE_ARGUMENT(float, SideToSideMargin)
 	SLATE_END_ARGS()
 
 	/**
@@ -36,7 +38,7 @@ public:
 	void SetIsEditable(TAttribute<bool> InIsEditable);
 
 	/** Handle clicking on the pin */
-	FReply OnPinMouseDown(const FGeometry& SenderGeometry, const FPointerEvent& MouseEvent);
+	virtual FReply OnPinMouseDown(const FGeometry& SenderGeometry, const FPointerEvent& MouseEvent);
 
 	/** Handle clicking on the pin name */
 	FReply OnPinNameMouseDown(const FGeometry& SenderGeometry, const FPointerEvent& MouseEvent);
@@ -140,6 +142,8 @@ protected:
 
 	TOptional<EMouseCursor::Type> GetPinCursor() const;
 
+	/** Spawns a FDragConnection or similar class for the pin drag event */
+	virtual TSharedRef<FDragDropOperation> SpawnPinDragEvent(const TSharedRef<class SGraphPanel>& InGraphPanel, const TArray< TSharedRef<SGraphPin> >& InStartingPins, bool bShiftOperation);
 private:
 	
 	/** True if pin can be edited */
@@ -176,6 +180,8 @@ protected:
 	/** Cached offset from owning node to approximate position of culled pins */
 	FVector2D CachedNodeOffset;
 
+	TSet< TWeakObjectPtr<UEdGraphPin> > HoverPinSet;
+
 	//@TODO: Want to cache these once for all SGraphPins, but still handle slate style updates
 	mutable const FSlateBrush* CachedImg_ArrayPin_ConnectedHovered;
 	mutable const FSlateBrush* CachedImg_ArrayPin_Connected;
