Updated SNodePanel::OnMouseMove() so it does zoom on Mac trackpads and Magic Mouse when right mouse button is down.

diff --git a/Engine/Source/Editor/GraphEditor/Private/SNodePanel.cpp b/Engine/Source/Editor/GraphEditor/Private/SNodePanel.cpp
--- a/Engine/Source/Editor/GraphEditor/Private/SNodePanel.cpp
+++ b/Engine/Source/Editor/GraphEditor/Private/SNodePanel.cpp
@@ -599,7 +599,7 @@ FReply SNodePanel::OnMouseMove(const FGeometry& MyGeometry, const FPointerEvent&
 		// Track how much the mouse moved since the mouse down.
 		TotalMouseDelta += CursorDelta.Size();
 
-		const bool bShouldZoom = (bIsLeftMouseButtonDown && bIsRightMouseButtonDown) || (bIsRightMouseButtonDown && ModifierKeysState.IsAltDown());
+		const bool bShouldZoom = bIsRightMouseButtonDown && (bIsLeftMouseButtonDown || ModifierKeysState.IsAltDown() || FSlateApplication::Get().IsUsingTrackpad());
 		if (bShouldZoom)
 		{
 			FReply ReplyState = FReply::Handled();
