Settings: moved graph editor settings into GraphEditor module; added 'Graph Editor' settings section to Settings UI

Settings specific to Anim, Blueprint and Material Editor should probably be broken out into their own settings classes in the respective modules at some point. It feels a little strange that the GraphEditor needs to be aware of specific pin and node title types. At this time I did not want to refactor the existing graph editors though, and since many of the properties are currently being shared, pulling these settings out of UEditorUserSettings is at least a good first step. In the future I would like to see settings sections for specific graph based editors instead of having a generic 'Graph Editor' section that is shared by some or all editors.

I exposed various color settings in the UI as many users seem to be quite excited about customizing their Editor. If you think that some settings should be added or removed, please feel free to make any desired changes or let me know. Thanks!

diff --git a/Engine/Source/Editor/AnimGraph/Private/AnimGraphNode_CustomTransitionResult.cpp b/Engine/Source/Editor/AnimGraph/Private/AnimGraphNode_CustomTransitionResult.cpp
--- a/Engine/Source/Editor/AnimGraph/Private/AnimGraphNode_CustomTransitionResult.cpp
+++ b/Engine/Source/Editor/AnimGraph/Private/AnimGraphNode_CustomTransitionResult.cpp
@@ -2,6 +2,8 @@
 
 #include "AnimGraphPrivatePCH.h"
 #include "AnimGraphNode_CustomTransitionResult.h"
+#include "GraphEditorSettings.h"
+
 
 /////////////////////////////////////////////////////
 // UAnimGraphNode_CustomTransitionResult
@@ -15,7 +17,7 @@ UAnimGraphNode_CustomTransitionResult::UAnimGraphNode_CustomTransitionResult(con
 
 FLinearColor UAnimGraphNode_CustomTransitionResult::GetNodeTitleColor() const
 {
-	return GEditor->AccessEditorUserSettings().ResultNodeTitleColor;
+	return GetDefault<UGraphEditorSettings>()->ResultNodeTitleColor;
 }
 
 FString UAnimGraphNode_CustomTransitionResult::GetTooltip() const
diff --git a/Engine/Source/Editor/AnimGraph/Private/AnimGraphNode_Root.cpp b/Engine/Source/Editor/AnimGraph/Private/AnimGraphNode_Root.cpp
--- a/Engine/Source/Editor/AnimGraph/Private/AnimGraphNode_Root.cpp
+++ b/Engine/Source/Editor/AnimGraph/Private/AnimGraphNode_Root.cpp
@@ -2,6 +2,8 @@
 
 #include "AnimGraphPrivatePCH.h"
 #include "AnimGraphNode_Root.h"
+#include "GraphEditorSettings.h"
+
 
 /////////////////////////////////////////////////////
 // FPoseLinkMappingRecord
@@ -43,7 +45,7 @@ UAnimGraphNode_Root::UAnimGraphNode_Root(const FPostConstructInitializePropertie
 
 FLinearColor UAnimGraphNode_Root::GetNodeTitleColor() const
 {
-	return GEditor->AccessEditorUserSettings().ResultNodeTitleColor;
+	return GetDefault<UGraphEditorSettings>()->ResultNodeTitleColor;
 }
 
 FText UAnimGraphNode_Root::GetNodeTitle(ENodeTitleType::Type TitleType) const
diff --git a/Engine/Source/Editor/AnimGraph/Private/AnimGraphNode_TransitionResult.cpp b/Engine/Source/Editor/AnimGraph/Private/AnimGraphNode_TransitionResult.cpp
--- a/Engine/Source/Editor/AnimGraph/Private/AnimGraphNode_TransitionResult.cpp
+++ b/Engine/Source/Editor/AnimGraph/Private/AnimGraphNode_TransitionResult.cpp
@@ -3,8 +3,10 @@
 #include "AnimGraphPrivatePCH.h"
 #include "AnimGraphNode_TransitionResult.h"
 #include "GraphEditorActions.h"
+#include "GraphEditorSettings.h"
 #include "ScopedTransaction.h"
 
+
 /////////////////////////////////////////////////////
 // UAnimGraphNode_TransitionResult
 
@@ -17,7 +19,7 @@ UAnimGraphNode_TransitionResult::UAnimGraphNode_TransitionResult(const FPostCons
 
 FLinearColor UAnimGraphNode_TransitionResult::GetNodeTitleColor() const
 {
-	return GEditor->AccessEditorUserSettings().ResultNodeTitleColor;
+	return GetDefault<UGraphEditorSettings>()->ResultNodeTitleColor;
 }
 
 FString UAnimGraphNode_TransitionResult::GetTooltip() const
diff --git a/Engine/Source/Editor/BehaviorTreeEditor/Private/EdGraphSchema_BehaviorTreeDecorator.cpp b/Engine/Source/Editor/BehaviorTreeEditor/Private/EdGraphSchema_BehaviorTreeDecorator.cpp
--- a/Engine/Source/Editor/BehaviorTreeEditor/Private/EdGraphSchema_BehaviorTreeDecorator.cpp
+++ b/Engine/Source/Editor/BehaviorTreeEditor/Private/EdGraphSchema_BehaviorTreeDecorator.cpp
@@ -3,6 +3,7 @@
 #include "BehaviorTreeEditorPrivatePCH.h"
 #include "BlueprintGraphDefinitions.h"
 #include "GraphEditorActions.h"
+#include "GraphEditorSettings.h"
 #include "BehaviorTreeConnectionDrawingPolicy.h"
 #include "ScopedTransaction.h"
 
@@ -269,10 +270,7 @@ const FPinConnectionResponse UEdGraphSchema_BehaviorTreeDecorator::CanCreateConn
 
 FLinearColor UEdGraphSchema_BehaviorTreeDecorator::GetPinTypeColor(const FEdGraphPinType& PinType) const
 {
-	const FString& TypeString = PinType.PinCategory;
-	const UEditorUserSettings& Options = GEditor->AccessEditorUserSettings();
-
-	return Options.BooleanPinTypeColor;
+	return GetDefault<UGraphEditorSettings>()->BooleanPinTypeColor;
 }
 
 bool UEdGraphSchema_BehaviorTreeDecorator::ShouldHidePinDefaultValue(UEdGraphPin* Pin) const
diff --git a/Engine/Source/Editor/BlueprintGraph/Private/BlueprintGraphPrivatePCH.h b/Engine/Source/Editor/BlueprintGraph/Private/BlueprintGraphPrivatePCH.h
--- a/Engine/Source/Editor/BlueprintGraph/Private/BlueprintGraphPrivatePCH.h
+++ b/Engine/Source/Editor/BlueprintGraph/Private/BlueprintGraphPrivatePCH.h
@@ -4,6 +4,7 @@
 #define __BlueprintGraphPrivatePCH_h__
 
 #include "Engine.h"
+#include "GraphEditorSettings.h"
 #include "BlueprintGraphDefinitions.h"
 #include "BlueprintUtilities.h"
 #include "Kismet2/BlueprintEditorUtils.h"
diff --git a/Engine/Source/Editor/BlueprintGraph/Private/EdGraphSchema_K2.cpp b/Engine/Source/Editor/BlueprintGraph/Private/EdGraphSchema_K2.cpp
--- a/Engine/Source/Editor/BlueprintGraph/Private/EdGraphSchema_K2.cpp
+++ b/Engine/Source/Editor/BlueprintGraph/Private/EdGraphSchema_K2.cpp
@@ -8,6 +8,7 @@
 #include "BlueprintGraphPrivatePCH.h"
 
 #include "GraphEditorActions.h"
+#include "GraphEditorSettings.h"
 #include "ScopedTransaction.h"
 #include "ComponentAssetBroker.h"
 
@@ -1921,35 +1922,35 @@ bool UEdGraphSchema_K2::DefaultValueSimpleValidation(const FEdGraphPinType& PinT
 FLinearColor UEdGraphSchema_K2::GetPinTypeColor(const FEdGraphPinType& PinType) const
 {
 	const FString& TypeString = PinType.PinCategory;
-	const UEditorUserSettings& Options = GEditor->AccessEditorUserSettings();
+	const UGraphEditorSettings* Settings = GetDefault<UGraphEditorSettings>();
 
 	if (TypeString == PC_Exec)
 	{
-		return Options.ExecutionPinTypeColor;
+		return Settings->ExecutionPinTypeColor;
 	}
 	else if (TypeString == PC_Object)
 	{
-		return Options.ObjectPinTypeColor;
+		return Settings->ObjectPinTypeColor;
 	}
 	else if (TypeString == PC_Interface)
 	{
-		return Options.InterfacePinTypeColor;
+		return Settings->InterfacePinTypeColor;
 	}
 	else if (TypeString == PC_Float)
 	{
-		return Options.FloatPinTypeColor;
+		return Settings->FloatPinTypeColor;
 	}
 	else if (TypeString == PC_Boolean)
 	{
-		return Options.BooleanPinTypeColor;
+		return Settings->BooleanPinTypeColor;
 	}
 	else if (TypeString == PC_Byte)
 	{
-		return Options.BytePinTypeColor;
+		return Settings->BytePinTypeColor;
 	}
 	else if (TypeString == PC_Int)
 	{
-		return Options.IntPinTypeColor;
+		return Settings->IntPinTypeColor;
 	}
 	else if (TypeString == PC_Struct)
 	{
@@ -1960,57 +1961,57 @@ FLinearColor UEdGraphSchema_K2::GetPinTypeColor(const FEdGraphPinType& PinType)
 		if (PinType.PinSubCategoryObject == VectorStruct)
 		{
 			// vector
-			return Options.VectorPinTypeColor;
+			return Settings->VectorPinTypeColor;
 		}
 		else if (PinType.PinSubCategoryObject == RotatorStruct)
 		{
 			// rotator
-			return Options.RotatorPinTypeColor;
+			return Settings->RotatorPinTypeColor;
 		}
 		else if (PinType.PinSubCategoryObject == TransformStruct)
 		{
 			// transform
-			return Options.TransformPinTypeColor;
+			return Settings->TransformPinTypeColor;
 		}
 		else
 		{
-			return Options.StructPinTypeColor;
+			return Settings->StructPinTypeColor;
 		}
 	}
 	else if (TypeString == PC_String)
 	{
-		return Options.StringPinTypeColor;
+		return Settings->StringPinTypeColor;
 	}
 	else if (TypeString == PC_Text)
 	{
-		return Options.TextPinTypeColor;
+		return Settings->TextPinTypeColor;
 	}
 	else if (TypeString == PC_Wildcard)
 	{
 		if (PinType.PinSubCategory == PSC_Index)
 		{
-			return Options.IndexPinTypeColor;
+			return Settings->IndexPinTypeColor;
 		}
 		else
 		{
-			return Options.WildcardPinTypeColor;
+			return Settings->WildcardPinTypeColor;
 		}
 	}
 	else if (TypeString == PC_Name)
 	{
-		return Options.NamePinTypeColor;
+		return Settings->NamePinTypeColor;
 	}
 	else if (TypeString == PC_Delegate)
 	{
-		return Options.DelegatePinTypeColor;
+		return Settings->DelegatePinTypeColor;
 	}
 	else if (TypeString == PC_Class)
 	{
-		return Options.ClassPinTypeColor;
+		return Settings->ClassPinTypeColor;
 	}
 
 	// Type does not have a defined color!
-	return Options.DefaultPinTypeColor;
+	return Settings->DefaultPinTypeColor;
 }
 
 FString UEdGraphSchema_K2::GetPinDisplayName(const UEdGraphPin* Pin) const 
diff --git a/Engine/Source/Editor/BlueprintGraph/Private/K2Node.cpp b/Engine/Source/Editor/BlueprintGraph/Private/K2Node.cpp
--- a/Engine/Source/Editor/BlueprintGraph/Private/K2Node.cpp
+++ b/Engine/Source/Editor/BlueprintGraph/Private/K2Node.cpp
@@ -519,12 +519,10 @@ FLinearColor UK2Node::GetNodeTitleColor() const
 	// Different color for pure operations
 	if (IsNodePure())
 	{
-		return GEditor->AccessEditorUserSettings().PureFunctionCallNodeTitleColor;
-	}
-	else
-	{
-		return GEditor->AccessEditorUserSettings().FunctionCallNodeTitleColor;
+		return GetDefault<UGraphEditorSettings>()->PureFunctionCallNodeTitleColor;
 	}
+
+	return GetDefault<UGraphEditorSettings>()->FunctionCallNodeTitleColor;
 }
 
 
diff --git a/Engine/Source/Editor/BlueprintGraph/Private/K2Node_CallFunction.cpp b/Engine/Source/Editor/BlueprintGraph/Private/K2Node_CallFunction.cpp
--- a/Engine/Source/Editor/BlueprintGraph/Private/K2Node_CallFunction.cpp
+++ b/Engine/Source/Editor/BlueprintGraph/Private/K2Node_CallFunction.cpp
@@ -762,12 +762,10 @@ FLinearColor UK2Node_CallFunction::GetNodeTitleColor() const
 {
 	if (IsNodePure())
 	{
-		return GEditor->AccessEditorUserSettings().PureFunctionCallNodeTitleColor;
-	}
-	else
-	{
-		return Super::GetNodeTitleColor();
+		return GetDefault<UGraphEditorSettings>()->PureFunctionCallNodeTitleColor;
 	}
+
+	return Super::GetNodeTitleColor();
 }
 
 FString UK2Node_CallFunction::GetTooltip() const
diff --git a/Engine/Source/Editor/BlueprintGraph/Private/K2Node_CallParentFunction.cpp b/Engine/Source/Editor/BlueprintGraph/Private/K2Node_CallParentFunction.cpp
--- a/Engine/Source/Editor/BlueprintGraph/Private/K2Node_CallParentFunction.cpp
+++ b/Engine/Source/Editor/BlueprintGraph/Private/K2Node_CallParentFunction.cpp
@@ -41,7 +41,7 @@ FString UK2Node_CallParentFunction::GetNodeNativeTitle(ENodeTitleType::Type Titl
 
 FLinearColor UK2Node_CallParentFunction::GetNodeTitleColor() const
 {
-	return GEditor->AccessEditorUserSettings().ParentFunctionCallNodeTitleColor;
+	return GetDefault<UGraphEditorSettings>()->ParentFunctionCallNodeTitleColor;
 }
 
 void UK2Node_CallParentFunction::AllocateDefaultPins()
diff --git a/Engine/Source/Editor/BlueprintGraph/Private/K2Node_ClassDynamicCast.cpp b/Engine/Source/Editor/BlueprintGraph/Private/K2Node_ClassDynamicCast.cpp
--- a/Engine/Source/Editor/BlueprintGraph/Private/K2Node_ClassDynamicCast.cpp
+++ b/Engine/Source/Editor/BlueprintGraph/Private/K2Node_ClassDynamicCast.cpp
@@ -49,7 +49,7 @@ void UK2Node_ClassDynamicCast::AllocateDefaultPins()
 FLinearColor UK2Node_ClassDynamicCast::GetNodeTitleColor() const
 {
 	const UEditorUserSettings& Options = GEditor->AccessEditorUserSettings();
-	return Options.ClassPinTypeColor;
+	return GetDefault<UGraphEditorSettings>()->ClassPinTypeColor;
 }
 
 UEdGraphPin* UK2Node_ClassDynamicCast::GetCastSourcePin() const
diff --git a/Engine/Source/Editor/BlueprintGraph/Private/K2Node_Event.cpp b/Engine/Source/Editor/BlueprintGraph/Private/K2Node_Event.cpp
--- a/Engine/Source/Editor/BlueprintGraph/Private/K2Node_Event.cpp
+++ b/Engine/Source/Editor/BlueprintGraph/Private/K2Node_Event.cpp
@@ -47,7 +47,7 @@ FNodeHandlingFunctor* UK2Node_Event::CreateNodeHandler(FKismetCompilerContext& C
 
 FLinearColor UK2Node_Event::GetNodeTitleColor() const
 {
-	return GEditor->AccessEditorUserSettings().EventNodeTitleColor;
+	return GetDefault<UGraphEditorSettings>()->EventNodeTitleColor;
 }
 
 FText UK2Node_Event::GetNodeTitle(ENodeTitleType::Type TitleType) const
diff --git a/Engine/Source/Editor/BlueprintGraph/Private/K2Node_FunctionTerminator.cpp b/Engine/Source/Editor/BlueprintGraph/Private/K2Node_FunctionTerminator.cpp
--- a/Engine/Source/Editor/BlueprintGraph/Private/K2Node_FunctionTerminator.cpp
+++ b/Engine/Source/Editor/BlueprintGraph/Private/K2Node_FunctionTerminator.cpp
@@ -10,5 +10,5 @@ UK2Node_FunctionTerminator::UK2Node_FunctionTerminator(const class FPostConstruc
 
 FLinearColor UK2Node_FunctionTerminator::GetNodeTitleColor() const
 {
-	return GEditor->AccessEditorUserSettings().FunctionTerminatorNodeTitleColor;
+	return GetDefault<UGraphEditorSettings>()->FunctionTerminatorNodeTitleColor;
 }
diff --git a/Engine/Source/Editor/BlueprintGraph/Private/K2Node_IfThenElse.cpp b/Engine/Source/Editor/BlueprintGraph/Private/K2Node_IfThenElse.cpp
--- a/Engine/Source/Editor/BlueprintGraph/Private/K2Node_IfThenElse.cpp
+++ b/Engine/Source/Editor/BlueprintGraph/Private/K2Node_IfThenElse.cpp
@@ -98,7 +98,7 @@ FText UK2Node_IfThenElse::GetNodeTitle(ENodeTitleType::Type TitleType) const
 
 FLinearColor UK2Node_IfThenElse::GetNodeTitleColor() const
 {
-	return GEditor->AccessEditorUserSettings().ExecBranchNodeTitleColor;
+	return GetDefault<UGraphEditorSettings>()->ExecBranchNodeTitleColor;
 }
 
 FString UK2Node_IfThenElse::GetTooltip() const
diff --git a/Engine/Source/Editor/BlueprintGraph/Private/K2Node_InputAction.cpp b/Engine/Source/Editor/BlueprintGraph/Private/K2Node_InputAction.cpp
--- a/Engine/Source/Editor/BlueprintGraph/Private/K2Node_InputAction.cpp
+++ b/Engine/Source/Editor/BlueprintGraph/Private/K2Node_InputAction.cpp
@@ -35,7 +35,7 @@ void UK2Node_InputAction::AllocateDefaultPins()
 
 FLinearColor UK2Node_InputAction::GetNodeTitleColor() const
 {
-	return GEditor->AccessEditorUserSettings().EventNodeTitleColor;
+	return GetDefault<UGraphEditorSettings>()->EventNodeTitleColor;
 }
 
 FText UK2Node_InputAction::GetNodeTitle(ENodeTitleType::Type TitleType) const
diff --git a/Engine/Source/Editor/BlueprintGraph/Private/K2Node_InputKey.cpp b/Engine/Source/Editor/BlueprintGraph/Private/K2Node_InputKey.cpp
--- a/Engine/Source/Editor/BlueprintGraph/Private/K2Node_InputKey.cpp
+++ b/Engine/Source/Editor/BlueprintGraph/Private/K2Node_InputKey.cpp
@@ -35,7 +35,7 @@ void UK2Node_InputKey::AllocateDefaultPins()
 
 FLinearColor UK2Node_InputKey::GetNodeTitleColor() const
 {
-	return GEditor->AccessEditorUserSettings().EventNodeTitleColor;
+	return GetDefault<UGraphEditorSettings>()->EventNodeTitleColor;
 }
 
 FName UK2Node_InputKey::GetModifierName() const
diff --git a/Engine/Source/Editor/BlueprintGraph/Private/K2Node_InputTouch.cpp b/Engine/Source/Editor/BlueprintGraph/Private/K2Node_InputTouch.cpp
--- a/Engine/Source/Editor/BlueprintGraph/Private/K2Node_InputTouch.cpp
+++ b/Engine/Source/Editor/BlueprintGraph/Private/K2Node_InputTouch.cpp
@@ -52,7 +52,7 @@ void UK2Node_InputTouch::AllocateDefaultPins()
 
 FLinearColor UK2Node_InputTouch::GetNodeTitleColor() const
 {
-	return GEditor->AccessEditorUserSettings().EventNodeTitleColor;
+	return GetDefault<UGraphEditorSettings>()->EventNodeTitleColor;
 }
 
 FText UK2Node_InputTouch::GetNodeTitle(ENodeTitleType::Type TitleType) const
diff --git a/Engine/Source/Editor/EditorSettingsViewer/EditorSettingsViewer.Build.cs b/Engine/Source/Editor/EditorSettingsViewer/EditorSettingsViewer.Build.cs
--- a/Engine/Source/Editor/EditorSettingsViewer/EditorSettingsViewer.Build.cs
+++ b/Engine/Source/Editor/EditorSettingsViewer/EditorSettingsViewer.Build.cs
@@ -16,6 +16,7 @@ namespace UnrealBuildTool.Rules
 				new string[] {
 					"CoreUObject",
 					"Engine",
+					"GraphEditor",
 					"InputBindingEditor",
 					"MessageLog",
 					"SettingsEditor",
diff --git a/Engine/Source/Editor/EditorSettingsViewer/Private/EditorSettingsViewerModule.cpp b/Engine/Source/Editor/EditorSettingsViewer/Private/EditorSettingsViewerModule.cpp
--- a/Engine/Source/Editor/EditorSettingsViewer/Private/EditorSettingsViewerModule.cpp
+++ b/Engine/Source/Editor/EditorSettingsViewer/Private/EditorSettingsViewerModule.cpp
@@ -5,7 +5,7 @@
 =============================================================================*/
 
 #include "EditorSettingsViewerPrivatePCH.h"
-#include "InternationalizationSettingsModel.h"
+
 
 #define LOCTEXT_NAMESPACE "FEditorSettingsViewerModule"
 
@@ -188,6 +188,13 @@ protected:
 			LOCTEXT("ContentEditorsDestructableMeshEditorSettingsDescription", "Change the behavior of the Destructable Mesh Editor."),
 			GetMutableDefault<UDestructableMeshEditorSettings>()
 		);*/
+
+		// graph editors
+		SettingsModule.RegisterSettings("Editor", "ContentEditors", "GraphEditor",
+			LOCTEXT("ContentEditorsGraphEditorSettingsName", "Graph Editors"),
+			LOCTEXT("ContentEditorsGraphEditorSettingsDescription", "Customize Anim, Blueprint and Material Editor."),
+			GetMutableDefault<UGraphEditorSettings>()
+		);
 	}
 
 	/**
@@ -216,6 +223,7 @@ protected:
 			// other tools
 			SettingsModule->UnregisterSettings("Editor", "ContentEditors", "ContentBrowser");
 //			SettingsModule->UnregisterSettings("Editor", "ContentEditors", "DestructableMeshEditor");
+			SettingsModule->UnregisterSettings("Editor", "ContentEditors", "GraphEditor");
 		}
 	}
 
diff --git a/Engine/Source/Editor/EditorSettingsViewer/Private/EditorSettingsViewerPrivatePCH.h b/Engine/Source/Editor/EditorSettingsViewer/Private/EditorSettingsViewerPrivatePCH.h
--- a/Engine/Source/Editor/EditorSettingsViewer/Private/EditorSettingsViewerPrivatePCH.h
+++ b/Engine/Source/Editor/EditorSettingsViewer/Private/EditorSettingsViewerPrivatePCH.h
@@ -11,7 +11,9 @@
  *****************************************************************************/
 
 #include "UnrealEd.h"
+#include "GraphEditorSettings.h"
 #include "InputBindingEditor.h"
+#include "InternationalizationSettingsModel.h"
 #include "MessageLog.h"
 #include "Settings.h"
 #include "SettingsEditor.h"
diff --git a/Engine/Source/Editor/GraphEditor/Private/ConnectionDrawingPolicy.cpp b/Engine/Source/Editor/GraphEditor/Private/ConnectionDrawingPolicy.cpp
--- a/Engine/Source/Editor/GraphEditor/Private/ConnectionDrawingPolicy.cpp
+++ b/Engine/Source/Editor/GraphEditor/Private/ConnectionDrawingPolicy.cpp
@@ -528,7 +528,7 @@ FKismetConnectionDrawingPolicy::FKismetConnectionDrawingPolicy(int32 InBackLayer
 	: FConnectionDrawingPolicy(InBackLayerID, InFrontLayerID, ZoomFactor, InClippingRect, InDrawElements)
 	, GraphObj(InGraphObj)
 {
-	const UEditorUserSettings& Options = GEditor->AccessEditorUserSettings();
+	const UGraphEditorSettings* Settings = GetDefault<UGraphEditorSettings>();
 
 	// Don't want to draw ending arrowheads
 	ArrowImage = nullptr;
@@ -542,22 +542,22 @@ FKismetConnectionDrawingPolicy::FKismetConnectionDrawingPolicy(int32 InBackLayer
 	}
 
 	// Cache off the editor options
-	AttackColor = Options.TraceAttackColor;
-	SustainColor = Options.TraceSustainColor;
-	ReleaseColor = Options.TraceReleaseColor;
-
-	AttackWireThickness = Options.TraceAttackWireThickness;
-	SustainWireThickness = Options.TraceSustainWireThickness;
-	ReleaseWireThickness = Options.TraceReleaseWireThickness;
-
-	TracePositionBonusPeriod = Options.TracePositionBonusPeriod;
-	TracePositionExponent = Options.TracePositionExponent;
-	AttackHoldPeriod = Options.TraceAttackHoldPeriod;
-	DecayPeriod = Options.TraceDecayPeriod;
-	DecayExponent = Options.TraceDecayExponent;
-	SustainHoldPeriod = Options.TraceSustainHoldPeriod;
-	ReleasePeriod = Options.TraceReleasePeriod;
-	ReleaseExponent = Options.TraceReleaseExponent;
+	AttackColor = Settings->TraceAttackColor;
+	SustainColor = Settings->TraceSustainColor;
+	ReleaseColor = Settings->TraceReleaseColor;
+
+	AttackWireThickness = Settings->TraceAttackWireThickness;
+	SustainWireThickness = Settings->TraceSustainWireThickness;
+	ReleaseWireThickness = Settings->TraceReleaseWireThickness;
+
+	TracePositionBonusPeriod = Settings->TracePositionBonusPeriod;
+	TracePositionExponent = Settings->TracePositionExponent;
+	AttackHoldPeriod = Settings->TraceAttackHoldPeriod;
+	DecayPeriod = Settings->TraceDecayPeriod;
+	DecayExponent = Settings->TraceDecayExponent;
+	SustainHoldPeriod = Settings->TraceSustainHoldPeriod;
+	ReleasePeriod = Settings->TraceReleasePeriod;
+	ReleaseExponent = Settings->TraceReleaseExponent;
 
 	CurrentTime = 0.0;
 	LatestTimeDiscovered = 0.0;
@@ -1202,12 +1202,13 @@ FSoundCueGraphConnectionDrawingPolicy::FSoundCueGraphConnectionDrawingPolicy(int
 	, GraphObj(InGraphObj)
 {
 	// Cache off the editor options
-	const UEditorUserSettings& Options = GEditor->AccessEditorUserSettings();
-	ActiveColor = Options.TraceAttackColor;
-	InactiveColor = Options.TraceReleaseColor;
+	const UGraphEditorSettings* Settings = GetDefault<UGraphEditorSettings>();
 
-	ActiveWireThickness = Options.TraceAttackWireThickness;
-	InactiveWireThickness = Options.TraceReleaseWireThickness;
+	ActiveColor = Settings->TraceAttackColor;
+	InactiveColor = Settings->TraceReleaseColor;
+
+	ActiveWireThickness = Settings->TraceAttackWireThickness;
+	InactiveWireThickness = Settings->TraceReleaseWireThickness;
 
 	// Don't want to draw ending arrowheads
 	ArrowImage = nullptr;
diff --git a/Engine/Source/Editor/GraphEditor/Private/GraphEditorCommon.h b/Engine/Source/Editor/GraphEditor/Private/GraphEditorCommon.h
--- a/Engine/Source/Editor/GraphEditor/Private/GraphEditorCommon.h
+++ b/Engine/Source/Editor/GraphEditor/Private/GraphEditorCommon.h
@@ -9,6 +9,7 @@
 
 #include "Slate.h"
 
+#include "GraphEditorSettings.h"
 #include "SNodePanel.h"
 #include "SGraphPanel.h"
 #include "SGraphNode.h"
diff --git a/Engine/Source/Editor/GraphEditor/Private/GraphEditorSettings.cpp b/Engine/Source/Editor/GraphEditor/Private/GraphEditorSettings.cpp
new file mode 100644
--- /dev/null
+++ b/Engine/Source/Editor/GraphEditor/Private/GraphEditorSettings.cpp
@@ -0,0 +1,66 @@
+// Copyright 1998-2014 Epic Games, Inc. All Rights Reserved.
+
+/*=============================================================================
+	GraphEditorSettings.cpp: Implements the UGraphEditorSettings class.
+=============================================================================*/
+
+#include "GraphEditorCommon.h"
+
+
+/* UGraphEditorSettings structors
+ *****************************************************************************/
+
+UGraphEditorSettings::UGraphEditorSettings( const class FPostConstructInitializeProperties& PCIP )
+	: Super(PCIP)
+{
+	DataPinStyle = BPST_VariantA;
+
+	// graph node pin type colors
+	DefaultPinTypeColor = FLinearColor(0.750000f, 0.6f, 0.4f, 1.0f);			// light brown
+	ExecutionPinTypeColor = FLinearColor(1.0f, 1.0f, 1.0f, 1.0f);				// white
+	BooleanPinTypeColor = FLinearColor(0.300000f, 0.0f, 0.0f, 1.0f);			// maroon
+	BytePinTypeColor = FLinearColor(0.0f, 0.160000f, 0.131270f, 1.0f);			// dark green
+	ClassPinTypeColor = FLinearColor(0.1f, 0.0f, 0.5f, 1.0f);					// deep purple (violet)
+	IntPinTypeColor = FLinearColor(0.013575f, 0.770000f, 0.429609f, 1.0f);		// green-blue
+	FloatPinTypeColor = FLinearColor(0.357667f, 1.0f, 0.060000f, 1.0f);			// bright green
+	NamePinTypeColor = FLinearColor(0.607717f, 0.224984f, 1.0f, 1.0f);			// lilac
+	DelegatePinTypeColor = FLinearColor(1.0f, 0.04f, 0.04f, 1.0f);				// bright red
+	ObjectPinTypeColor = FLinearColor(0.0f, 0.4f, 0.910000f, 1.0f);				// sharp blue
+	InterfacePinTypeColor = FLinearColor(0.8784f, 1.0f, 0.4f, 1.0f);			// pale green
+	StringPinTypeColor = FLinearColor(1.0f, 0.0f, 0.660537f, 1.0f);				// bright pink
+	TextPinTypeColor = FLinearColor(0.8f, 0.2f, 0.4f, 1.0f);					// salmon (light pink)
+	StructPinTypeColor = FLinearColor(0.0f, 0.1f, 0.6f, 1.0f);					// deep blue
+	WildcardPinTypeColor = FLinearColor(0.220000f, 0.195800f, 0.195800f, 1.0f);	// dark gray
+	VectorPinTypeColor = FLinearColor(1.0f, 0.591255f, 0.016512f, 1.0f);		// yellow
+	RotatorPinTypeColor = FLinearColor(0.353393f, 0.454175f, 1.0f, 1.0f);		// periwinkle 
+	TransformPinTypeColor = FLinearColor(1.0f, 0.172585f, 0.0f, 1.0f);			// orange
+	IndexPinTypeColor = FLinearColor(0.013575f, 0.770000f, 0.429609f, 1.0f);	// green-blue
+
+	// graph debugging visuals
+	TraceAttackColor = FLinearColor(1.0f, 0.05f, 0.0f, 1.0f);
+	TraceAttackWireThickness = 12.0f;
+	TraceSustainColor = FLinearColor(1.0f, 0.7f, 0.4f, 1.0f);
+	TraceSustainWireThickness = 8.0f;
+	TraceReleaseColor = FLinearColor(0.5f, 0.5f, 0.5f, 1.0f);
+	TraceReleaseWireThickness = 2.0f;
+
+	// graph debugging exec curve constants
+	TracePositionBonusPeriod = 0.5f;
+	TracePositionExponent = 5.0f;
+	TraceAttackHoldPeriod = 0.3f;
+	TraceDecayPeriod = 0.4f;
+	TraceDecayExponent = 1.8f;
+	TraceSustainHoldPeriod = 0.4f;
+	TraceReleasePeriod = 1.5f;
+	TraceReleaseExponent = 1.4f;
+
+	// Blueprint editor graph node title colors
+	EventNodeTitleColor = FLinearColor(1.f, 0.0f, 0.0f, 1.0f);
+	FunctionCallNodeTitleColor = FLinearColor(0.190525f, 0.583898f, 1.0f, 1.0f);
+	PureFunctionCallNodeTitleColor = FLinearColor(0.4f, 0.850000f, 0.350000f, 1.0f);
+	ParentFunctionCallNodeTitleColor = FLinearColor(1.0f, 0.170000f, 0.0f, 1.0f);
+	FunctionTerminatorNodeTitleColor = FLinearColor(0.6f, 0.0f, 1.0f, 1.0f);
+	ExecBranchNodeTitleColor = FLinearColor(1.0f, 1.0f, 1.0f, 1.0f);
+	ExecSequenceNodeTitleColor= FLinearColor(0.8f, 0.4f, 0.4f, 1.0f);
+	ResultNodeTitleColor = FLinearColor(1.0f, 0.65f, 0.4f, 1.0f);
+}
diff --git a/Engine/Source/Editor/GraphEditor/Private/SGraphPin.cpp b/Engine/Source/Editor/GraphEditor/Private/SGraphPin.cpp
--- a/Engine/Source/Editor/GraphEditor/Private/SGraphPin.cpp
+++ b/Engine/Source/Editor/GraphEditor/Private/SGraphPin.cpp
@@ -112,7 +112,8 @@ SGraphPin::SGraphPin()
 	static const FName NAME_Pin_Background("Graph.Pin.Background");
 	static const FName NAME_Pin_BackgroundHovered("Graph.Pin.BackgroundHovered");
 
-	const EBlueprintPinStyleType StyleType = GEditor->GetEditorUserSettings().DataPinStyle;
+	const EBlueprintPinStyleType StyleType = GetDefault<UGraphEditorSettings>()->DataPinStyle;
+
 	switch(StyleType)
 	{
 	case BPST_VariantA:
diff --git a/Engine/Source/Editor/GraphEditor/Public/GraphEditorSettings.h b/Engine/Source/Editor/GraphEditor/Public/GraphEditorSettings.h
new file mode 100644
--- /dev/null
+++ b/Engine/Source/Editor/GraphEditor/Public/GraphEditorSettings.h
@@ -0,0 +1,187 @@
+// Copyright 1998-2014 Epic Games, Inc. All Rights Reserved.
+
+/*=============================================================================
+	GraphEditorSettings.h: Declares the UGraphEditorSettings class.
+=============================================================================*/
+
+#pragma once
+
+#include "GraphEditorSettings.generated.h"
+
+
+/**
+ * Implements settings for the graph editor.
+ */
+UCLASS(config=EditorUserSettings)
+class GRAPHEDITOR_API UGraphEditorSettings
+	:	public UObject
+{
+	GENERATED_UCLASS_BODY()
+
+public:
+
+	/** The visual styling to use for graph editor pins (in Blueprints, materials, etc...) */
+	UPROPERTY(config, EditAnywhere, Category=Blueprint)
+	TEnumAsByte<EBlueprintPinStyleType> DataPinStyle;
+
+public:
+
+	/** The default color is used only for types not specifically defined below.  Generally if it's seen, it means another type needs to be defined so that the wire in question can have an appropriate color. */
+	UPROPERTY(EditAnywhere,  config, Category=PinColors)
+	FLinearColor DefaultPinTypeColor;
+
+	/** Execution pin type color */
+	UPROPERTY(EditAnywhere,  config, Category=PinColors)
+	FLinearColor ExecutionPinTypeColor;
+
+	/** Boolean pin type color */
+	UPROPERTY(EditAnywhere,  config, Category=PinColors)
+	FLinearColor BooleanPinTypeColor;
+
+	/** Byte pin type color */
+	UPROPERTY(EditAnywhere,  config, Category=PinColors)
+	FLinearColor BytePinTypeColor;
+
+	/** Class pin type color */
+	UPROPERTY(EditAnywhere,  config, Category=PinColors)
+	FLinearColor ClassPinTypeColor;
+
+	/** Integer pin type color */
+	UPROPERTY(EditAnywhere,  config, Category=PinColors)
+	FLinearColor IntPinTypeColor;
+
+	/** Floating-point pin type color */
+	UPROPERTY(EditAnywhere,  config, Category=PinColors)
+	FLinearColor FloatPinTypeColor;
+
+	/** Name pin type color */
+	UPROPERTY(EditAnywhere,  config, Category=PinColors)
+	FLinearColor NamePinTypeColor;
+
+	/** Delegate pin type color */
+	UPROPERTY(EditAnywhere,  config, Category=PinColors)
+	FLinearColor DelegatePinTypeColor;
+
+	/** Object pin type color */
+	UPROPERTY(EditAnywhere,  config, Category=PinColors)
+	FLinearColor ObjectPinTypeColor;
+
+	/** Interface pin type color */
+	UPROPERTY(EditAnywhere,  config, Category=PinColors)
+	FLinearColor InterfacePinTypeColor;
+
+	/** String pin type color */
+	UPROPERTY(EditAnywhere,  config, Category=PinColors)
+	FLinearColor StringPinTypeColor;
+
+	/** Text pin type color */
+	UPROPERTY(EditAnywhere,  config, Category=PinColors)
+	FLinearColor TextPinTypeColor;
+
+	/** Struct pin type color */
+	UPROPERTY(EditAnywhere,  config, Category=PinColors)
+	FLinearColor StructPinTypeColor;
+
+	/** Wildcard pin type color */
+	UPROPERTY(EditAnywhere,  config, Category=PinColors)
+	FLinearColor WildcardPinTypeColor;
+
+	/** Vector pin type color */
+	UPROPERTY(EditAnywhere,  config, Category=PinColors)
+	FLinearColor VectorPinTypeColor;
+
+	/** Rotator pin type color */
+	UPROPERTY(EditAnywhere,  config, Category=PinColors)
+	FLinearColor RotatorPinTypeColor;
+
+	/** Transform pin type color */
+	UPROPERTY(EditAnywhere,  config, Category=PinColors)
+	FLinearColor TransformPinTypeColor;
+
+	/** Index pin type color */
+	UPROPERTY(EditAnywhere,  config, Category=PinColors)
+	FLinearColor IndexPinTypeColor;
+
+public:
+
+	/** Event node title color */
+	UPROPERTY(EditAnywhere,  config, Category=NodeTitleColors)
+	FLinearColor EventNodeTitleColor;
+
+	/** CallFunction node title color */
+	UPROPERTY(EditAnywhere,  config, Category=NodeTitleColors)
+	FLinearColor FunctionCallNodeTitleColor;
+
+	/** Pure function call node title color */
+	UPROPERTY(EditAnywhere,  config, Category=NodeTitleColors)
+	FLinearColor PureFunctionCallNodeTitleColor;
+
+	/** Parent class function call node title color */
+	UPROPERTY(EditAnywhere,  config, Category=NodeTitleColors)
+	FLinearColor ParentFunctionCallNodeTitleColor;
+
+	/** Function Terminator node title color */
+	UPROPERTY(EditAnywhere,  config, Category=NodeTitleColors)
+	FLinearColor FunctionTerminatorNodeTitleColor;
+
+	/** Exec Branch node title color */
+	UPROPERTY(EditAnywhere,  config, Category=NodeTitleColors)
+	FLinearColor ExecBranchNodeTitleColor;
+
+	/** Exec Sequence node title color */
+	UPROPERTY(EditAnywhere,  config, Category=NodeTitleColors)
+	FLinearColor ExecSequenceNodeTitleColor;
+
+	/** Result node title color */
+	UPROPERTY(EditAnywhere,  config, Category=NodeTitleColors)
+	FLinearColor ResultNodeTitleColor;
+
+public:
+
+	UPROPERTY(EditAnywhere,  config, Category=Tracing)
+	FLinearColor TraceAttackColor;
+
+	UPROPERTY(EditAnywhere,  config, Category=Tracing)
+	float TraceAttackWireThickness;
+
+	/** How long is the attack color fully visible */
+	UPROPERTY()
+	float TraceAttackHoldPeriod;
+
+	/** How long does it take to fade from the attack to the sustain color */
+	UPROPERTY()
+	float TraceDecayPeriod;
+
+	UPROPERTY()
+	float TraceDecayExponent;
+
+	UPROPERTY(EditAnywhere,  config, Category=Tracing)
+	FLinearColor TraceSustainColor;
+
+	UPROPERTY(EditAnywhere,  config, Category=Tracing)
+	float TraceSustainWireThickness;
+
+	/** How long is the sustain color fully visible */
+	UPROPERTY()
+	float TraceSustainHoldPeriod;
+
+	UPROPERTY(EditAnywhere,  config, Category=Tracing)
+	FLinearColor TraceReleaseColor;
+
+	UPROPERTY(EditAnywhere,  config, Category=Tracing)
+	float TraceReleaseWireThickness;
+
+	/** How long does it take to fade from the sustain to the release color */
+	UPROPERTY()
+	float TraceReleasePeriod;
+
+	UPROPERTY()
+	float TraceReleaseExponent;
+
+	/** How much of a bonus does an exec get for being near the top of the trace stack, and how does that fall off with position? */
+	UPROPERTY()
+	float TracePositionBonusPeriod;
+
+	UPROPERTY()
+	float TracePositionExponent;
+};
diff --git a/Engine/Source/Editor/UnrealEd/Classes/Editor/EditorUserSettings.h b/Engine/Source/Editor/UnrealEd/Classes/Editor/EditorUserSettings.h
--- a/Engine/Source/Editor/UnrealEd/Classes/Editor/EditorUserSettings.h
+++ b/Engine/Source/Editor/UnrealEd/Classes/Editor/EditorUserSettings.h
@@ -62,162 +62,6 @@ class UEditorUserSettings : public UObject
 	UPROPERTY(config)
 	int32 PropertyMatrix_NumberOfPasteOperationsBeforeWarning;
 
-	/** The visual styling to use for graph editor pins (in Blueprints, materials, etc...) */
-	UPROPERTY(config, EditAnywhere, Category = Blueprint)
-	TEnumAsByte<EBlueprintPinStyleType> DataPinStyle;
-
-	/**
-	* Blueprint editor settings; These are not config/visible/editable unless we need to tweak them.  Make ones you want to adjust 'config, EditAnywhere, Category=Blueprint'
-	*/
-
-	/** The default color is used only for types not specifically defined below.  Generally if it's seen, it means another type needs to be defined so that the wire in question can have an appropriate color. */
-	UPROPERTY()
-	FLinearColor DefaultPinTypeColor;
-	
-	/** Execution pin type color */
-	UPROPERTY()
-	FLinearColor ExecutionPinTypeColor;
-
-	/** Boolean pin type color */
-	UPROPERTY()
-	FLinearColor BooleanPinTypeColor;
-
-	/** Byte pin type color */
-	UPROPERTY()
-	FLinearColor BytePinTypeColor;
-
-	/** Class pin type color */
-	UPROPERTY()
-	FLinearColor ClassPinTypeColor;
-
-	/** Integer pin type color */
-	UPROPERTY()
-	FLinearColor IntPinTypeColor;
-
-	/** Floating-point pin type color */
-	UPROPERTY()
-	FLinearColor FloatPinTypeColor;
-
-	/** Name pin type color */
-	UPROPERTY()
-	FLinearColor NamePinTypeColor;
-
-	/** Delegate pin type color */
-	UPROPERTY()
-	FLinearColor DelegatePinTypeColor;
-
-	/** Object pin type color */
-	UPROPERTY()
-	FLinearColor ObjectPinTypeColor;
-
-	/** Interface pin type color */
-	UPROPERTY()
-	FLinearColor InterfacePinTypeColor;
-
-	/** String pin type color */
-	UPROPERTY()
-	FLinearColor StringPinTypeColor;
-
-	/** Text pin type color */
-	UPROPERTY()
-	FLinearColor TextPinTypeColor;
-
-	/** Struct pin type color */
-	UPROPERTY()
-	FLinearColor StructPinTypeColor;
-
-	/** Wildcard pin type color */
-	UPROPERTY()
-	FLinearColor WildcardPinTypeColor;
-
-	/** Vector pin type color */
-	UPROPERTY()
-	FLinearColor VectorPinTypeColor;
-
-	/** Rotator pin type color */
-	UPROPERTY()
-	FLinearColor RotatorPinTypeColor;
-
-	/** Transform pin type color */
-	UPROPERTY()
-	FLinearColor TransformPinTypeColor;
-
-	/** Index pin type color */
-	UPROPERTY()
-	FLinearColor IndexPinTypeColor;
-
-	/** Event node title color */
-	UPROPERTY()
-	FLinearColor EventNodeTitleColor;
-
-	/** CallFunction node title color */
-	UPROPERTY()
-	FLinearColor FunctionCallNodeTitleColor;
-
-	/** Pure function call node title color */
-	UPROPERTY()
-	FLinearColor PureFunctionCallNodeTitleColor;
-
-	/** Parent class function call node title color */
-	UPROPERTY()
-	FLinearColor ParentFunctionCallNodeTitleColor;
-
-	/** Function Terminator node title color */
-	UPROPERTY()
-	FLinearColor FunctionTerminatorNodeTitleColor;
-
-	/** Exec Branch node title color */
-	UPROPERTY()
-	FLinearColor ExecBranchNodeTitleColor;
-
-	/** Exec Sequence node title color */
-	UPROPERTY()
-	FLinearColor ExecSequenceNodeTitleColor;
-
-	/** Result node title color */
-	UPROPERTY()
-	FLinearColor ResultNodeTitleColor;
-
-	UPROPERTY()
-	FLinearColor TraceAttackColor;
-
-	UPROPERTY()
-	float TraceAttackWireThickness;
-
-	/** How long is the attack color fully visible */
-	UPROPERTY()
-	float TraceAttackHoldPeriod;
-
-	/** How long does it take to fade from the attack to the sustain color */
-	UPROPERTY()
-	float TraceDecayPeriod;
-
-	UPROPERTY()
-	float TraceDecayExponent;
-
-	UPROPERTY()
-	FLinearColor TraceSustainColor;
-
-	UPROPERTY()
-	float TraceSustainWireThickness;
-
-	/** How long is the sustain color fully visible */
-	UPROPERTY()
-	float TraceSustainHoldPeriod;
-
-	UPROPERTY()
-	FLinearColor TraceReleaseColor;
-
-	UPROPERTY()
-	float TraceReleaseWireThickness;
-
-	/** How long does it take to fade from the sustain to the release color */
-	UPROPERTY()
-	float TraceReleasePeriod;
-
-	UPROPERTY()
-	float TraceReleaseExponent;
-
 	UPROPERTY(config)
 	bool bSCSEditorShowGrid;
 
@@ -233,12 +77,6 @@ class UEditorUserSettings : public UObject
 	//
 	// The effective time at which an event occurs is it's most recent exec time plus a bonus based on the position in the execution trace
 
-	/** How much of a bonus does an exec get for being near the top of the trace stack, and how does that fall off with position? */
-	UPROPERTY()
-	float TracePositionBonusPeriod;
-
-	UPROPERTY()
-	float TracePositionExponent;
 
 	// =====================================================================
 	// The following options are NOT exposed in the preferences Editor
diff --git a/Engine/Source/Editor/UnrealEd/Private/Editor/EditorUserSettings.cpp b/Engine/Source/Editor/UnrealEd/Private/Editor/EditorUserSettings.cpp
--- a/Engine/Source/Editor/UnrealEd/Private/Editor/EditorUserSettings.cpp
+++ b/Engine/Source/Editor/UnrealEd/Private/Editor/EditorUserSettings.cpp
@@ -9,58 +9,6 @@
 UEditorUserSettings::UEditorUserSettings(const class FPostConstructInitializeProperties& PCIP)
 	: Super(PCIP)
 {
-	// Sadness
-	DataPinStyle = BPST_VariantA;
-
-	// Blueprint editor graph node pin type colors
-	DefaultPinTypeColor   = FLinearColor(0.750000f, 0.6f, 0.4f, 1.0f);			// light brown
-	ExecutionPinTypeColor = FLinearColor(1.0f, 1.0f, 1.0f, 1.0f);				// white
-	BooleanPinTypeColor	  = FLinearColor(0.300000f, 0.0f, 0.0f, 1.0f);			// maroon
-	BytePinTypeColor      = FLinearColor(0.0f, 0.160000f, 0.131270f, 1.0f);		// dark green
-	ClassPinTypeColor     = FLinearColor(0.1f, 0.0f, 0.5f, 1.0f);				// deep purple (violet)
-	IntPinTypeColor       = FLinearColor(0.013575f, 0.770000f, 0.429609f, 1.0f);// green-blue
-	FloatPinTypeColor     = FLinearColor(0.357667f, 1.0f, 0.060000f, 1.0f);		// bright green
-	NamePinTypeColor      = FLinearColor(0.607717f, 0.224984f, 1.0f, 1.0f);		// lilac
-	DelegatePinTypeColor  = FLinearColor(1.0f, 0.04f, 0.04f, 1.0f);				// bright red
-	ObjectPinTypeColor    = FLinearColor(0.0f, 0.4f, 0.910000f, 1.0f);			// sharp blue
-	InterfacePinTypeColor = FLinearColor(0.8784f, 1.0f, 0.4f, 1.0f);			// pale green
-	StringPinTypeColor    = FLinearColor(1.0f, 0.0f, 0.660537f, 1.0f);			// bright pink
-	TextPinTypeColor      = FLinearColor(0.8f, 0.2f, 0.4f, 1.0f);				// salmon (light pink)
-	StructPinTypeColor    = FLinearColor(0.0f, 0.1f, 0.6f, 1.0f);				// deep blue
-	WildcardPinTypeColor  = FLinearColor(0.220000f, 0.195800f, 0.195800f, 1.0f);// dark gray
-	VectorPinTypeColor    = FLinearColor(1.0f, 0.591255f, 0.016512f, 1.0f);		// yellow
-	RotatorPinTypeColor   = FLinearColor(0.353393f, 0.454175f, 1.0f, 1.0f);		// periwinkle 
-	TransformPinTypeColor = FLinearColor(1.0f, 0.172585f, 0.0f, 1.0f);			// orange
-	IndexPinTypeColor     = FLinearColor(0.013575f, 0.770000f, 0.429609f, 1.0f);// green-blue
-
-	// Blueprint editor graph node title colors
-	EventNodeTitleColor = FLinearColor(1.f, 0.0f, 0.0f, 1.0f);
-	FunctionCallNodeTitleColor = FLinearColor(0.190525f, 0.583898f, 1.0f, 1.0f);
-	PureFunctionCallNodeTitleColor = FLinearColor(0.4f, 0.850000f, 0.350000f, 1.0f);
-	ParentFunctionCallNodeTitleColor = FLinearColor(1.0f, 0.170000f, 0.0f, 1.0f);
-	FunctionTerminatorNodeTitleColor = FLinearColor(0.6f, 0.0f, 1.0f, 1.0f);
-	ExecBranchNodeTitleColor = FLinearColor(1.0f, 1.0f, 1.0f, 1.0f);
-	ExecSequenceNodeTitleColor= FLinearColor(0.8f, 0.4f, 0.4f, 1.0f);
-	ResultNodeTitleColor = FLinearColor(1.0f, 0.65f, 0.4f, 1.0f);
-
-	// Blueprint debugging visuals
-	TraceAttackColor = FLinearColor(1.0f, 0.05f, 0.0f, 1.0f);
-	TraceAttackWireThickness = 12.0f;
-	TraceSustainColor = FLinearColor(1.0f, 0.7f, 0.4f, 1.0f);
-	TraceSustainWireThickness = 8.0f;
-	TraceReleaseColor = FLinearColor(0.5f, 0.5f, 0.5f, 1.0f);
-	TraceReleaseWireThickness = 2.0f;
-
-	// Blueprint debugging exec curve constants
-	TracePositionBonusPeriod=0.5f;
-	TracePositionExponent=5.0f;
-	TraceAttackHoldPeriod=0.3f;
-	TraceDecayPeriod=0.4f;
-	TraceDecayExponent=1.8f;
-	TraceSustainHoldPeriod=0.4f;
-	TraceReleasePeriod=1.5f;
-	TraceReleaseExponent=1.4f;
-
 	//Default to high quality
 	MaterialQualityLevel = 1;
 
diff --git a/Engine/Source/Editor/UnrealEd/Private/MaterialGraphNode.cpp b/Engine/Source/Editor/UnrealEd/Private/MaterialGraphNode.cpp
--- a/Engine/Source/Editor/UnrealEd/Private/MaterialGraphNode.cpp
+++ b/Engine/Source/Editor/UnrealEd/Private/MaterialGraphNode.cpp
@@ -8,6 +8,7 @@
 #include "MaterialEditorUtilities.h"
 #include "MaterialEditorActions.h"
 #include "GraphEditorActions.h"
+#include "GraphEditorSettings.h"
 
 #define LOCTEXT_NAMESPACE "MaterialGraphNode"
 
@@ -214,37 +215,37 @@ FLinearColor UMaterialGraphNode::GetNodeTitleColor() const
 		return FColor( 70, 100, 200 );
 	}
 
-	const UEditorUserSettings& Options = GEditor->AccessEditorUserSettings();
+	const UGraphEditorSettings* Settings = GetDefault<UGraphEditorSettings>();
 
 	if (UsesBoolColour(MaterialExpression))
 	{
-		return Options.BooleanPinTypeColor;
+		return Settings->BooleanPinTypeColor;
 	}
 	else if (UsesFloatColour(MaterialExpression))
 	{
-		return Options.FloatPinTypeColor;
+		return Settings->FloatPinTypeColor;
 	}
 	else if (UsesVectorColour(MaterialExpression))
 	{
-		return Options.VectorPinTypeColor;
+		return Settings->VectorPinTypeColor;
 	}
 	else if (UsesObjectColour(MaterialExpression))
 	{
-		return Options.ObjectPinTypeColor;
+		return Settings->ObjectPinTypeColor;
 	}
 	else if (UsesEventColour(MaterialExpression))
 	{
-		return Options.EventNodeTitleColor;
+		return Settings->EventNodeTitleColor;
 	}
 	else if (MaterialExpression->IsA(UMaterialExpressionMaterialFunctionCall::StaticClass()))
 	{
 		// Previously FColor(0, 116, 255);
-		return Options.FunctionCallNodeTitleColor;
+		return Settings->FunctionCallNodeTitleColor;
 	}
 	else if (MaterialExpression->IsA(UMaterialExpressionFunctionOutput::StaticClass()))
 	{
 		// Previously FColor(255, 155, 0);
-		return Options.ResultNodeTitleColor;
+		return Settings->ResultNodeTitleColor;
 	}
 	else if (UMaterial::IsParameter(MaterialExpression))
 	{
@@ -270,7 +271,7 @@ FLinearColor UMaterialGraphNode::GetNodeTitleColor() const
 	}
 
 	// Assume that most material expressions act like pure functions and don't affect anything else
-	return Options.PureFunctionCallNodeTitleColor;
+	return Settings->PureFunctionCallNodeTitleColor;
 }
 
 FString UMaterialGraphNode::GetTooltip() const
diff --git a/Engine/Source/Editor/UnrealEd/Private/MaterialGraphNode_Root.cpp b/Engine/Source/Editor/UnrealEd/Private/MaterialGraphNode_Root.cpp
--- a/Engine/Source/Editor/UnrealEd/Private/MaterialGraphNode_Root.cpp
+++ b/Engine/Source/Editor/UnrealEd/Private/MaterialGraphNode_Root.cpp
@@ -7,6 +7,7 @@
 #include "UnrealEd.h"
 #include "MaterialEditorUtilities.h"
 #include "GraphEditorActions.h"
+#include "GraphEditorSettings.h"
 
 #define LOCTEXT_NAMESPACE "MaterialGraphNode_Root"
 
@@ -25,7 +26,7 @@ FText UMaterialGraphNode_Root::GetNodeTitle(ENodeTitleType::Type TitleType) cons
 
 FLinearColor UMaterialGraphNode_Root::GetNodeTitleColor() const
 {
-	return GEditor->AccessEditorUserSettings().ResultNodeTitleColor;
+	return GetDefault<UGraphEditorSettings>()->ResultNodeTitleColor;
 }
 
 FString UMaterialGraphNode_Root::GetTooltip() const
diff --git a/Engine/Source/Editor/UnrealEd/Private/Niagara/EdGraphSchema_Niagara.cpp b/Engine/Source/Editor/UnrealEd/Private/Niagara/EdGraphSchema_Niagara.cpp
--- a/Engine/Source/Editor/UnrealEd/Private/Niagara/EdGraphSchema_Niagara.cpp
+++ b/Engine/Source/Editor/UnrealEd/Private/Niagara/EdGraphSchema_Niagara.cpp
@@ -3,6 +3,7 @@
 #include "UnrealEd.h"
 
 #include "BlueprintGraphDefinitions.h"
+#include "GraphEditorSettings.h"
 
 #include "VectorVM.h"
 
@@ -189,14 +190,14 @@ const FPinConnectionResponse UEdGraphSchema_Niagara::CanCreateConnection(const U
 FLinearColor UEdGraphSchema_Niagara::GetPinTypeColor(const FEdGraphPinType& PinType) const
 {
 	const FString& TypeString = PinType.PinCategory;
-	const UEditorUserSettings& Options = GEditor->AccessEditorUserSettings();
+	const UGraphEditorSettings* Settings = GetDefault<UGraphEditorSettings>();
 
 	if (TypeString == PC_Float)
 	{
-		return Options.FloatPinTypeColor;
+		return Settings->FloatPinTypeColor;
 	}
 
-	return Options.DefaultPinTypeColor;
+	return Settings->DefaultPinTypeColor;
 }
 
 bool UEdGraphSchema_Niagara::ShouldHidePinDefaultValue(UEdGraphPin* Pin) const
diff --git a/Engine/Source/Editor/UnrealEd/Private/Niagara/NiagaraNodeOp.cpp b/Engine/Source/Editor/UnrealEd/Private/Niagara/NiagaraNodeOp.cpp
--- a/Engine/Source/Editor/UnrealEd/Private/Niagara/NiagaraNodeOp.cpp
+++ b/Engine/Source/Editor/UnrealEd/Private/Niagara/NiagaraNodeOp.cpp
@@ -2,6 +2,7 @@
 
 #include "UnrealEd.h"
 #include "BlueprintGraphDefinitions.h"
+#include "GraphEditorSettings.h"
 #include "VectorVM.h"
 
 UNiagaraNodeOp::UNiagaraNodeOp(const class FPostConstructInitializeProperties& PCIP)
@@ -44,5 +45,5 @@ FText UNiagaraNodeOp::GetNodeTitle(ENodeTitleType::Type TitleType) const
 
 FLinearColor UNiagaraNodeOp::GetNodeTitleColor() const
 {
-	return GEditor->AccessEditorUserSettings().FunctionCallNodeTitleColor;
+	return GetDefault<UGraphEditorSettings>()->FunctionCallNodeTitleColor;
 }
diff --git a/Engine/Source/Editor/UnrealEd/Private/Niagara/NiagaraNodeOutputUpdate.cpp b/Engine/Source/Editor/UnrealEd/Private/Niagara/NiagaraNodeOutputUpdate.cpp
--- a/Engine/Source/Editor/UnrealEd/Private/Niagara/NiagaraNodeOutputUpdate.cpp
+++ b/Engine/Source/Editor/UnrealEd/Private/Niagara/NiagaraNodeOutputUpdate.cpp
@@ -2,6 +2,7 @@
 
 #include "UnrealEd.h"
 #include "BlueprintGraphDefinitions.h"
+#include "GraphEditorSettings.h"
 
 UNiagaraNodeOutputUpdate::UNiagaraNodeOutputUpdate(const class FPostConstructInitializeProperties& PCIP)
 	: Super(PCIP)
@@ -32,5 +33,5 @@ FText UNiagaraNodeOutputUpdate::GetNodeTitle(ENodeTitleType::Type TitleType) con
 
 FLinearColor UNiagaraNodeOutputUpdate::GetNodeTitleColor() const
 {
-	return GEditor->AccessEditorUserSettings().FunctionTerminatorNodeTitleColor;
+	return GetDefault<UGraphEditorSettings>()->FunctionTerminatorNodeTitleColor;
 }
diff --git a/Engine/Source/Editor/UnrealEd/Private/SoundCueGraphNode_Root.cpp b/Engine/Source/Editor/UnrealEd/Private/SoundCueGraphNode_Root.cpp
--- a/Engine/Source/Editor/UnrealEd/Private/SoundCueGraphNode_Root.cpp
+++ b/Engine/Source/Editor/UnrealEd/Private/SoundCueGraphNode_Root.cpp
@@ -5,6 +5,7 @@
 =============================================================================*/
 
 #include "UnrealEd.h"
+#include "GraphEditorSettings.h"
 #include "SoundCueGraphEditorCommands.h"
 
 #define LOCTEXT_NAMESPACE "SoundCueGraphNode_Root"
@@ -19,7 +20,7 @@ USoundCueGraphNode_Root::USoundCueGraphNode_Root(const class FPostConstructIniti
 
 FLinearColor USoundCueGraphNode_Root::GetNodeTitleColor() const
 {
-	return GEditor->AccessEditorUserSettings().ResultNodeTitleColor;
+	return GetDefault<UGraphEditorSettings>()->ResultNodeTitleColor;
 }
 
 FText USoundCueGraphNode_Root::GetNodeTitle(ENodeTitleType::Type TitleType) const
