Safe compiler warning fixes from PR 185.

This is a subset of changes (compiler warning fixes needed to compile Linux editor) from pull request 185 by sbc100, 3dluvr, wshearn, abarbu et al: https://github.com/EpicGames/UnrealEngine/pull/185
- Initialization order fixes
- Missing virtual destructors in classes with virtual funcs.
- Missing newlines at the end of file

diff --git a/Engine/Source/Developer/ScreenShotComparisonTools/Public/ScreenShotDataItem.h b/Engine/Source/Developer/ScreenShotComparisonTools/Public/ScreenShotDataItem.h
--- a/Engine/Source/Developer/ScreenShotComparisonTools/Public/ScreenShotDataItem.h
+++ b/Engine/Source/Developer/ScreenShotComparisonTools/Public/ScreenShotDataItem.h
@@ -11,10 +11,10 @@ struct FScreenShotDataItem
 {
 	// Constructors
 	FScreenShotDataItem( const FString& InViewName, const FString& InDeviceName, const FString& InAssetName, const int32 InChangeListNumber )
-		: ViewName( InViewName )
-		, DeviceName( InDeviceName )
-		, AssetName( InAssetName )
+		: AssetName( InAssetName )
 		, ChangeListNumber( InChangeListNumber )
+		, DeviceName( InDeviceName )
+		, ViewName( InViewName )
 	{}
 
 public:
diff --git a/Engine/Source/Editor/ClassViewer/Public/ClassViewerFilter.h b/Engine/Source/Editor/ClassViewer/Public/ClassViewerFilter.h
--- a/Engine/Source/Editor/ClassViewer/Public/ClassViewerFilter.h
+++ b/Engine/Source/Editor/ClassViewer/Public/ClassViewerFilter.h
@@ -5,6 +5,7 @@
 class IClassViewerFilter
 {
 public:
+	virtual ~IClassViewerFilter() {}
 
 	/**
 	 * Checks if a class is allowed by this filter.
@@ -205,4 +206,4 @@ public:
 	 * @return				true if it has an Is-A relationship to the passed in class.
 	 */
 	virtual bool IsA(const UClass* InClass) const = 0;
-};
\ No newline at end of file
+};
diff --git a/Engine/Source/Editor/ContentBrowser/Private/AssetViewTypes.h b/Engine/Source/Editor/ContentBrowser/Private/AssetViewTypes.h
--- a/Engine/Source/Editor/ContentBrowser/Private/AssetViewTypes.h
+++ b/Engine/Source/Editor/ContentBrowser/Private/AssetViewTypes.h
@@ -23,6 +23,8 @@ struct FAssetViewItem
 	{
 	}
 
+	virtual ~FAssetViewItem() {}
+
 	/** Get the type of this asset item */
 	virtual EAssetItemType::Type GetType() const = 0;
 
@@ -166,4 +168,4 @@ struct FAssetViewDuplication : public FAssetViewAsset
 	{
 		return true;
 	}
-};
\ No newline at end of file
+};
diff --git a/Engine/Source/Editor/GraphEditor/Private/ConnectionDrawingPolicy.h b/Engine/Source/Editor/GraphEditor/Private/ConnectionDrawingPolicy.h
--- a/Engine/Source/Editor/GraphEditor/Private/ConnectionDrawingPolicy.h
+++ b/Engine/Source/Editor/GraphEditor/Private/ConnectionDrawingPolicy.h
@@ -43,6 +43,8 @@ protected:
 	TSet< UEdGraphPin* > HoveredPins;
 	double LastHoverTimeEvent;
 public:
+	virtual ~FConnectionDrawingPolicy() {}
+
 	FConnectionDrawingPolicy(int32 InBackLayerID, int32 InFrontLayerID, float InZoomFactor, const FSlateRect& InClippingRect, FSlateWindowElementList& InDrawElements);
 
 	// Update the drawing policy with the set of hovered pins (which can be empty)
@@ -281,4 +283,4 @@ public:
 	// FConnectionDrawingPolicy interface
 	virtual void DetermineWiringStyle(UEdGraphPin* OutputPin, UEdGraphPin* InputPin, /*inout*/ float& Thickness, /*inout*/ FLinearColor& WireColor, /*inout*/bool& bDrawBubbles, /*inout*/ bool& bBidirectional) OVERRIDE;
 	// End of FConnectionDrawingPolicy interface
-};
\ No newline at end of file
+};
diff --git a/Engine/Source/Editor/Kismet/Private/BlueprintEditorCommands.cpp b/Engine/Source/Editor/Kismet/Private/BlueprintEditorCommands.cpp
--- a/Engine/Source/Editor/Kismet/Private/BlueprintEditorCommands.cpp
+++ b/Engine/Source/Editor/Kismet/Private/BlueprintEditorCommands.cpp
@@ -140,6 +140,8 @@ namespace NodeSpawnInfoHelpers
 class FNodeSpawnInfo
 {
 public:
+	virtual ~FNodeSpawnInfo () {}
+
 	/** Holds the UI Command to verify gestures for this action are held */
 	TSharedPtr< FUICommandInfo > CommandInfo;
 
diff --git a/Engine/Source/Editor/Kismet/Public/ISCSEditorCustomization.h b/Engine/Source/Editor/Kismet/Public/ISCSEditorCustomization.h
--- a/Engine/Source/Editor/Kismet/Public/ISCSEditorCustomization.h
+++ b/Engine/Source/Editor/Kismet/Public/ISCSEditorCustomization.h
@@ -8,6 +8,8 @@
 class ISCSEditorCustomization
 {
 public:
+	virtual ~ISCSEditorCustomization() {}
+
 	/** 
 	 * Process click in the viewport.
 	 * @param	InViewportClient	The viewport client receiving the click
@@ -48,4 +50,4 @@ public:
 	 * @return true if the operation was handled.
 	 */
 	virtual bool HandleGetWidgetTransform(class USceneComponent* InSceneComponent, FMatrix& OutWidgetTransform) = 0;
-};
\ No newline at end of file
+};
diff --git a/Engine/Source/Editor/Kismet/Public/WorkflowOrientedApp/ApplicationMode.h b/Engine/Source/Editor/Kismet/Public/WorkflowOrientedApp/ApplicationMode.h
--- a/Engine/Source/Editor/Kismet/Public/WorkflowOrientedApp/ApplicationMode.h
+++ b/Engine/Source/Editor/Kismet/Public/WorkflowOrientedApp/ApplicationMode.h
@@ -31,6 +31,8 @@ public:
 		ToolbarExtender = MakeShareable(new FExtender);
 	}
 
+	virtual ~FApplicationMode() {}
+
 	void DeactivateMode(TSharedPtr<FTabManager> InTabManager);
 	TSharedRef<FTabManager::FLayout> ActivateMode(TSharedPtr<FTabManager> InTabManager);
 	virtual void RegisterTabFactories(TSharedPtr<FTabManager> InTabManager) {}
@@ -44,4 +46,4 @@ public:
 	}
 
 	TSharedPtr<FExtender> GetToolbarExtender() { return ToolbarExtender; }
-};
\ No newline at end of file
+};
diff --git a/Engine/Source/Editor/Kismet/Public/WorkflowOrientedApp/WorkflowTabFactory.h b/Engine/Source/Editor/Kismet/Public/WorkflowOrientedApp/WorkflowTabFactory.h
--- a/Engine/Source/Editor/Kismet/Public/WorkflowOrientedApp/WorkflowTabFactory.h
+++ b/Engine/Source/Editor/Kismet/Public/WorkflowOrientedApp/WorkflowTabFactory.h
@@ -11,6 +11,8 @@ public:
 	// The type of this payload
 	const FName PayloadType;
 public:
+	virtual ~FTabPayload() {}
+
 	// Is this payload equal to the other one?
 	virtual bool IsEqual(const TSharedRef<FTabPayload>& OtherPayload) const
 	{
@@ -72,6 +74,8 @@ protected:
 	TWeakPtr<class FAssetEditorToolkit> HostingApp;
 
 public:
+	virtual ~FWorkflowTabFactory() {}
+
 	// Creates a toolkit tab given the spawning information
 	virtual TSharedRef<SDockTab> SpawnTab(const FWorkflowTabSpawnInfo& Info) const;
 	// Creates a toolkit tab without any content, given the spawning information
diff --git a/Engine/Source/Editor/Kismet/Public/WorkflowOrientedApp/WorkflowTabManager.h b/Engine/Source/Editor/Kismet/Public/WorkflowOrientedApp/WorkflowTabManager.h
--- a/Engine/Source/Editor/Kismet/Public/WorkflowOrientedApp/WorkflowTabManager.h
+++ b/Engine/Source/Editor/Kismet/Public/WorkflowOrientedApp/WorkflowTabManager.h
@@ -140,6 +140,8 @@ public:
 
 	}
 
+	virtual ~FGenericTabHistory() {}
+
 	/**
 	 * Evokes the history data this struct represents, reinstating the payload and other data it had at the time the history was saved
 	 *
diff --git a/Engine/Source/Editor/PropertyEditor/Public/IPropertyTableCustomColumn.h b/Engine/Source/Editor/PropertyEditor/Public/IPropertyTableCustomColumn.h
--- a/Engine/Source/Editor/PropertyEditor/Public/IPropertyTableCustomColumn.h
+++ b/Engine/Source/Editor/PropertyEditor/Public/IPropertyTableCustomColumn.h
@@ -9,9 +9,11 @@
 class IPropertyTableCustomColumn
 {
 public:
+	virtual ~IPropertyTableCustomColumn() {}
+
 	virtual bool Supports( const TSharedRef< IPropertyTableColumn >& Column, const TSharedRef< IPropertyTableUtilities >& Utilities ) const = 0;
 
 	virtual TSharedPtr< SWidget > CreateColumnLabel( const TSharedRef< IPropertyTableColumn >& Column, const TSharedRef< IPropertyTableUtilities >& Utilities, const FName& Style ) const = 0;
 
 	virtual TSharedPtr< IPropertyTableCellPresenter > CreateCellPresenter( const TSharedRef< IPropertyTableCell >& Cell, const TSharedRef< IPropertyTableUtilities >& Utilities, const FName& Style ) const = 0;
-};
\ No newline at end of file
+};
diff --git a/Engine/Source/Editor/PropertyEditor/Public/PropertyEditorModule.h b/Engine/Source/Editor/PropertyEditor/Public/PropertyEditorModule.h
--- a/Engine/Source/Editor/PropertyEditor/Public/PropertyEditorModule.h
+++ b/Engine/Source/Editor/PropertyEditor/Public/PropertyEditorModule.h
@@ -48,6 +48,7 @@ namespace FPropertyAccess
 class IDetailCustomization : public TSharedFromThis<IDetailCustomization>
 {
 public:
+  virtual ~IDetailCustomization() {}
 	/** Called when details should be customized */
 	virtual void CustomizeDetails( class IDetailLayoutBuilder& DetailBuilder ) = 0;
 };
@@ -143,6 +144,8 @@ public:
 class IStructCustomization : public TSharedFromThis<IStructCustomization>
 {
 public:
+	virtual ~IStructCustomization() {}
+
 	/**
 	 * Called when the header of the struct (usually where the name of the struct and information about the struct as a whole is added)
 	 * If nothing is added to the row, the header is not displayed
diff --git a/Engine/Source/Editor/Sequencer/Public/IKeyArea.h b/Engine/Source/Editor/Sequencer/Public/IKeyArea.h
--- a/Engine/Source/Editor/Sequencer/Public/IKeyArea.h
+++ b/Engine/Source/Editor/Sequencer/Public/IKeyArea.h
@@ -22,6 +22,8 @@ struct FKeyDrawingInfo
 class IKeyArea
 {
 public:
+	virtual ~IKeyArea() {}
+
 	/** @return The array of unsorted key handles in the key area */
 	virtual TArray<FKeyHandle> GetUnsortedKeyHandles() const = 0;
 
diff --git a/Engine/Source/Editor/StatsViewer/Private/StatsCustomColumn.cpp b/Engine/Source/Editor/StatsViewer/Private/StatsCustomColumn.cpp
--- a/Engine/Source/Editor/StatsViewer/Private/StatsCustomColumn.cpp
+++ b/Engine/Source/Editor/StatsViewer/Private/StatsCustomColumn.cpp
@@ -17,6 +17,8 @@ public:
 		Text = FStatsCustomColumn::GetPropertyAsText( PropertyHandle );
 	}
 
+	virtual ~FNumericStatCellPresenter() {}
+
 	virtual TSharedRef< class SWidget > ConstructDisplayWidget() OVERRIDE
 	{
 		return 
diff --git a/Engine/Source/Editor/StatsViewer/Private/StatsPageManager.h b/Engine/Source/Editor/StatsViewer/Private/StatsPageManager.h
--- a/Engine/Source/Editor/StatsViewer/Private/StatsPageManager.h
+++ b/Engine/Source/Editor/StatsViewer/Private/StatsPageManager.h
@@ -9,6 +9,8 @@
 class FStatsPageManager
 {
 public:
+	virtual ~FStatsPageManager() {}
+
 	/** Gets the instance of this manager */
 	static FStatsPageManager& Get();
 
diff --git a/Engine/Source/Editor/UnrealEd/Private/Settings/EditorLoadingSavingSettingsCustomization.h b/Engine/Source/Editor/UnrealEd/Private/Settings/EditorLoadingSavingSettingsCustomization.h
--- a/Engine/Source/Editor/UnrealEd/Private/Settings/EditorLoadingSavingSettingsCustomization.h
+++ b/Engine/Source/Editor/UnrealEd/Private/Settings/EditorLoadingSavingSettingsCustomization.h
@@ -20,6 +20,7 @@ class FEditorLoadingSavingSettingsCustomization
 	: public IDetailCustomization
 {
 public:
+	virtual ~FEditorLoadingSavingSettingsCustomization() {}
 
 	// Begin IDetailCustomization interface
 
diff --git a/Engine/Source/Editor/UnrealEd/Private/Settings/GameMapsSettingsCustomization.h b/Engine/Source/Editor/UnrealEd/Private/Settings/GameMapsSettingsCustomization.h
--- a/Engine/Source/Editor/UnrealEd/Private/Settings/GameMapsSettingsCustomization.h
+++ b/Engine/Source/Editor/UnrealEd/Private/Settings/GameMapsSettingsCustomization.h
@@ -20,6 +20,7 @@ class FGameMapsSettingsCustomization
 	: public IDetailCustomization
 {
 public:
+	virtual ~FGameMapsSettingsCustomization() {}
 
 	// Begin IDetailCustomization interface
 
diff --git a/Engine/Source/Editor/UnrealEd/Private/Settings/LevelEditorPlaySettingsCustomization.h b/Engine/Source/Editor/UnrealEd/Private/Settings/LevelEditorPlaySettingsCustomization.h
--- a/Engine/Source/Editor/UnrealEd/Private/Settings/LevelEditorPlaySettingsCustomization.h
+++ b/Engine/Source/Editor/UnrealEd/Private/Settings/LevelEditorPlaySettingsCustomization.h
@@ -288,6 +288,7 @@ class FLevelEditorPlaySettingsCustomization
 	: public IDetailCustomization
 {
 public:
+	virtual ~FLevelEditorPlaySettingsCustomization() {}
 
 	// Begin IDetailCustomization interface
 
diff --git a/Engine/Source/Editor/UnrealEd/Public/EdGraphUtilities.h b/Engine/Source/Editor/UnrealEd/Public/EdGraphUtilities.h
--- a/Engine/Source/Editor/UnrealEd/Public/EdGraphUtilities.h
+++ b/Engine/Source/Editor/UnrealEd/Public/EdGraphUtilities.h
@@ -9,6 +9,7 @@
 struct UNREALED_API FGraphPanelNodeFactory : public TSharedFromThis<FGraphPanelNodeFactory>
 {
 public:
+	virtual ~FGraphPanelNodeFactory() {}
 	virtual TSharedPtr<class SGraphNode> CreateNode(class UEdGraphNode* Node) const { return NULL; }
 };
 
@@ -17,6 +18,7 @@ public:
 struct UNREALED_API FGraphPanelPinFactory : public TSharedFromThis<FGraphPanelPinFactory>
 {
 public:
+	virtual ~FGraphPanelPinFactory() {}
 	virtual TSharedPtr<class SGraphPin> CreatePin(class UEdGraphPin* Pin) const { return NULL; }
 };
 
diff --git a/Engine/Source/Editor/UnrealEd/Public/Kismet2/Kismet2NameValidators.h b/Engine/Source/Editor/UnrealEd/Public/Kismet2/Kismet2NameValidators.h
--- a/Engine/Source/Editor/UnrealEd/Public/Kismet2/Kismet2NameValidators.h
+++ b/Engine/Source/Editor/UnrealEd/Public/Kismet2/Kismet2NameValidators.h
@@ -34,6 +34,8 @@ enum EValidatorResult
 class UNREALED_API INameValidatorInterface
 {
 public:
+	virtual ~INameValidatorInterface() {}
+
 	/** @return true if FName is valid, false otherwise */
 	virtual EValidatorResult IsValid (const FName& Name, bool bOriginal = false) = 0;
 
@@ -139,4 +141,4 @@ public:
 
 private:
 	EValidatorResult ReturnValue;
-};
\ No newline at end of file
+};
diff --git a/Engine/Source/Editor/UnrealEd/Public/Layers/ILayers.h b/Engine/Source/Editor/UnrealEd/Public/Layers/ILayers.h
--- a/Engine/Source/Editor/UnrealEd/Public/Layers/ILayers.h
+++ b/Engine/Source/Editor/UnrealEd/Public/Layers/ILayers.h
@@ -33,6 +33,8 @@ class ILayers
 public:
 	typedef IFilter< const TWeakObjectPtr< AActor >& > ActorFilter;
 
+	virtual ~ILayers() {}
+
 	/** Broadcasts whenever one or more Layers are modified*/
 	DECLARE_EVENT_ThreeParams( ILayers, FOnLayersChanged, const ELayersAction::Type /*Action*/, const TWeakObjectPtr< ULayer >& /*ChangedLayer*/, const FName& /*ChangedProperty*/);
 	virtual FOnLayersChanged& OnLayersChanged() = 0;
@@ -395,4 +397,4 @@ public:
 	 */
 	virtual bool RenameLayer( const FName OriginalLayerName, const FName& NewLayerName ) = 0;
 
-};
\ No newline at end of file
+};
diff --git a/Engine/Source/Editor/UnrealEd/Public/MatineeExporter.h b/Engine/Source/Editor/UnrealEd/Public/MatineeExporter.h
--- a/Engine/Source/Editor/UnrealEd/Public/MatineeExporter.h
+++ b/Engine/Source/Editor/UnrealEd/Public/MatineeExporter.h
@@ -23,6 +23,7 @@ class UInterpTrackVectorProp;
 class MatineeExporter
 {
 public:
+	virtual ~MatineeExporter() {}
 
 	/**
 	 * Creates and readies an empty document for export.
diff --git a/Engine/Source/Editor/UnrealEd/Public/Toolkits/AssetEditorToolkit.h b/Engine/Source/Editor/UnrealEd/Public/Toolkits/AssetEditorToolkit.h
--- a/Engine/Source/Editor/UnrealEd/Public/Toolkits/AssetEditorToolkit.h
+++ b/Engine/Source/Editor/UnrealEd/Public/Toolkits/AssetEditorToolkit.h
@@ -248,6 +248,8 @@ DECLARE_DELEGATE_RetVal_TwoParams( TSharedRef<FExtender>, FAssetEditorExtender,
 class UNREALED_API FExtensibilityManager
 {
 public:
+	virtual ~FExtensibilityManager() {}
+
 	/** Functions for outsiders to add or remove their extenders */
 	virtual void AddExtender(TSharedPtr<FExtender> Extender) {Extenders.AddUnique(Extender);}
 	virtual void RemoveExtender(TSharedPtr<FExtender> Extender) {Extenders.Remove(Extender);}
diff --git a/Engine/Source/Editor/ViewportSnapping/Public/ISnappingPolicy.h b/Engine/Source/Editor/ViewportSnapping/Public/ISnappingPolicy.h
--- a/Engine/Source/Editor/ViewportSnapping/Public/ISnappingPolicy.h
+++ b/Engine/Source/Editor/ViewportSnapping/Public/ISnappingPolicy.h
@@ -5,6 +5,8 @@
 class ISnappingPolicy
 {
 public:
+	virtual ~ISnappingPolicy() {}
+
 	/**
 	 * Snaps a scale value to the scale grid
 	 *
diff --git a/Engine/Source/Runtime/Core/Public/Misc/OutputDevice.h b/Engine/Source/Runtime/Core/Public/Misc/OutputDevice.h
--- a/Engine/Source/Runtime/Core/Public/Misc/OutputDevice.h
+++ b/Engine/Source/Runtime/Core/Public/Misc/OutputDevice.h
@@ -449,6 +449,8 @@ struct CORE_API FError
 class CORE_API FExec
 {
 public:
+	virtual ~FExec() {}
+
 	/**
 	* Exec handler
 	*
diff --git a/Engine/Source/Runtime/CoreUObject/Public/UObject/UObjectArray.h b/Engine/Source/Runtime/CoreUObject/Public/UObject/UObjectArray.h
--- a/Engine/Source/Runtime/CoreUObject/Public/UObject/UObjectArray.h
+++ b/Engine/Source/Runtime/CoreUObject/Public/UObject/UObjectArray.h
@@ -32,6 +32,7 @@ public:
 	class FUObjectCreateListener
 	{
 	public:
+		virtual ~FUObjectCreateListener() {}
 		/**
 		* Provides notification that a UObjectBase has been added to the uobject array
 		 *
@@ -47,6 +48,8 @@ public:
 	class FUObjectDeleteListener
 	{
 	public:
+		virtual ~FUObjectDeleteListener() {}
+
 		/**
 		 * Provides notification that a UObjectBase has been removed from the uobject array
 		 *
diff --git a/Engine/Source/Runtime/Engine/Classes/Curves/CurveBase.h b/Engine/Source/Runtime/Engine/Classes/Curves/CurveBase.h
--- a/Engine/Source/Runtime/Engine/Classes/Curves/CurveBase.h
+++ b/Engine/Source/Runtime/Engine/Classes/Curves/CurveBase.h
@@ -403,6 +403,8 @@ typedef FRichCurveEditInfoTemplate<const FRichCurve*>	FRichCurveEditInfoConst;
 class FCurveOwnerInterface
 {
 public:
+	virtual ~FCurveOwnerInterface() {}
+
 	/** Returns set of curves to edit. Must not release the curves while being edited. */
 	virtual TArray<FRichCurveEditInfoConst> GetCurves() const = 0;
 
diff --git a/Engine/Source/Runtime/GameplayTags/Classes/GameplayTagContainer.h b/Engine/Source/Runtime/GameplayTags/Classes/GameplayTagContainer.h
--- a/Engine/Source/Runtime/GameplayTags/Classes/GameplayTagContainer.h
+++ b/Engine/Source/Runtime/GameplayTags/Classes/GameplayTagContainer.h
@@ -107,6 +107,7 @@ struct GAMEPLAYTAGS_API FGameplayTagContainer
 	/** Constructors */
 	FGameplayTagContainer();
 	FGameplayTagContainer(FGameplayTagContainer const& Other);
+	virtual ~FGameplayTagContainer() {}
 
 	/** Assignment/Equality operators */
 	FGameplayTagContainer& operator=(FGameplayTagContainer const& Other);
