#UE4
Fix for TTP 334808, just needed to call modify on the nodes before updating their position. Identical to what we do for the comment node itself in SGraphNode::MoveTo

diff --git a/Engine/Source/Editor/GraphEditor/Private/SGraphNodeComment.cpp b/Engine/Source/Editor/GraphEditor/Private/SGraphNodeComment.cpp
--- a/Engine/Source/Editor/GraphEditor/Private/SGraphNodeComment.cpp
+++ b/Engine/Source/Editor/GraphEditor/Private/SGraphNodeComment.cpp
@@ -570,7 +570,8 @@ void SGraphNodeComment::MoveTo( const FVector2D& NewPosition, FNodeSet& NodeFilt
 				{
 					if ( !Panel->SelectionManager.IsNodeSelected(Node) && !NodeFilter.Find( Node->NodeWidget.Pin() ))
 					{
-						NodeFilter.Add( Node->NodeWidget.Pin() );
+						NodeFilter.Add(Node->NodeWidget.Pin());
+						Node->Modify();
 						Node->NodePosX += PositionDelta.X;
 						Node->NodePosY += PositionDelta.Y;
 					}
