Remove MenuWidth and MenuHeight from SComboButton.
* Replaced only usage with an SBox with Width/Height Override like SComboButton was doing.
* Removes unnecessary extra widget layer from almost all SComboButtons.

diff --git a/Engine/Source/Editor/GraphEditor/Private/KismetPins/SGraphPinEnum.cpp b/Engine/Source/Editor/GraphEditor/Private/KismetPins/SGraphPinEnum.cpp
--- a/Engine/Source/Editor/GraphEditor/Private/KismetPins/SGraphPinEnum.cpp
+++ b/Engine/Source/Editor/GraphEditor/Private/KismetPins/SGraphPinEnum.cpp
@@ -16,14 +16,18 @@ void SPinComboBox::Construct( const FArguments& InArgs )
 	this->ChildSlot
 	[
 		SAssignNew( ComboButton, SComboButton )
-		.MenuWidth(200.0f)
 		.ContentPadding(3)
 		.MenuPlacement(MenuPlacement_BelowAnchor)
 		.ButtonContent()
 		[
-			SNew( STextBlock ).ToolTipText(NSLOCTEXT("PinComboBox", "ToolTip", "Select enum values from the list"))
-			.Text( this, &SPinComboBox::OnGetVisibleTextInternal )
-			.Font( FEditorStyle::GetFontStyle( TEXT("PropertyWindow.NormalFont") ) )
+			// Wrap in configurable box to restrain height/width of menu
+			SNew(SBox)
+			.WidthOverride(200.0f)
+			[
+				SNew( STextBlock ).ToolTipText(NSLOCTEXT("PinComboBox", "ToolTip", "Select enum values from the list"))
+				.Text( this, &SPinComboBox::OnGetVisibleTextInternal )
+				.Font( FEditorStyle::GetFontStyle( TEXT("PropertyWindow.NormalFont") ) )
+			]
 		]
 		.MenuContent()
 		[
diff --git a/Engine/Source/Runtime/Slate/Private/Widgets/Input/SComboButton.cpp b/Engine/Source/Runtime/Slate/Private/Widgets/Input/SComboButton.cpp
--- a/Engine/Source/Runtime/Slate/Private/Widgets/Input/SComboButton.cpp
+++ b/Engine/Source/Runtime/Slate/Private/Widgets/Input/SComboButton.cpp
@@ -15,8 +15,6 @@ void SComboButton::Construct( const FArguments& InArgs )
 	
 	OnGetMenuContent = InArgs._OnGetMenuContent;
 	OnComboBoxOpened = InArgs._OnComboBoxOpened;
-	MenuHeight = InArgs._MenuHeight;
-	MenuWidth = InArgs._MenuWidth;
 	Method = InArgs._Method;
 	ContentWidgetPtr = InArgs._MenuContent.Widget;
 	ContentScale = InArgs._ContentScale;
@@ -127,17 +125,11 @@ FReply SComboButton::OnButtonClicked()
 void SComboButton::SetMenuContent( const TSharedRef< SWidget >& InContent )
 {
 	MenuContent = 
-		// Wrap in configurable box to restrain height/width of menu
-		SNew(SBox)
-		.WidthOverride(MenuWidth)
-		.HeightOverride(MenuHeight)
+		SNew(SBorder)
+		.BorderImage(MenuBorderBrush)
+		.Padding(MenuBorderPadding)
 		[
-			SNew(SBorder)
-			.BorderImage(MenuBorderBrush)
-			.Padding(MenuBorderPadding)
-			[
-				InContent
-			]
+			InContent
 		];
 }
 
diff --git a/Engine/Source/Runtime/Slate/Public/Widgets/Input/SComboButton.h b/Engine/Source/Runtime/Slate/Public/Widgets/Input/SComboButton.h
--- a/Engine/Source/Runtime/Slate/Public/Widgets/Input/SComboButton.h
+++ b/Engine/Source/Runtime/Slate/Public/Widgets/Input/SComboButton.h
@@ -55,10 +55,6 @@ public:
 
 		/** Spawn a new window or reuse current window for this combo*/
 		SLATE_ARGUMENT( SMenuAnchor::EMethod, Method )
-		/** The max height of the combo menu list */
-		SLATE_ATTRIBUTE( FOptionalSize, MenuHeight )
-		/** The max width of the combo menu list */
-		SLATE_ATTRIBUTE( FOptionalSize, MenuWidth )
 	SLATE_END_ARGS()
 	
 	/**
@@ -107,12 +103,6 @@ protected:
 	/** Padding to use to add a "menu border" around the drop-down content */
 	FMargin MenuBorderPadding;
 
-	/** The max height of the menu */
-	TAttribute<FOptionalSize> MenuWidth;
-
-	/** When specified, ignore the content's desired size and report the.HeightOverride as the Box's desired height. */
-	TAttribute<FOptionalSize> MenuHeight;
-
 	/** The content widget, if any, set by the user on creation */
 	TWeakPtr<SWidget> ContentWidgetPtr;
 
