Removed unused header include

diff --git a/Engine/Source/Editor/GraphEditor/Private/AnimationStateNodes/SGraphNodeAnimTransition.cpp b/Engine/Source/Editor/GraphEditor/Private/AnimationStateNodes/SGraphNodeAnimTransition.cpp
--- a/Engine/Source/Editor/GraphEditor/Private/AnimationStateNodes/SGraphNodeAnimTransition.cpp
+++ b/Engine/Source/Editor/GraphEditor/Private/AnimationStateNodes/SGraphNodeAnimTransition.cpp
@@ -7,7 +7,6 @@
 #include "Editor/UnrealEd/Public/Kismet2/BlueprintEditorUtils.h"
 #include "Editor/KismetWidgets/Public/SKismetLinearExpression.h"
 #include "Editor/PropertyEditor/Public/PropertyEditorModule.h"
-#include "Editor/PropertyEditor/Public/ISinglePropertyView.h"
 
 #include "ConnectionDrawingPolicy.h"
 #include "IDocumentation.h"
