#UE4 SGraphEditor can now lock to more than one other SGraphEditor. This will allow for a three way merge widget to lock three graph views to eachother. Also cleaned up some errant includes.. tested non-unity build w/ QAGame

diff --git a/Engine/Source/Editor/GraphEditor/Private/SGraphEditorImpl.cpp b/Engine/Source/Editor/GraphEditor/Private/SGraphEditorImpl.cpp
--- a/Engine/Source/Editor/GraphEditor/Private/SGraphEditorImpl.cpp
+++ b/Engine/Source/Editor/GraphEditor/Private/SGraphEditorImpl.cpp
@@ -339,7 +339,7 @@ void SGraphEditorImpl::Tick( const FGeometry& AllottedGeometry, const double InC
 	// If locked to another graph editor, and our panel has moved, synchronise the locked graph editor accordingly
 	if ((EdGraphObj != NULL) && GraphPanel.IsValid())
 	{
-		if(GraphPanel->HasMoved() && LockedGraph.IsValid())
+		if(GraphPanel->HasMoved() && IsLocked())
 		{
 			FocusLockedEditorHere();
 		}
@@ -509,6 +509,18 @@ bool SGraphEditorImpl::IsGraphEditable() const
 	return (EdGraphObj != NULL) && IsEditable.Get();
 }
 
+bool SGraphEditorImpl::IsLocked() const
+{
+	for( auto LockedGraph : LockedGraphs )
+	{
+		if( LockedGraph.IsValid() )
+		{
+			return true;
+		}
+	}
+	return false;
+}
+
 TSharedPtr<SWidget> SGraphEditorImpl::GetTitleBar() const
 {
 	return TitleBar;
@@ -516,7 +528,7 @@ TSharedPtr<SWidget> SGraphEditorImpl::GetTitleBar() const
 
 void SGraphEditorImpl::SetViewLocation( const FVector2D& Location, float ZoomAmount ) 
 {
-	if( GraphPanel.IsValid() &&  EdGraphObj && (!LockedGraph.IsValid() || !GraphPanel->HasDeferredObjectFocus()))
+	if( GraphPanel.IsValid() &&  EdGraphObj && (!IsLocked() || !GraphPanel->HasDeferredObjectFocus()))
 	{
 		GraphPanel->RestoreViewSettings(Location, ZoomAmount);
 	}
@@ -524,7 +536,7 @@ void SGraphEditorImpl::SetViewLocation( const FVector2D& Location, float ZoomAmo
 
 void SGraphEditorImpl::GetViewLocation( FVector2D& Location, float& ZoomAmount ) 
 {
-	if( GraphPanel.IsValid() &&  EdGraphObj && (!LockedGraph.IsValid() || !GraphPanel->HasDeferredObjectFocus()))
+	if( GraphPanel.IsValid() &&  EdGraphObj && (!IsLocked() || !GraphPanel->HasDeferredObjectFocus()))
 	{
 		Location = GraphPanel->GetViewOffset();
 		ZoomAmount = GraphPanel->GetZoomAmount();
@@ -533,7 +545,10 @@ void SGraphEditorImpl::GetViewLocation( FVector2D& Location, float& ZoomAmount )
 
 void SGraphEditorImpl::LockToGraphEditor( TWeakPtr<SGraphEditor> Other ) 
 {
-	LockedGraph = Other;
+	if( !LockedGraphs.Contains(Other) )
+	{
+		LockedGraphs.Push(Other);
+	}
 
 	if (GraphPanel.IsValid())
 	{
@@ -541,6 +556,16 @@ void SGraphEditorImpl::LockToGraphEditor( TWeakPtr<SGraphEditor> Other )
 	}
 }
 
+void SGraphEditorImpl::UnlockFromGraphEditor( TWeakPtr<SGraphEditor> Other )
+{
+	check(Other.IsValid());
+	int idx = LockedGraphs.Find(Other);
+	if( ensureMsgf(idx != INDEX_NONE, TEXT("Attempted to unlock graphs that were not locked together: %s %s"), *GetReadableLocation(), *(Other.Pin()->GetReadableLocation()) ) )
+	{
+		LockedGraphs.RemoveAtSwap(idx);
+	}
+}
+
 void SGraphEditorImpl::AddNotification( FNotificationInfo& Info, bool bSuccess )
 {
 	// set up common notification properties
@@ -555,9 +580,17 @@ void SGraphEditorImpl::AddNotification( FNotificationInfo& Info, bool bSuccess )
 
 void SGraphEditorImpl::FocusLockedEditorHere()
 {
-	if(SGraphEditor* Editor = LockedGraph.Pin().Get())
+	for( int i = 0; i < LockedGraphs.Num(); ++i )
 	{
-		Editor->SetViewLocation(GraphPanel->GetViewOffset(), GraphPanel->GetZoomAmount());	
+		TSharedPtr<SGraphEditor> LockedGraph = LockedGraphs[i].Pin();
+		if (LockedGraph != TSharedPtr<SGraphEditor>())
+		{
+			LockedGraph->SetViewLocation(GraphPanel->GetViewOffset(), GraphPanel->GetZoomAmount());
+		}
+		else
+		{
+			LockedGraphs.RemoveAtSwap(i--);
+		}
 	}
 }
 
diff --git a/Engine/Source/Editor/GraphEditor/Private/SGraphEditorImpl.h b/Engine/Source/Editor/GraphEditor/Private/SGraphEditorImpl.h
--- a/Engine/Source/Editor/GraphEditor/Private/SGraphEditorImpl.h
+++ b/Engine/Source/Editor/GraphEditor/Private/SGraphEditorImpl.h
@@ -67,7 +67,7 @@ private:
 	bool bShowPIENotification;
 
 	//FOnViewChanged	OnViewChanged;
-	TWeakPtr<SGraphEditor> LockedGraph;
+	TArray< TWeakPtr<SGraphEditor> > LockedGraphs;
 
 	/** Function to check whether PIE is active to display "Simulating" text in graph panel*/
 	EVisibility PIENotification( ) const;
@@ -125,6 +125,7 @@ public:
 	virtual void SetViewLocation(const FVector2D& Location, float ZoomAmount) OVERRIDE;
 	virtual void GetViewLocation(FVector2D& Location, float& ZoomAmount) OVERRIDE;
 	virtual void LockToGraphEditor(TWeakPtr<SGraphEditor> Other) OVERRIDE;
+	virtual void UnlockFromGraphEditor(TWeakPtr<SGraphEditor> Other) OVERRIDE;
 	virtual void AddNotification ( FNotificationInfo& Info, bool bSuccess ) OVERRIDE;
 	virtual void SetPinVisibility(SGraphEditor::EPinVisibility Visibility) OVERRIDE;
 	// End of SGraphEditor interface
@@ -149,5 +150,7 @@ private:
 	FSlateColor GetZoomTextColorAndOpacity() const;
 
 	bool IsGraphEditable() const;
+
+	bool IsLocked() const;
 };
 
diff --git a/Engine/Source/Editor/Kismet/Private/SBlueprintDiff.cpp b/Engine/Source/Editor/Kismet/Private/SBlueprintDiff.cpp
--- a/Engine/Source/Editor/Kismet/Private/SBlueprintDiff.cpp
+++ b/Engine/Source/Editor/Kismet/Private/SBlueprintDiff.cpp
@@ -1,16 +1,11 @@
 // Copyright 1998-2014 Epic Games, Inc. All Rights Reserved.
 
 #include "BlueprintEditorPrivatePCH.h"
-#include "UnrealEd.h"
 #include "BlueprintUtilities.h"
 #include "Editor/UnrealEd/Public/EdGraphUtilities.h"
 #include "BlueprintEditorCommands.h"
 #include "BlueprintEditor.h"
 #include "SBlueprintDiff.h"
-#include "Editor/UnrealEd/Public/Kismet2/DebuggerCommands.h"
-#include "Editor/UnrealEd/Public/Kismet2/BlueprintEditorUtils.h"
-#include "AssetRegistryModule.h"
-#include "Editor/PropertyEditor/Public/PropertyEditorModule.h"
 #include "GraphEditor.h"
 #include "SGraphTitleBar.h"
 #include "GraphDiffControl.h"
@@ -648,8 +643,8 @@ void SBlueprintDiff::ResetGraphEditors()
 		}
 		else
 		{
-			PanelOld.GraphEditor.Pin()->LockToGraphEditor(TWeakPtr<SGraphEditor>());
-			PanelNew.GraphEditor.Pin()->LockToGraphEditor(TWeakPtr<SGraphEditor>());
+			PanelOld.GraphEditor.Pin()->UnlockFromGraphEditor(PanelNew.GraphEditor);
+			PanelNew.GraphEditor.Pin()->UnlockFromGraphEditor(PanelOld.GraphEditor);
 		}	
 	}
 }
diff --git a/Engine/Source/Editor/UnrealEd/Public/GraphEditor.h b/Engine/Source/Editor/UnrealEd/Public/GraphEditor.h
--- a/Engine/Source/Editor/UnrealEd/Public/GraphEditor.h
+++ b/Engine/Source/Editor/UnrealEd/Public/GraphEditor.h
@@ -195,6 +195,15 @@ public:
 		}
 	}
 
+	/* Unlock two graph editors from each other */
+	virtual void UnlockFromGraphEditor(TWeakPtr<SGraphEditor> Other)
+	{
+		if (Implementation.IsValid())
+		{
+			Implementation->UnlockFromGraphEditor(Other);
+		}
+	}
+
 	/** Bring the specified node into view */
 	virtual void JumpToNode( const class UEdGraphNode* JumpToMe, bool bRequestRename )
 	{
