Allowed configuration of graph state decorations to disable PIE and read only decorations
#ttp 335218 - Remove read only/PIE borders and notifications from graph preview tooltips

diff --git a/Engine/Source/Editor/BehaviorTreeEditor/Private/SGraphNode_BehaviorTree.cpp b/Engine/Source/Editor/BehaviorTreeEditor/Private/SGraphNode_BehaviorTree.cpp
--- a/Engine/Source/Editor/BehaviorTreeEditor/Private/SGraphNode_BehaviorTree.cpp
+++ b/Engine/Source/Editor/BehaviorTreeEditor/Private/SGraphNode_BehaviorTree.cpp
@@ -892,8 +892,11 @@ TSharedPtr<SToolTip> SGraphNode_BehaviorTree::GetComplexTooltip()
 	{
 		return SNew(SToolTip)
 			[
+				// Create the tooltip graph preview, make sure to disable state overlays to
+				// prevent the PIE / read-only borders from obscuring the graph
 				SNew(SGraphPreviewer, DecoratorNode->GetBoundGraph())
 				.CornerOverlayText(this, &SGraphNode_BehaviorTree::GetPreviewCornerText)
+				.ShowGraphStateOverlay(false)
 			];
 	}
 	return NULL;
diff --git a/Engine/Source/Editor/GraphEditor/Private/AnimationStateNodes/SGraphNodeAnimState.cpp b/Engine/Source/Editor/GraphEditor/Private/AnimationStateNodes/SGraphNodeAnimState.cpp
--- a/Engine/Source/Editor/GraphEditor/Private/AnimationStateNodes/SGraphNodeAnimState.cpp
+++ b/Engine/Source/Editor/GraphEditor/Private/AnimationStateNodes/SGraphNodeAnimState.cpp
@@ -294,8 +294,11 @@ TSharedPtr<SToolTip> SGraphNodeAnimState::GetComplexTooltip()
 			+SVerticalBox::Slot()
 			.AutoHeight()
 			[
+				// Create the tooltip preview, ensure to disable state overlays to stop
+				// PIE and read-only borders obscuring the graph
 				SNew(SGraphPreviewer, StateNode->GetBoundGraph())
 				.CornerOverlayText(this, &SGraphNodeAnimState::GetPreviewCornerText)
+				.ShowGraphStateOverlay(false)
 			]
 	
 			+SVerticalBox::Slot()
diff --git a/Engine/Source/Editor/GraphEditor/Private/GraphEditorModule.cpp b/Engine/Source/Editor/GraphEditor/Private/GraphEditorModule.cpp
--- a/Engine/Source/Editor/GraphEditor/Private/GraphEditorModule.cpp
+++ b/Engine/Source/Editor/GraphEditor/Private/GraphEditorModule.cpp
@@ -63,22 +63,22 @@ TSharedRef<SGraphEditor> FGraphEditorModule::PRIVATE_MakeGraphEditor(
 	UEdGraph* InGraphToDiff,
 	FSimpleDelegate InOnNavigateHistoryBack,
 	FSimpleDelegate InOnNavigateHistoryForward,
-	bool InShowPIENotification)
+	TAttribute<bool> ShowGraphStateOverlay)
 {
 	return
 		SNew(SGraphEditorImpl)
-		.AdditionalCommands( InAdditionalCommands )
-		.IsEditable( InIsEditable )
-		.Appearance( Appearance )
-		.TitleBar( InTitleBar )
-		.TitleBarEnabledOnly( InTitleBarEnabledOnly )
-		.GraphToEdit( InGraphToEdit)
+		.AdditionalCommands(InAdditionalCommands)
+		.IsEditable(InIsEditable)
+		.Appearance(Appearance)
+		.TitleBar(InTitleBar)
+		.TitleBarEnabledOnly(InTitleBarEnabledOnly)
+		.GraphToEdit(InGraphToEdit)
 		.GraphEvents(InGraphEvents)
 		.AutoExpandActionMenu(InAutoExpandActionMenu)
 		.GraphToDiff(InGraphToDiff)
 		.OnNavigateHistoryBack(InOnNavigateHistoryBack)
 		.OnNavigateHistoryForward(InOnNavigateHistoryForward)
-		.ShowPIENotification(InShowPIENotification);
+		.ShowGraphStateOverlay(ShowGraphStateOverlay);
 }
 
 /////////////////////////////////////////////////////
diff --git a/Engine/Source/Editor/GraphEditor/Private/KismetNodes/SGraphNodeK2Composite.cpp b/Engine/Source/Editor/GraphEditor/Private/KismetNodes/SGraphNodeK2Composite.cpp
--- a/Engine/Source/Editor/GraphEditor/Private/KismetNodes/SGraphNodeK2Composite.cpp
+++ b/Engine/Source/Editor/GraphEditor/Private/KismetNodes/SGraphNodeK2Composite.cpp
@@ -148,8 +148,11 @@ TSharedPtr<SToolTip> SGraphNodeK2Composite::GetComplexTooltip()
 				+SVerticalBox::Slot()
 				.AutoHeight()
 				[
+					// Create preview for the tooltip, make sure to disable state overlays to prevent
+					// PIE and read-only borders obscuring the graph
 					SNew(SGraphPreviewer, BoundGraph)
 					.CornerOverlayText(this, &SGraphNodeK2Composite::GetPreviewCornerText)
+					.ShowGraphStateOverlay(false)
 				]
 			];
 
diff --git a/Engine/Source/Editor/GraphEditor/Private/SGraphEditorImpl.cpp b/Engine/Source/Editor/GraphEditor/Private/SGraphEditorImpl.cpp
--- a/Engine/Source/Editor/GraphEditor/Private/SGraphEditorImpl.cpp
+++ b/Engine/Source/Editor/GraphEditor/Private/SGraphEditorImpl.cpp
@@ -148,7 +148,7 @@ void SGraphEditorImpl::Construct( const FArguments& InArgs )
 	TitleBarEnabledOnly = InArgs._TitleBarEnabledOnly;
 	TitleBar	= InArgs._TitleBar;
 	bAutoExpandActionMenu = InArgs._AutoExpandActionMenu;
-	bShowPIENotification = InArgs._ShowPIENotification;
+	ShowGraphStateOverlay = InArgs._ShowGraphStateOverlay;
 
 	OnNavigateHistoryBack = InArgs._OnNavigateHistoryBack;
 	OnNavigateHistoryForward = InArgs._OnNavigateHistoryForward;
@@ -241,7 +241,7 @@ void SGraphEditorImpl::Construct( const FArguments& InArgs )
 			.OnSpawnNodeByShortcut( InArgs._GraphEvents.OnSpawnNodeByShortcut )
 			.OnUpdateGraphPanel( this, &SGraphEditorImpl::GraphEd_OnPanelUpdated )
 			.OnDisallowedPinConnection( InArgs._GraphEvents.OnDisallowedPinConnection )
-			.ShowPIENotification( InArgs._ShowPIENotification )
+			.ShowGraphStateOverlay(InArgs._ShowGraphStateOverlay)
 		]
 
 		// Indicator of current zoom level
@@ -317,7 +317,7 @@ void SGraphEditorImpl::Construct( const FArguments& InArgs )
 
 EVisibility SGraphEditorImpl::PIENotification( ) const
 {
-	if (bShowPIENotification && (GEditor->bIsSimulatingInEditor || GEditor->PlayWorld != NULL))
+	if(ShowGraphStateOverlay.Get() && (GEditor->bIsSimulatingInEditor || GEditor->PlayWorld != NULL))
 	{
 		return EVisibility::Visible;
 	}
@@ -609,7 +609,7 @@ void SGraphEditorImpl::SetPinVisibility( SGraphEditor::EPinVisibility Visibility
 
 EVisibility SGraphEditorImpl::ReadOnlyVisibility() const
 {
-	if(PIENotification() == EVisibility::Hidden && !IsEditable.Get())
+	if(ShowGraphStateOverlay.Get() && PIENotification() == EVisibility::Hidden && !IsEditable.Get())
 	{
 		return EVisibility::Visible;
 	}
diff --git a/Engine/Source/Editor/GraphEditor/Private/SGraphEditorImpl.h b/Engine/Source/Editor/GraphEditor/Private/SGraphEditorImpl.h
--- a/Engine/Source/Editor/GraphEditor/Private/SGraphEditorImpl.h
+++ b/Engine/Source/Editor/GraphEditor/Private/SGraphEditorImpl.h
@@ -34,7 +34,9 @@ public:
 		SLATE_ARGUMENT(bool, AutoExpandActionMenu)
 		SLATE_EVENT(FSimpleDelegate, OnNavigateHistoryBack)
 		SLATE_EVENT(FSimpleDelegate, OnNavigateHistoryForward)
-		SLATE_ARGUMENT(bool, ShowPIENotification)
+
+		/** Show overlay elements for the graph state such as the PIE and read-only borders and text */
+		SLATE_ATTRIBUTE(bool, ShowGraphStateOverlay)
 	SLATE_END_ARGS()
 
 	void Construct( const FArguments& InArgs );
@@ -64,7 +66,8 @@ private:
 
 	bool bAutoExpandActionMenu;
 
-	bool bShowPIENotification;
+	/** Whether to show the state (read only / PIE etc) Overlay on the panel */
+	TAttribute<bool> ShowGraphStateOverlay;
 
 	//FOnViewChanged	OnViewChanged;
 	TArray< TWeakPtr<SGraphEditor> > LockedGraphs;
diff --git a/Engine/Source/Editor/GraphEditor/Private/SGraphPanel.cpp b/Engine/Source/Editor/GraphEditor/Private/SGraphPanel.cpp
--- a/Engine/Source/Editor/GraphEditor/Private/SGraphPanel.cpp
+++ b/Engine/Source/Editor/GraphEditor/Private/SGraphPanel.cpp
@@ -87,8 +87,8 @@ void SGraphPanel::Construct( const SGraphPanel::FArguments& InArgs )
 	// Register for notifications
 	MyRegisteredGraphChangedDelegate = FOnGraphChanged::FDelegate::CreateSP(this, &SGraphPanel::OnGraphChanged);
 	this->GraphObj->AddOnGraphChangedHandler(MyRegisteredGraphChangedDelegate);
-
-	bShowPIENotification = InArgs._ShowPIENotification;
+	
+	ShowGraphStateOverlay = InArgs._ShowGraphStateOverlay;
 }
 
 SGraphPanel::~SGraphPanel()
@@ -459,28 +459,31 @@ int32 SGraphPanel::OnPaint( const FGeometry& AllottedGeometry, const FSlateRect&
 	++MaxLayerId;
 	PaintSurroundSunkenShadow(FEditorStyle::GetBrush(TEXT("Graph.Shadow")), AllottedGeometry, MyClippingRect, OutDrawElements, MaxLayerId);
 
-	const FSlateBrush* BorderBrush = nullptr;
-	if (bShowPIENotification && (GEditor->bIsSimulatingInEditor || GEditor->PlayWorld != NULL))
-	{
-		// Draw a surrounding indicator when PIE is active, to make it clear that the graph is read-only, etc...
-		BorderBrush = FEditorStyle::GetBrush(TEXT("Graph.PlayInEditor"));
-	}
-	else if (!IsEditable.Get())
+	if(ShowGraphStateOverlay.Get())
 	{
-		// Draw a different border when we're not simulating but the graph is read-only
-		BorderBrush = FEditorStyle::GetBrush(TEXT("Graph.ReadOnlyBorder"));
-	}
+		const FSlateBrush* BorderBrush = nullptr;
+		if((GEditor->bIsSimulatingInEditor || GEditor->PlayWorld != NULL))
+		{
+			// Draw a surrounding indicator when PIE is active, to make it clear that the graph is read-only, etc...
+			BorderBrush = FEditorStyle::GetBrush(TEXT("Graph.PlayInEditor"));
+		}
+		else if(!IsEditable.Get())
+		{
+			// Draw a different border when we're not simulating but the graph is read-only
+			BorderBrush = FEditorStyle::GetBrush(TEXT("Graph.ReadOnlyBorder"));
+		}
 
-	if (BorderBrush)
-	{
-		// Actually draw the border
-		FSlateDrawElement::MakeBox(
-			OutDrawElements,
-			MaxLayerId,
-			AllottedGeometry.ToPaintGeometry(),
-			BorderBrush,
-			MyClippingRect
-			);
+		if(BorderBrush)
+		{
+			// Actually draw the border
+			FSlateDrawElement::MakeBox(
+				OutDrawElements,
+				MaxLayerId,
+				AllottedGeometry.ToPaintGeometry(),
+				BorderBrush,
+				MyClippingRect
+				);
+		}
 	}
 
 	// Draw the marquee selection rectangle
diff --git a/Engine/Source/Editor/GraphEditor/Private/SGraphPanel.h b/Engine/Source/Editor/GraphEditor/Private/SGraphPanel.h
--- a/Engine/Source/Editor/GraphEditor/Private/SGraphPanel.h
+++ b/Engine/Source/Editor/GraphEditor/Private/SGraphPanel.h
@@ -43,8 +43,8 @@ public:
 		, _GraphObjToDiff( static_cast<UEdGraph*>(NULL) )
 		, _InitialZoomToFit( false )
 		, _IsEditable( true )
+		, _ShowGraphStateOverlay(true)
 		, _OnUpdateGraphPanel()
-		, _ShowPIENotification( true )
 		{}
 
 		SLATE_EVENT( FOnGetContextMenuFor, OnGetContextMenuFor )
@@ -56,12 +56,13 @@ public:
 		SLATE_ARGUMENT( class UEdGraph*, GraphObjToDiff )
 		SLATE_ARGUMENT( bool, InitialZoomToFit )
 		SLATE_ATTRIBUTE( bool, IsEditable )
+		/** Show overlay elements for the graph state such as the PIE and read-only borders and text */
+		SLATE_ATTRIBUTE(bool, ShowGraphStateOverlay)
 		SLATE_EVENT( FOnNodeVerifyTextCommit, OnVerifyTextCommit )
 		SLATE_EVENT( FOnNodeTextCommitted, OnTextCommitted )
 		SLATE_EVENT( SGraphEditor::FOnSpawnNodeByShortcut, OnSpawnNodeByShortcut )
 		SLATE_EVENT( FOnUpdateGraphPanel, OnUpdateGraphPanel )
 		SLATE_EVENT( SGraphEditor::FOnDisallowedPinConnection, OnDisallowedPinConnection )
-		SLATE_ARGUMENT( bool, ShowPIENotification )
 		//SLATE_ATTRIBUTE( FGraphAppearanceInfo, Appearance )
 	SLATE_END_ARGS()
 
@@ -231,6 +232,9 @@ protected:
 	/** Whether to draw the overlay indicating we're in PIE */
 	bool bShowPIENotification;
 
+	/** Whether to draw decorations for graph state (PIE / ReadOnly etc.) */
+	TAttribute<bool> ShowGraphStateOverlay;
+
 private:
 	/** Map of recently added nodes for the panel */
 	TMap<class UEdGraphNode*, FEdGraphEditAction> UserAddedNodes;
diff --git a/Engine/Source/Editor/GraphEditor/Private/SGraphPreviewer.cpp b/Engine/Source/Editor/GraphEditor/Private/SGraphPreviewer.cpp
--- a/Engine/Source/Editor/GraphEditor/Private/SGraphPreviewer.cpp
+++ b/Engine/Source/Editor/GraphEditor/Private/SGraphPreviewer.cpp
@@ -34,6 +34,7 @@ void SGraphPreviewer::Construct( const FArguments& InArgs, UEdGraph* InGraphObj
 			SAssignNew(GraphPanel, SGraphPanel)
 			.GraphObj( EdGraphObj )
 			.IsEditable( false )
+			.ShowGraphStateOverlay(InArgs._ShowGraphStateOverlay)
 			.InitialZoomToFit( true )
 		]
 
diff --git a/Engine/Source/Editor/GraphEditor/Private/SGraphPreviewer.h b/Engine/Source/Editor/GraphEditor/Private/SGraphPreviewer.h
--- a/Engine/Source/Editor/GraphEditor/Private/SGraphPreviewer.h
+++ b/Engine/Source/Editor/GraphEditor/Private/SGraphPreviewer.h
@@ -6,8 +6,12 @@
 class GRAPHEDITOR_API SGraphPreviewer : public SCompoundWidget
 {
 public:
-	SLATE_BEGIN_ARGS(SGraphPreviewer) {}
+	SLATE_BEGIN_ARGS(SGraphPreviewer)
+		: _ShowGraphStateOverlay(true)
+	{}
 		SLATE_ATTRIBUTE( FString, CornerOverlayText )
+		/** Show overlay elements for the graph state such as the PIE and read-only borders and text */
+		SLATE_ATTRIBUTE(bool, ShowGraphStateOverlay)
 		SLATE_ARGUMENT( TSharedPtr<SWidget>, TitleBar )
 	SLATE_END_ARGS()
 
diff --git a/Engine/Source/Editor/GraphEditor/Public/GraphEditorModule.h b/Engine/Source/Editor/GraphEditor/Public/GraphEditorModule.h
--- a/Engine/Source/Editor/GraphEditor/Public/GraphEditorModule.h
+++ b/Engine/Source/Editor/GraphEditor/Public/GraphEditorModule.h
@@ -41,7 +41,7 @@ private:
 		UEdGraph* InGraphToDiff,
 		FSimpleDelegate InOnNavigateHistoryBack,
 		FSimpleDelegate InOnNavigateHistoryForward,
-		bool InShowPIENotification);
+		TAttribute<bool> ShowGraphStateOverlay);
 
 private:
 	/** All extender delegates for the graph menus */
diff --git a/Engine/Source/Editor/MaterialEditor/Private/MaterialEditor.cpp b/Engine/Source/Editor/MaterialEditor/Private/MaterialEditor.cpp
--- a/Engine/Source/Editor/MaterialEditor/Private/MaterialEditor.cpp
+++ b/Engine/Source/Editor/MaterialEditor/Private/MaterialEditor.cpp
@@ -3531,7 +3531,7 @@ TSharedRef<SGraphEditor> FMaterialEditor::CreateGraphEditorWidget()
 		.Appearance(AppearanceInfo)
 		.GraphToEdit(Material->MaterialGraph)
 		.GraphEvents(InEvents)
-		.ShowPIENotification(false);
+		.ShowGraphStateOverlay(false);
 }
 
 void FMaterialEditor::CleanUnusedExpressions()
diff --git a/Engine/Source/Editor/SoundClassEditor/Private/SoundClassEditor.cpp b/Engine/Source/Editor/SoundClassEditor/Private/SoundClassEditor.cpp
--- a/Engine/Source/Editor/SoundClassEditor/Private/SoundClassEditor.cpp
+++ b/Engine/Source/Editor/SoundClassEditor/Private/SoundClassEditor.cpp
@@ -228,7 +228,7 @@ TSharedRef<SGraphEditor> FSoundClassEditor::CreateGraphEditorWidget()
 		.Appearance(AppearanceInfo)
 		.GraphToEdit(SoundClass->SoundClassGraph)
 		.GraphEvents(InEvents)
-		.ShowPIENotification(false);
+		.ShowGraphStateOverlay(false);
 }
 
 void FSoundClassEditor::OnSelectedNodesChanged(const TSet<class UObject*>& NewSelection)
diff --git a/Engine/Source/Editor/SoundCueEditor/Private/SoundCueEditor.cpp b/Engine/Source/Editor/SoundCueEditor/Private/SoundCueEditor.cpp
--- a/Engine/Source/Editor/SoundCueEditor/Private/SoundCueEditor.cpp
+++ b/Engine/Source/Editor/SoundCueEditor/Private/SoundCueEditor.cpp
@@ -528,7 +528,7 @@ TSharedRef<SGraphEditor> FSoundCueEditor::CreateGraphEditorWidget()
 		.GraphToEdit(SoundCue->GetGraph())
 		.GraphEvents(InEvents)
 		.AutoExpandActionMenu(true)
-		.ShowPIENotification(false);
+		.ShowGraphStateOverlay(false);
 }
 
 FGraphPanelSelectionSet FSoundCueEditor::GetSelectedNodes() const
diff --git a/Engine/Source/Editor/UnrealEd/Private/GraphEditor.cpp b/Engine/Source/Editor/UnrealEd/Private/GraphEditor.cpp
--- a/Engine/Source/Editor/UnrealEd/Private/GraphEditor.cpp
+++ b/Engine/Source/Editor/UnrealEd/Private/GraphEditor.cpp
@@ -28,7 +28,7 @@ void SGraphEditor::ConstructImplementation( const FArguments& InArgs )
 		InArgs._GraphToDiff,
 		InArgs._OnNavigateHistoryBack,
 		InArgs._OnNavigateHistoryForward,
-		InArgs._ShowPIENotification
+		InArgs._ShowGraphStateOverlay
 		);
 
 	this->ChildSlot
diff --git a/Engine/Source/Editor/UnrealEd/Public/GraphEditor.h b/Engine/Source/Editor/UnrealEd/Public/GraphEditor.h
--- a/Engine/Source/Editor/UnrealEd/Public/GraphEditor.h
+++ b/Engine/Source/Editor/UnrealEd/Public/GraphEditor.h
@@ -109,7 +109,7 @@ public:
 		, _GraphToEdit(NULL)
 		, _GraphToDiff(NULL)
 		, _AutoExpandActionMenu(false)
-		, _ShowPIENotification(true)
+		, _ShowGraphStateOverlay(true)
 		{}
 
 		SLATE_ARGUMENT( TSharedPtr<FUICommandList>, AdditionalCommands )
@@ -124,7 +124,9 @@ public:
 		SLATE_ARGUMENT( bool, AutoExpandActionMenu )
 		SLATE_EVENT(FSimpleDelegate, OnNavigateHistoryBack)
 		SLATE_EVENT(FSimpleDelegate, OnNavigateHistoryForward)
-		SLATE_ARGUMENT(bool, ShowPIENotification)
+
+		/** Show overlay elements for the graph state such as the PIE and read-only borders and text */
+		SLATE_ATTRIBUTE(bool, ShowGraphStateOverlay)				
 	SLATE_END_ARGS()
 
 	/**
