#UE4 Trivial fix for diff view displaying {PinName} instead of the actual name of the pin..

diff --git a/Engine/Source/Editor/GraphEditor/Private/GraphDiffControl.cpp b/Engine/Source/Editor/GraphEditor/Private/GraphDiffControl.cpp
--- a/Engine/Source/Editor/GraphEditor/Private/GraphDiffControl.cpp
+++ b/Engine/Source/Editor/GraphEditor/Private/GraphDiffControl.cpp
@@ -242,7 +242,7 @@ static void DiffR_LinkedToNode(FDiffResults& Results, class UEdGraphPin* Pin1, c
 		Diff.ToolTip =  FText::Format(LOCTEXT("DIF_PinLinkMovedToolTip", "Pin '{PinNameForNode1}' was linked to Node '{NodeName1}', but is now linked to Node '{NodeName2}'"), Args).ToString();
 
 		Diff.DisplayColor = FLinearColor(0.85f,0.71f,0.25f);
-		Diff.DisplayString = FText::Format(LOCTEXT("DIF_PinLinkMoved", "Link Moved  '{PinName}' ['{NodeName1}' -> '{NodeName2}']"), Args).ToString();
+		Diff.DisplayString = FText::Format(LOCTEXT("DIF_PinLinkMoved", "Link Moved  '{PinNameForNode1}' ['{NodeName1}' -> '{NodeName2}']"), Args).ToString();
 		Results.Add(Diff);
 	}
 }
