UE4 Refactoring. Changed OVERRIDE and FINAL macros to keywords override and final.

diff --git a/Engine/Source/Editor/DeviceProfileEditor/Private/DeviceProfileConsoleVariableColumn.h b/Engine/Source/Editor/DeviceProfileEditor/Private/DeviceProfileConsoleVariableColumn.h
--- a/Engine/Source/Editor/DeviceProfileEditor/Private/DeviceProfileConsoleVariableColumn.h
+++ b/Engine/Source/Editor/DeviceProfileEditor/Private/DeviceProfileConsoleVariableColumn.h
@@ -24,9 +24,9 @@ public:
 	FDeviceProfileConsoleVariableColumn();
 
 	/** Begin IPropertyTableCustomColumn interface */
-	virtual bool Supports( const TSharedRef< IPropertyTableColumn >& Column, const TSharedRef< IPropertyTableUtilities >& Utilities ) const OVERRIDE;
-	virtual TSharedPtr< SWidget > CreateColumnLabel( const TSharedRef< IPropertyTableColumn >& Column, const TSharedRef< IPropertyTableUtilities >& Utilities, const FName& Style ) const OVERRIDE;
-	virtual TSharedPtr< IPropertyTableCellPresenter > CreateCellPresenter( const TSharedRef< IPropertyTableCell >& Cell, const TSharedRef< IPropertyTableUtilities >& Utilities, const FName& Style ) const OVERRIDE;
+	virtual bool Supports( const TSharedRef< IPropertyTableColumn >& Column, const TSharedRef< IPropertyTableUtilities >& Utilities ) const override;
+	virtual TSharedPtr< SWidget > CreateColumnLabel( const TSharedRef< IPropertyTableColumn >& Column, const TSharedRef< IPropertyTableUtilities >& Utilities, const FName& Style ) const override;
+	virtual TSharedPtr< IPropertyTableCellPresenter > CreateCellPresenter( const TSharedRef< IPropertyTableCell >& Cell, const TSharedRef< IPropertyTableUtilities >& Utilities, const FName& Style ) const override;
 	/** End IPropertyTableCustomColumn interface */
 
 	/**
diff --git a/Engine/Source/Editor/DeviceProfileEditor/Private/SDeviceProfileSelectionPanel.cpp b/Engine/Source/Editor/DeviceProfileEditor/Private/SDeviceProfileSelectionPanel.cpp
--- a/Engine/Source/Editor/DeviceProfileEditor/Private/SDeviceProfileSelectionPanel.cpp
+++ b/Engine/Source/Editor/DeviceProfileEditor/Private/SDeviceProfileSelectionPanel.cpp
@@ -44,7 +44,7 @@ public:
 	 * @param ColumnName - The name of the column to generate the widget for.
 	 * @return The widget.
 	 */
-	virtual TSharedRef<SWidget> GenerateWidgetForColumn( const FName& ColumnName ) OVERRIDE;
+	virtual TSharedRef<SWidget> GenerateWidgetForColumn( const FName& ColumnName ) override;
 
 
 public:
diff --git a/Engine/Source/Editor/DeviceProfileEditor/Public/DeviceProfileEditorModule.h b/Engine/Source/Editor/DeviceProfileEditor/Public/DeviceProfileEditorModule.h
--- a/Engine/Source/Editor/DeviceProfileEditor/Public/DeviceProfileEditorModule.h
+++ b/Engine/Source/Editor/DeviceProfileEditor/Public/DeviceProfileEditorModule.h
@@ -10,9 +10,9 @@ class FDeviceProfileEditorModule : public IModuleInterface
 public:
 
 	// Begin IModuleInterface interface
-	virtual void StartupModule() OVERRIDE;
+	virtual void StartupModule() override;
 
-	virtual void ShutdownModule() OVERRIDE;
+	virtual void ShutdownModule() override;
 	// End IModuleInterface interface
 
 public:
diff --git a/Engine/Source/Editor/DeviceProfileServices/Private/DeviceProfileServicesModule.cpp b/Engine/Source/Editor/DeviceProfileServices/Private/DeviceProfileServicesModule.cpp
--- a/Engine/Source/Editor/DeviceProfileServices/Private/DeviceProfileServicesModule.cpp
+++ b/Engine/Source/Editor/DeviceProfileServices/Private/DeviceProfileServicesModule.cpp
@@ -15,7 +15,7 @@ class FDeviceProfileServicesModule
 public:
 
 		// Begin IDeviceProfileServicesModule interface
-		virtual IDeviceProfileServicesUIManagerRef GetProfileServicesManager( ) OVERRIDE
+		virtual IDeviceProfileServicesUIManagerRef GetProfileServicesManager( ) override
 		{
 			if (!DeviceProfileServicesUIManagerSingleton.IsValid())
 			{
diff --git a/Engine/Source/Editor/DeviceProfileServices/Private/DeviceProfileServicesUIManager.h b/Engine/Source/Editor/DeviceProfileServices/Private/DeviceProfileServicesUIManager.h
--- a/Engine/Source/Editor/DeviceProfileServices/Private/DeviceProfileServicesUIManager.h
+++ b/Engine/Source/Editor/DeviceProfileServices/Private/DeviceProfileServicesUIManager.h
@@ -24,11 +24,11 @@ public:
 
 	// Begin IDeviceProfileServicesUIManager Interface
 
-	virtual const FName GetDeviceIconName( const FString& DeviceName ) const OVERRIDE;
-	virtual const TArray<TSharedPtr<FString> > GetPlatformList( ) OVERRIDE;
-	virtual void GetProfilesByType( TArray<UDeviceProfile*>& OutDeviceProfiles, const FString& InType ) OVERRIDE;
-	virtual const FName GetPlatformIconName( const FString& DeviceName ) const OVERRIDE;
-	virtual void SetProfile( const FString& DeviceProfileName ) OVERRIDE;
+	virtual const FName GetDeviceIconName( const FString& DeviceName ) const override;
+	virtual const TArray<TSharedPtr<FString> > GetPlatformList( ) override;
+	virtual void GetProfilesByType( TArray<UDeviceProfile*>& OutDeviceProfiles, const FString& InType ) override;
+	virtual const FName GetPlatformIconName( const FString& DeviceName ) const override;
+	virtual void SetProfile( const FString& DeviceProfileName ) override;
 
 	// End IDeviceProfileServicesUIManager Interface
 
diff --git a/Engine/Source/Editor/DistCurveEditor/Private/CurveEditorViewportClient.h b/Engine/Source/Editor/DistCurveEditor/Private/CurveEditorViewportClient.h
--- a/Engine/Source/Editor/DistCurveEditor/Private/CurveEditorViewportClient.h
+++ b/Engine/Source/Editor/DistCurveEditor/Private/CurveEditorViewportClient.h
@@ -19,10 +19,10 @@ public:
 	~FCurveEditorViewportClient();
 
 	/** FViewportClient interface */
-	virtual void Draw(FViewport* Viewport, FCanvas* Canvas) OVERRIDE;
-	virtual bool InputKey(FViewport* Viewport, int32 ControllerId, FKey Key, EInputEvent Event, float AmountDepressed = 1.0f, bool bGamepad = false) OVERRIDE;
-	virtual void MouseMove(FViewport* Viewport, int32 X, int32 Y) OVERRIDE;
-	virtual bool InputAxis(FViewport* Viewport, int32 ControllerId, FKey Key, float Delta, float DeltaTime, int32 NumSamples = 1, bool bGamepad = false) OVERRIDE;
+	virtual void Draw(FViewport* Viewport, FCanvas* Canvas) override;
+	virtual bool InputKey(FViewport* Viewport, int32 ControllerId, FKey Key, EInputEvent Event, float AmountDepressed = 1.0f, bool bGamepad = false) override;
+	virtual void MouseMove(FViewport* Viewport, int32 X, int32 Y) override;
+	virtual bool InputAxis(FViewport* Viewport, int32 ControllerId, FKey Key, float Delta, float DeltaTime, int32 NumSamples = 1, bool bGamepad = false) override;
 
 	/** Exec handler */
 	void Exec(const TCHAR* Cmd);
diff --git a/Engine/Source/Editor/DistCurveEditor/Private/DistCurveEditorModule.cpp b/Engine/Source/Editor/DistCurveEditor/Private/DistCurveEditorModule.cpp
--- a/Engine/Source/Editor/DistCurveEditor/Private/DistCurveEditorModule.cpp
+++ b/Engine/Source/Editor/DistCurveEditor/Private/DistCurveEditorModule.cpp
@@ -20,21 +20,21 @@ public:
 	}
 
 	/** Called right after the module DLL has been loaded and the module object has been created */
-	virtual void StartupModule() OVERRIDE
+	virtual void StartupModule() override
 	{
 	}
 
 	/** Called before the module is unloaded, right before the module object is destroyed. */
-	virtual void ShutdownModule() OVERRIDE
+	virtual void ShutdownModule() override
 	{
 	}
 
-	virtual TSharedRef<IDistributionCurveEditor> CreateCurveEditorWidget(UInterpCurveEdSetup* EdSetup, FCurveEdNotifyInterface* NotifyObject) OVERRIDE
+	virtual TSharedRef<IDistributionCurveEditor> CreateCurveEditorWidget(UInterpCurveEdSetup* EdSetup, FCurveEdNotifyInterface* NotifyObject) override
 	{
 		return CreateCurveEditorWidget(EdSetup, NotifyObject, IDistributionCurveEditor::FCurveEdOptions());
 	}
 
-	virtual TSharedRef<IDistributionCurveEditor> CreateCurveEditorWidget(UInterpCurveEdSetup* EdSetup, FCurveEdNotifyInterface* NotifyObject, IDistributionCurveEditor::FCurveEdOptions Options) OVERRIDE
+	virtual TSharedRef<IDistributionCurveEditor> CreateCurveEditorWidget(UInterpCurveEdSetup* EdSetup, FCurveEdNotifyInterface* NotifyObject, IDistributionCurveEditor::FCurveEdOptions Options) override
 	{
 		TSharedRef<IDistributionCurveEditor> Widget = SNew(SDistributionCurveEditor).EdSetup(EdSetup).NotifyObject(NotifyObject).CurveEdOptions(Options);
 		return Widget;
diff --git a/Engine/Source/Editor/DistCurveEditor/Private/SDistributionCurveEditor.h b/Engine/Source/Editor/DistCurveEditor/Private/SDistributionCurveEditor.h
--- a/Engine/Source/Editor/DistCurveEditor/Private/SDistributionCurveEditor.h
+++ b/Engine/Source/Editor/DistCurveEditor/Private/SDistributionCurveEditor.h
@@ -49,26 +49,26 @@ public:
 	void Construct(const FArguments& InArgs);
 
 	/** IDistributionCurveEditor interface */
-	virtual void RefreshViewport() OVERRIDE;
-	virtual void CurveChanged() OVERRIDE;
-	virtual void SetCurveVisible(const UObject* InCurve, bool bShow) OVERRIDE;
-	virtual void ClearAllVisibleCurves() OVERRIDE;
-	virtual void SetCurveSelected(const UObject* InCurve, bool bSelected) OVERRIDE;
-	virtual void ClearAllSelectedCurves() OVERRIDE;
-	virtual void ScrollToFirstSelected() OVERRIDE;
-	virtual void SetActiveTabToFirstSelected() OVERRIDE;
-	virtual UInterpCurveEdSetup* GetEdSetup() OVERRIDE;
-	virtual float GetStartIn() OVERRIDE;
-	virtual float GetEndIn() OVERRIDE;
-	virtual void SetPositionMarker(bool bEnabled, float InPosition, const FColor& InMarkerColor) OVERRIDE;
-	virtual void SetEndMarker(bool bEnabled, float InEndPosition) OVERRIDE;
-	virtual void SetRegionMarker(bool bEnabled, float InRegionStart, float InRegionEnd, const FColor& InRegionFillColor) OVERRIDE;
-	virtual void SetInSnap(bool bEnabled, float SnapAmount, bool bInSnapToFrames) OVERRIDE;
-	virtual void SetViewInterval(float StartIn, float EndIn) OVERRIDE;
+	virtual void RefreshViewport() override;
+	virtual void CurveChanged() override;
+	virtual void SetCurveVisible(const UObject* InCurve, bool bShow) override;
+	virtual void ClearAllVisibleCurves() override;
+	virtual void SetCurveSelected(const UObject* InCurve, bool bSelected) override;
+	virtual void ClearAllSelectedCurves() override;
+	virtual void ScrollToFirstSelected() override;
+	virtual void SetActiveTabToFirstSelected() override;
+	virtual UInterpCurveEdSetup* GetEdSetup() override;
+	virtual float GetStartIn() override;
+	virtual float GetEndIn() override;
+	virtual void SetPositionMarker(bool bEnabled, float InPosition, const FColor& InMarkerColor) override;
+	virtual void SetEndMarker(bool bEnabled, float InEndPosition) override;
+	virtual void SetRegionMarker(bool bEnabled, float InRegionStart, float InRegionEnd, const FColor& InRegionFillColor) override;
+	virtual void SetInSnap(bool bEnabled, float SnapAmount, bool bInSnapToFrames) override;
+	virtual void SetViewInterval(float StartIn, float EndIn) override;
 	/** Fits the curve editor view horizontally to the curve data */
-	virtual void FitViewHorizontally() OVERRIDE;
+	virtual void FitViewHorizontally() override;
 	/** Fits the curve editor view vertically to the curve data */
-	virtual void FitViewVertically() OVERRIDE;
+	virtual void FitViewVertically() override;
 
 	/** Accessors */
 	TSharedPtr<FCurveEditorSharedData> GetSharedData();
diff --git a/Engine/Source/Editor/Documentation/Private/Documentation.h b/Engine/Source/Editor/Documentation/Private/Documentation.h
--- a/Engine/Source/Editor/Documentation/Private/Documentation.h
+++ b/Engine/Source/Editor/Documentation/Private/Documentation.h
@@ -13,25 +13,25 @@ public:
 
 	~FDocumentation();
 
-	virtual bool OpenHome() const OVERRIDE;
+	virtual bool OpenHome() const override;
 
-	virtual bool OpenHome(const TSharedRef<FCulture>& Culture) const OVERRIDE;
+	virtual bool OpenHome(const TSharedRef<FCulture>& Culture) const override;
 
-	virtual bool OpenAPIHome() const OVERRIDE;
+	virtual bool OpenAPIHome() const override;
 
-	virtual bool Open( const FString& Link ) const OVERRIDE;
+	virtual bool Open( const FString& Link ) const override;
 
-	virtual bool Open(const FString& Link, const TSharedRef<FCulture>& Culture) const OVERRIDE;
+	virtual bool Open(const FString& Link, const TSharedRef<FCulture>& Culture) const override;
 
-	virtual TSharedRef< SWidget > CreateAnchor( const FString& Link, const FString& PreviewLink = FString(), const FString& PreviewExcerptName = FString() ) const OVERRIDE;
+	virtual TSharedRef< SWidget > CreateAnchor( const FString& Link, const FString& PreviewLink = FString(), const FString& PreviewExcerptName = FString() ) const override;
 
-	virtual TSharedRef< IDocumentationPage > GetPage( const FString& Link, const TSharedPtr< FParserConfiguration >& Config, const FDocumentationStyle& Style = FDocumentationStyle() ) OVERRIDE;
+	virtual TSharedRef< IDocumentationPage > GetPage( const FString& Link, const TSharedPtr< FParserConfiguration >& Config, const FDocumentationStyle& Style = FDocumentationStyle() ) override;
 
-	virtual bool PageExists(const FString& Link) const OVERRIDE;
+	virtual bool PageExists(const FString& Link) const override;
 
-	virtual bool PageExists(const FString& Link, const TSharedRef<FCulture>& Culture) const OVERRIDE;
+	virtual bool PageExists(const FString& Link, const TSharedRef<FCulture>& Culture) const override;
 
-	virtual TSharedRef< class SToolTip > CreateToolTip( const TAttribute<FText>& Text, const TSharedPtr<SWidget>& OverrideContent, const FString& Link, const FString& ExcerptName ) const OVERRIDE;
+	virtual TSharedRef< class SToolTip > CreateToolTip( const TAttribute<FText>& Text, const TSharedPtr<SWidget>& OverrideContent, const FString& Link, const FString& ExcerptName ) const override;
 
 private:
 
diff --git a/Engine/Source/Editor/Documentation/Private/DocumentationModule.cpp b/Engine/Source/Editor/Documentation/Private/DocumentationModule.cpp
--- a/Engine/Source/Editor/Documentation/Private/DocumentationModule.cpp
+++ b/Engine/Source/Editor/Documentation/Private/DocumentationModule.cpp
@@ -9,14 +9,14 @@ class FDocumentationModule : public IDocumentationModule
 public:
 
 	/** IModuleInterface implementation */
-	virtual void StartupModule() OVERRIDE
+	virtual void StartupModule() override
 	{
 		Documentation = FDocumentation::Create();
 
 		FMultiBoxSettings::ToolTipConstructor = FMultiBoxSettings::FConstructToolTip::CreateRaw( this, &FDocumentationModule::ConstructDefaultToolTip );
 	}
 
-	virtual void ShutdownModule() OVERRIDE
+	virtual void ShutdownModule() override
 	{
 		if ( FModuleManager::Get().IsModuleLoaded("Slate") )
 		{
@@ -27,7 +27,7 @@ public:
 		// we call this function before unloading the module.
 	}
 
-	virtual class TSharedRef< IDocumentation > GetDocumentation() const OVERRIDE
+	virtual class TSharedRef< IDocumentation > GetDocumentation() const override
 	{
 		return Documentation.ToSharedRef();
 	}
diff --git a/Engine/Source/Editor/Documentation/Private/DocumentationPage.h b/Engine/Source/Editor/Documentation/Private/DocumentationPage.h
--- a/Engine/Source/Editor/Documentation/Private/DocumentationPage.h
+++ b/Engine/Source/Editor/Documentation/Private/DocumentationPage.h
@@ -14,14 +14,14 @@ public:
 public:
 
 	~FDocumentationPage();
-	virtual bool HasExcerpt( const FString& ExcerptName ) OVERRIDE;
-	virtual int32 GetNumExcerpts() const OVERRIDE;
-	virtual void GetExcerpts( /*OUT*/ TArray< FExcerpt >& Excerpts ) OVERRIDE;
-	virtual bool GetExcerptContent( FExcerpt& Excerpt ) OVERRIDE;
+	virtual bool HasExcerpt( const FString& ExcerptName ) override;
+	virtual int32 GetNumExcerpts() const override;
+	virtual void GetExcerpts( /*OUT*/ TArray< FExcerpt >& Excerpts ) override;
+	virtual bool GetExcerptContent( FExcerpt& Excerpt ) override;
 
-	virtual FText GetTitle() OVERRIDE;
+	virtual FText GetTitle() override;
 
-	virtual void Reload() OVERRIDE;
+	virtual void Reload() override;
 
 private:
 
diff --git a/Engine/Source/Editor/EditorWidgets/Private/SObjectNameEditableTextBox.h b/Engine/Source/Editor/EditorWidgets/Private/SObjectNameEditableTextBox.h
--- a/Engine/Source/Editor/EditorWidgets/Private/SObjectNameEditableTextBox.h
+++ b/Engine/Source/Editor/EditorWidgets/Private/SObjectNameEditableTextBox.h
@@ -27,7 +27,7 @@ protected:
 		* @param  InCurrentTime  Current absolute real time
 		* @param  InDeltaTime  Real time passed since last tick
 		*/
-	virtual void Tick( const FGeometry& AllottedGeometry, const double InCurrentTime, const float InDeltaTime ) OVERRIDE;
+	virtual void Tick( const FGeometry& AllottedGeometry, const double InCurrentTime, const float InDeltaTime ) override;
 
 	/**
 		* The widget should respond by populating the OutDrawElements array with FDrawElements 
@@ -42,7 +42,7 @@ protected:
 		*
 		* @return The maximum layer ID attained by this widget or any of its children.
 		*/
-	virtual int32 OnPaint( const FGeometry& AllottedGeometry, const FSlateRect& MyClippingRect, FSlateWindowElementList& OutDrawElements, int32 LayerId, const FWidgetStyle& InWidgetStyle, bool bParentEnabled ) const OVERRIDE;
+	virtual int32 OnPaint( const FGeometry& AllottedGeometry, const FSlateRect& MyClippingRect, FSlateWindowElementList& OutDrawElements, int32 LayerId, const FWidgetStyle& InWidgetStyle, bool bParentEnabled ) const override;
 
 private:
 	/** Getter for the Text attribute of the editable text inside this widget */
diff --git a/Engine/Source/Editor/EditorWidgets/Public/SAssetDropTarget.h b/Engine/Source/Editor/EditorWidgets/Public/SAssetDropTarget.h
--- a/Engine/Source/Editor/EditorWidgets/Public/SAssetDropTarget.h
+++ b/Engine/Source/Editor/EditorWidgets/Public/SAssetDropTarget.h
@@ -30,10 +30,10 @@ private:
 	FSlateColor GetDropBorderColor() const;
 
 	// SWidget interface
-	virtual FReply OnDragOver( const FGeometry& MyGeometry, const FDragDropEvent& DragDropEvent ) OVERRIDE;
-	virtual FReply OnDrop( const FGeometry& MyGeometry, const FDragDropEvent& DragDropEvent ) OVERRIDE;
-	virtual void OnDragEnter( const FGeometry& MyGeometry, const FDragDropEvent& DragDropEvent ) OVERRIDE;
-	virtual void OnDragLeave( const FDragDropEvent& DragDropEvent ) OVERRIDE;
+	virtual FReply OnDragOver( const FGeometry& MyGeometry, const FDragDropEvent& DragDropEvent ) override;
+	virtual FReply OnDrop( const FGeometry& MyGeometry, const FDragDropEvent& DragDropEvent ) override;
+	virtual void OnDragEnter( const FGeometry& MyGeometry, const FDragDropEvent& DragDropEvent ) override;
+	virtual void OnDragLeave( const FDragDropEvent& DragDropEvent ) override;
 	// End of SWidget interface
 
 	/**
diff --git a/Engine/Source/Editor/EditorWidgets/Public/SAssetSearchBox.h b/Engine/Source/Editor/EditorWidgets/Public/SAssetSearchBox.h
--- a/Engine/Source/Editor/EditorWidgets/Public/SAssetSearchBox.h
+++ b/Engine/Source/Editor/EditorWidgets/Public/SAssetSearchBox.h
@@ -53,11 +53,11 @@ public:
 	void SetText(const TAttribute< FText >& InNewText);
 
 	// SWidget implementation
-	virtual FReply OnPreviewKeyDown( const FGeometry& MyGeometry, const FKeyboardEvent& InKeyboardEvent ) OVERRIDE;
-	virtual FReply OnKeyDown( const FGeometry& MyGeometry, const FKeyboardEvent& InKeyboardEvent ) OVERRIDE;
-	virtual bool SupportsKeyboardFocus() const OVERRIDE;
-	virtual bool HasKeyboardFocus() const OVERRIDE;
-	virtual FReply OnKeyboardFocusReceived( const FGeometry& MyGeometry, const FKeyboardFocusEvent& InKeyboardFocusEvent ) OVERRIDE;
+	virtual FReply OnPreviewKeyDown( const FGeometry& MyGeometry, const FKeyboardEvent& InKeyboardEvent ) override;
+	virtual FReply OnKeyDown( const FGeometry& MyGeometry, const FKeyboardEvent& InKeyboardEvent ) override;
+	virtual bool SupportsKeyboardFocus() const override;
+	virtual bool HasKeyboardFocus() const override;
+	virtual FReply OnKeyboardFocusReceived( const FGeometry& MyGeometry, const FKeyboardFocusEvent& InKeyboardFocusEvent ) override;
 
 private:
 	/** Handler for when text in the editable text box changed */
diff --git a/Engine/Source/Editor/EnvironmentQueryEditor/Classes/EdGraphSchema_EnvironmentQuery.h b/Engine/Source/Editor/EnvironmentQueryEditor/Classes/EdGraphSchema_EnvironmentQuery.h
--- a/Engine/Source/Editor/EnvironmentQueryEditor/Classes/EdGraphSchema_EnvironmentQuery.h
+++ b/Engine/Source/Editor/EnvironmentQueryEditor/Classes/EdGraphSchema_EnvironmentQuery.h
@@ -26,9 +26,9 @@ struct FEnvironmentQuerySchemaAction_NewNode : public FEdGraphSchemaAction
 	{}
 
 	// FEdGraphSchemaAction interface
-	virtual UEdGraphNode* PerformAction(class UEdGraph* ParentGraph, UEdGraphPin* FromPin, const FVector2D Location, bool bSelectNewNode = true) OVERRIDE;
-	virtual UEdGraphNode* PerformAction(class UEdGraph* ParentGraph, TArray<UEdGraphPin*>& FromPins, const FVector2D Location, bool bSelectNewNode = true) OVERRIDE;
-	virtual void AddReferencedObjects( FReferenceCollector& Collector ) OVERRIDE;
+	virtual UEdGraphNode* PerformAction(class UEdGraph* ParentGraph, UEdGraphPin* FromPin, const FVector2D Location, bool bSelectNewNode = true) override;
+	virtual UEdGraphNode* PerformAction(class UEdGraph* ParentGraph, TArray<UEdGraphPin*>& FromPins, const FVector2D Location, bool bSelectNewNode = true) override;
+	virtual void AddReferencedObjects( FReferenceCollector& Collector ) override;
 	// End of FEdGraphSchemaAction interface
 
 	template <typename NodeType>
@@ -66,9 +66,9 @@ struct FEnvironmentQuerySchemaAction_NewSubNode : public FEdGraphSchemaAction
 	{}
 
 	// FEdGraphSchemaAction interface
-	virtual UEdGraphNode* PerformAction(class UEdGraph* ParentGraph, UEdGraphPin* FromPin, const FVector2D Location, bool bSelectNewNode = true) OVERRIDE;
-	virtual UEdGraphNode* PerformAction(class UEdGraph* ParentGraph, TArray<UEdGraphPin*>& FromPins, const FVector2D Location, bool bSelectNewNode = true) OVERRIDE;
-	virtual void AddReferencedObjects( FReferenceCollector& Collector ) OVERRIDE;
+	virtual UEdGraphNode* PerformAction(class UEdGraph* ParentGraph, UEdGraphPin* FromPin, const FVector2D Location, bool bSelectNewNode = true) override;
+	virtual UEdGraphNode* PerformAction(class UEdGraph* ParentGraph, TArray<UEdGraphPin*>& FromPins, const FVector2D Location, bool bSelectNewNode = true) override;
+	virtual void AddReferencedObjects( FReferenceCollector& Collector ) override;
 	// End of FEdGraphSchemaAction interface
 };
 
@@ -80,14 +80,14 @@ class UEdGraphSchema_EnvironmentQuery : public UEdGraphSchema
 	void GetBreakLinkToSubMenuActions(class FMenuBuilder& MenuBuilder, class UEdGraphPin* InGraphPin);
 
 	// Begin EdGraphSchema interface
-	virtual void CreateDefaultNodesForGraph(UEdGraph& Graph) const OVERRIDE;
-	virtual void GetGraphContextActions(FGraphContextMenuBuilder& ContextMenuBuilder) const OVERRIDE;
-	virtual void GetContextMenuActions(const UEdGraph* CurrentGraph, const UEdGraphNode* InGraphNode, const UEdGraphPin* InGraphPin, class FMenuBuilder* MenuBuilder, bool bIsDebugging) const OVERRIDE;
-	virtual const FPinConnectionResponse CanCreateConnection(const UEdGraphPin* A, const UEdGraphPin* B) const OVERRIDE;
-	virtual const FPinConnectionResponse CanMergeNodes(const UEdGraphNode* A, const UEdGraphNode* B) const OVERRIDE;
-	virtual FLinearColor GetPinTypeColor(const FEdGraphPinType& PinType) const OVERRIDE;
-	virtual bool ShouldHidePinDefaultValue(UEdGraphPin* Pin) const OVERRIDE;
-	virtual class FConnectionDrawingPolicy* CreateConnectionDrawingPolicy(int32 InBackLayerID, int32 InFrontLayerID, float InZoomFactor, const FSlateRect& InClippingRect, class FSlateWindowElementList& InDrawElements, class UEdGraph* InGraphObj) const OVERRIDE;
+	virtual void CreateDefaultNodesForGraph(UEdGraph& Graph) const override;
+	virtual void GetGraphContextActions(FGraphContextMenuBuilder& ContextMenuBuilder) const override;
+	virtual void GetContextMenuActions(const UEdGraph* CurrentGraph, const UEdGraphNode* InGraphNode, const UEdGraphPin* InGraphPin, class FMenuBuilder* MenuBuilder, bool bIsDebugging) const override;
+	virtual const FPinConnectionResponse CanCreateConnection(const UEdGraphPin* A, const UEdGraphPin* B) const override;
+	virtual const FPinConnectionResponse CanMergeNodes(const UEdGraphNode* A, const UEdGraphNode* B) const override;
+	virtual FLinearColor GetPinTypeColor(const FEdGraphPinType& PinType) const override;
+	virtual bool ShouldHidePinDefaultValue(UEdGraphPin* Pin) const override;
+	virtual class FConnectionDrawingPolicy* CreateConnectionDrawingPolicy(int32 InBackLayerID, int32 InFrontLayerID, float InZoomFactor, const FSlateRect& InClippingRect, class FSlateWindowElementList& InDrawElements, class UEdGraph* InGraphObj) const override;
 	// End EdGraphSchema interface
 
 	void GetGraphNodeContextActions(FGraphContextMenuBuilder& ContextMenuBuilder) const;
diff --git a/Engine/Source/Editor/EnvironmentQueryEditor/Classes/EnvironmentQueryFactory.h b/Engine/Source/Editor/EnvironmentQueryEditor/Classes/EnvironmentQueryFactory.h
--- a/Engine/Source/Editor/EnvironmentQueryEditor/Classes/EnvironmentQueryFactory.h
+++ b/Engine/Source/Editor/EnvironmentQueryEditor/Classes/EnvironmentQueryFactory.h
@@ -9,7 +9,7 @@ class UEnvironmentQueryFactory : public UFactory
 	GENERATED_UCLASS_BODY()
 
 	// Begin UFactory Interface
-	virtual UObject* FactoryCreateNew(UClass* Class,UObject* InParent,FName Name,EObjectFlags Flags,UObject* Context,FFeedbackContext* Warn) OVERRIDE;
+	virtual UObject* FactoryCreateNew(UClass* Class,UObject* InParent,FName Name,EObjectFlags Flags,UObject* Context,FFeedbackContext* Warn) override;
 	// Begin UFactory Interface	
 };
 
diff --git a/Engine/Source/Editor/EnvironmentQueryEditor/Classes/EnvironmentQueryGraphNode.h b/Engine/Source/Editor/EnvironmentQueryEditor/Classes/EnvironmentQueryGraphNode.h
--- a/Engine/Source/Editor/EnvironmentQueryEditor/Classes/EnvironmentQueryGraphNode.h
+++ b/Engine/Source/Editor/EnvironmentQueryEditor/Classes/EnvironmentQueryGraphNode.h
@@ -18,9 +18,9 @@ class UEnvironmentQueryGraphNode : public UEdGraphNode
 	/** Get the BT graph that owns this node */
 	virtual class UEnvironmentQueryGraph* GetEnvironmentQueryGraph();
 
-	virtual void AutowireNewNode(UEdGraphPin* FromPin) OVERRIDE;
-	virtual void PostEditImport() OVERRIDE;
-	virtual void PrepareForCopying() OVERRIDE;
+	virtual void AutowireNewNode(UEdGraphPin* FromPin) override;
+	virtual void PostEditImport() override;
+	virtual void PrepareForCopying() override;
 	virtual void PostCopyNode();
 
 	// @return the input pin for this state
@@ -30,12 +30,12 @@ class UEnvironmentQueryGraphNode : public UEdGraphNode
 	//
 	virtual UEdGraph* GetBoundGraph() const { return NULL; }
 
-	virtual FText GetNodeTitle(ENodeTitleType::Type TitleType) const OVERRIDE;
+	virtual FText GetNodeTitle(ENodeTitleType::Type TitleType) const override;
 	virtual FText GetDescription() const;
 
-	virtual void NodeConnectionListChanged() OVERRIDE;
+	virtual void NodeConnectionListChanged() override;
 
-	virtual bool CanCreateUnderSpecifiedSchema(const UEdGraphSchema* DesiredSchema) const OVERRIDE;
+	virtual bool CanCreateUnderSpecifiedSchema(const UEdGraphSchema* DesiredSchema) const override;
 
 	virtual bool IsSubNode() const { return false; }
 
diff --git a/Engine/Source/Editor/EnvironmentQueryEditor/Classes/EnvironmentQueryGraphNode_Option.h b/Engine/Source/Editor/EnvironmentQueryEditor/Classes/EnvironmentQueryGraphNode_Option.h
--- a/Engine/Source/Editor/EnvironmentQueryEditor/Classes/EnvironmentQueryGraphNode_Option.h
+++ b/Engine/Source/Editor/EnvironmentQueryEditor/Classes/EnvironmentQueryGraphNode_Option.h
@@ -11,12 +11,12 @@ class UEnvironmentQueryGraphNode_Option : public UEnvironmentQueryGraphNode
 	UPROPERTY()
 	TArray<class UEnvironmentQueryGraphNode_Test*> Tests;
 
-	virtual void AllocateDefaultPins() OVERRIDE;
-	virtual void PostPlacedNewNode() OVERRIDE;
-	virtual void GetContextMenuActions(const FGraphNodeContextMenuBuilder& Context) const OVERRIDE;
-	virtual void PrepareForCopying() OVERRIDE;
-	virtual FText GetNodeTitle(ENodeTitleType::Type TitleType) const OVERRIDE;
-	virtual FText GetDescription() const OVERRIDE;
+	virtual void AllocateDefaultPins() override;
+	virtual void PostPlacedNewNode() override;
+	virtual void GetContextMenuActions(const FGraphNodeContextMenuBuilder& Context) const override;
+	virtual void PrepareForCopying() override;
+	virtual FText GetNodeTitle(ENodeTitleType::Type TitleType) const override;
+	virtual FText GetDescription() const override;
 
 	void AddSubNode(class UEnvironmentQueryGraphNode_Test* NodeTemplate, class UEdGraph* ParentGraph);
 	void CreateAddTestSubMenu(class FMenuBuilder& MenuBuilder, UEdGraph* Graph) const;
@@ -25,5 +25,5 @@ class UEnvironmentQueryGraphNode_Option : public UEnvironmentQueryGraphNode
 
 protected:
 
-	virtual void ResetNodeOwner() OVERRIDE;
+	virtual void ResetNodeOwner() override;
 };
diff --git a/Engine/Source/Editor/EnvironmentQueryEditor/Classes/EnvironmentQueryGraphNode_Root.h b/Engine/Source/Editor/EnvironmentQueryEditor/Classes/EnvironmentQueryGraphNode_Root.h
--- a/Engine/Source/Editor/EnvironmentQueryEditor/Classes/EnvironmentQueryGraphNode_Root.h
+++ b/Engine/Source/Editor/EnvironmentQueryEditor/Classes/EnvironmentQueryGraphNode_Root.h
@@ -7,7 +7,7 @@ class UEnvironmentQueryGraphNode_Root : public UEnvironmentQueryGraphNode
 {
 	GENERATED_UCLASS_BODY()
 
-	virtual void AllocateDefaultPins() OVERRIDE;
-	virtual bool CanUserDeleteNode() const OVERRIDE { return false; }
-	virtual FText GetNodeTitle(ENodeTitleType::Type TitleType) const OVERRIDE;
+	virtual void AllocateDefaultPins() override;
+	virtual bool CanUserDeleteNode() const override { return false; }
+	virtual FText GetNodeTitle(ENodeTitleType::Type TitleType) const override;
 };
diff --git a/Engine/Source/Editor/EnvironmentQueryEditor/Classes/EnvironmentQueryGraphNode_Test.h b/Engine/Source/Editor/EnvironmentQueryEditor/Classes/EnvironmentQueryGraphNode_Test.h
--- a/Engine/Source/Editor/EnvironmentQueryEditor/Classes/EnvironmentQueryGraphNode_Test.h
+++ b/Engine/Source/Editor/EnvironmentQueryEditor/Classes/EnvironmentQueryGraphNode_Test.h
@@ -23,11 +23,11 @@ class UEnvironmentQueryGraphNode_Test : public UEnvironmentQueryGraphNode
 	UPROPERTY()
 	class UEnvironmentQueryGraphNode_Option* ParentNode;
 
-	virtual void PostPlacedNewNode() OVERRIDE;
-	virtual void DestroyNode() OVERRIDE;
-	virtual bool IsSubNode() const OVERRIDE { return true; }
-	virtual FText GetNodeTitle(ENodeTitleType::Type TitleType) const OVERRIDE;
-	virtual FText GetDescription() const OVERRIDE;
+	virtual void PostPlacedNewNode() override;
+	virtual void DestroyNode() override;
+	virtual bool IsSubNode() const override { return true; }
+	virtual FText GetNodeTitle(ENodeTitleType::Type TitleType) const override;
+	virtual FText GetDescription() const override;
 
 	void SetDisplayedWeight(float Pct, bool bNamed);
 };
diff --git a/Engine/Source/Editor/EnvironmentQueryEditor/Private/AssetTypeActions_EnvironmentQuery.h b/Engine/Source/Editor/EnvironmentQueryEditor/Private/AssetTypeActions_EnvironmentQuery.h
--- a/Engine/Source/Editor/EnvironmentQueryEditor/Private/AssetTypeActions_EnvironmentQuery.h
+++ b/Engine/Source/Editor/EnvironmentQueryEditor/Private/AssetTypeActions_EnvironmentQuery.h
@@ -7,13 +7,13 @@ class FAssetTypeActions_EnvironmentQuery : public FAssetTypeActions_Base
 {
 public:
 	// IAssetTypeActions Implementation
-	virtual FText GetName() const OVERRIDE { return NSLOCTEXT("AssetTypeActions", "AssetTypeActions_EnvironmentQuery", "Environment Query"); }
-	virtual FColor GetTypeColor() const OVERRIDE { return FColor(129,50,255); }
-	virtual UClass* GetSupportedClass() const OVERRIDE;
-	virtual bool HasActions ( const TArray<UObject*>& InObjects ) const OVERRIDE { return true; }
-	virtual void GetActions( const TArray<UObject*>& InObjects, FMenuBuilder& MenuBuilder ) OVERRIDE;
-	virtual void OpenAssetEditor( const TArray<UObject*>& InObjects, TSharedPtr<class IToolkitHost> EditWithinLevelEditor = TSharedPtr<IToolkitHost>() ) OVERRIDE;
-	virtual uint32 GetCategories() OVERRIDE { return EAssetTypeCategories::Misc; }
+	virtual FText GetName() const override { return NSLOCTEXT("AssetTypeActions", "AssetTypeActions_EnvironmentQuery", "Environment Query"); }
+	virtual FColor GetTypeColor() const override { return FColor(129,50,255); }
+	virtual UClass* GetSupportedClass() const override;
+	virtual bool HasActions ( const TArray<UObject*>& InObjects ) const override { return true; }
+	virtual void GetActions( const TArray<UObject*>& InObjects, FMenuBuilder& MenuBuilder ) override;
+	virtual void OpenAssetEditor( const TArray<UObject*>& InObjects, TSharedPtr<class IToolkitHost> EditWithinLevelEditor = TSharedPtr<IToolkitHost>() ) override;
+	virtual uint32 GetCategories() override { return EAssetTypeCategories::Misc; }
 
 private:
 	/** Handler for when Edit is selected */
diff --git a/Engine/Source/Editor/EnvironmentQueryEditor/Private/EnvironmentQueryConnectionDrawingPolicy.h b/Engine/Source/Editor/EnvironmentQueryEditor/Private/EnvironmentQueryConnectionDrawingPolicy.h
--- a/Engine/Source/Editor/EnvironmentQueryEditor/Private/EnvironmentQueryConnectionDrawingPolicy.h
+++ b/Engine/Source/Editor/EnvironmentQueryEditor/Private/EnvironmentQueryConnectionDrawingPolicy.h
@@ -14,12 +14,12 @@ public:
 	FEnvironmentQueryConnectionDrawingPolicy(int32 InBackLayerID, int32 InFrontLayerID, float ZoomFactor, const FSlateRect& InClippingRect, FSlateWindowElementList& InDrawElements, UEdGraph* InGraphObj);
 
 	// FConnectionDrawingPolicy interface 
-	virtual void DetermineWiringStyle(UEdGraphPin* OutputPin, UEdGraphPin* InputPin, /*inout*/ float& Thickness, /*inout*/ FLinearColor& WireColor, /*inout*/ bool& bDrawBubbles, /*inout*/ bool& bBidirectional) OVERRIDE;
-	virtual void Draw(TMap<TSharedRef<SWidget>, FArrangedWidget>& PinGeometries, FArrangedChildren& ArrangedNodes) OVERRIDE;
-	virtual void DrawSplineWithArrow(FGeometry& StartGeom, FGeometry& EndGeom, const FLinearColor& WireColor, float WireThickness, bool bDrawBubbles, bool Bidirectional) OVERRIDE;
-	virtual void DrawSplineWithArrow(const FVector2D& StartPoint, const FVector2D& EndPoint, const FLinearColor& WireColor, float WireThickness, bool bDrawBubbles, bool Bidirectional) OVERRIDE;
-	virtual void DrawPreviewConnector(const FGeometry& PinGeometry, const FVector2D& StartPoint, const FVector2D& EndPoint, UEdGraphPin* Pin) OVERRIDE;
-	virtual void DrawConnection(int32 LayerId, const FVector2D& Start, const FVector2D& End, const FLinearColor& InColor, float Thickness, bool bDrawBubbles) OVERRIDE;
+	virtual void DetermineWiringStyle(UEdGraphPin* OutputPin, UEdGraphPin* InputPin, /*inout*/ float& Thickness, /*inout*/ FLinearColor& WireColor, /*inout*/ bool& bDrawBubbles, /*inout*/ bool& bBidirectional) override;
+	virtual void Draw(TMap<TSharedRef<SWidget>, FArrangedWidget>& PinGeometries, FArrangedChildren& ArrangedNodes) override;
+	virtual void DrawSplineWithArrow(FGeometry& StartGeom, FGeometry& EndGeom, const FLinearColor& WireColor, float WireThickness, bool bDrawBubbles, bool Bidirectional) override;
+	virtual void DrawSplineWithArrow(const FVector2D& StartPoint, const FVector2D& EndPoint, const FLinearColor& WireColor, float WireThickness, bool bDrawBubbles, bool Bidirectional) override;
+	virtual void DrawPreviewConnector(const FGeometry& PinGeometry, const FVector2D& StartPoint, const FVector2D& EndPoint, UEdGraphPin* Pin) override;
+	virtual void DrawConnection(int32 LayerId, const FVector2D& Start, const FVector2D& End, const FLinearColor& InColor, float Thickness, bool bDrawBubbles) override;
 	// End of FConnectionDrawingPolicy interface
 
 protected:
diff --git a/Engine/Source/Editor/EnvironmentQueryEditor/Private/EnvironmentQueryEditor.h b/Engine/Source/Editor/EnvironmentQueryEditor/Private/EnvironmentQueryEditor.h
--- a/Engine/Source/Editor/EnvironmentQueryEditor/Private/EnvironmentQueryEditor.h
+++ b/Engine/Source/Editor/EnvironmentQueryEditor/Private/EnvironmentQueryEditor.h
@@ -12,8 +12,8 @@ class FEnvironmentQueryEditor : public IEnvironmentQueryEditor, public FEditorUn
 
 public:
 
-	virtual void RegisterTabSpawners(const TSharedRef<class FTabManager>& TabManager) OVERRIDE;
-	virtual void UnregisterTabSpawners(const TSharedRef<class FTabManager>& TabManager) OVERRIDE;
+	virtual void RegisterTabSpawners(const TSharedRef<class FTabManager>& TabManager) override;
+	virtual void UnregisterTabSpawners(const TSharedRef<class FTabManager>& TabManager) override;
 
 	void InitEnvironmentQueryEditor( const EToolkitMode::Type Mode, const TSharedPtr< class IToolkitHost >& InitToolkitHost, UEnvQuery* Script );
 	FGraphPanelSelectionSet GetSelectedNodes() const;
@@ -22,15 +22,15 @@ public:
 	virtual ~FEnvironmentQueryEditor();
 
 	// Begin IToolkit interface
-	virtual FName GetToolkitFName() const OVERRIDE;
-	virtual FText GetBaseToolkitName() const OVERRIDE;
-	virtual FString GetWorldCentricTabPrefix() const OVERRIDE;
-	virtual FLinearColor GetWorldCentricTabColorScale() const OVERRIDE;
+	virtual FName GetToolkitFName() const override;
+	virtual FText GetBaseToolkitName() const override;
+	virtual FString GetWorldCentricTabPrefix() const override;
+	virtual FLinearColor GetWorldCentricTabColorScale() const override;
 	// End IToolkit interface
 
 	// Begin FEditorUndoClient Interface
-	virtual void	PostUndo(bool bSuccess) OVERRIDE;
-	virtual void	PostRedo(bool bSuccess) OVERRIDE;
+	virtual void	PostUndo(bool bSuccess) override;
+	virtual void	PostRedo(bool bSuccess) override;
 	// End of FEditorUndoClient
 
 	// Delegates
@@ -53,7 +53,7 @@ public:
 
 protected:
 	/** Called when "Save" is clicked for this asset */
-	virtual void SaveAsset_Execute() OVERRIDE;
+	virtual void SaveAsset_Execute() override;
 
 private:
 	/** Create widget for graph editing */
diff --git a/Engine/Source/Editor/EnvironmentQueryEditor/Private/EnvironmentQueryEditorModule.cpp b/Engine/Source/Editor/EnvironmentQueryEditor/Private/EnvironmentQueryEditorModule.cpp
--- a/Engine/Source/Editor/EnvironmentQueryEditor/Private/EnvironmentQueryEditorModule.cpp
+++ b/Engine/Source/Editor/EnvironmentQueryEditor/Private/EnvironmentQueryEditorModule.cpp
@@ -20,7 +20,7 @@ const FName FEnvironmentQueryEditorModule::EnvironmentQueryEditorAppIdentifier(
 
 class FGraphPanelNodeFactory_EnvironmentQuery : public FGraphPanelNodeFactory
 {
-	virtual TSharedPtr<class SGraphNode> CreateNode(UEdGraphNode* Node) const OVERRIDE
+	virtual TSharedPtr<class SGraphNode> CreateNode(UEdGraphNode* Node) const override
 	{
 		if (UEnvironmentQueryGraphNode* EnvQueryNode = Cast<UEnvironmentQueryGraphNode>(Node))
 		{
diff --git a/Engine/Source/Editor/EnvironmentQueryEditor/Private/SGraphNode_EnvironmentQuery.cpp b/Engine/Source/Editor/EnvironmentQueryEditor/Private/SGraphNode_EnvironmentQuery.cpp
--- a/Engine/Source/Editor/EnvironmentQueryEditor/Private/SGraphNode_EnvironmentQuery.cpp
+++ b/Engine/Source/Editor/EnvironmentQueryEditor/Private/SGraphNode_EnvironmentQuery.cpp
@@ -27,8 +27,8 @@ public:
 	void Construct(const FArguments& InArgs, UEdGraphPin* InPin);
 protected:
 	// Begin SGraphPin interface
-	virtual FSlateColor GetPinColor() const OVERRIDE;
-	virtual TSharedRef<SWidget>	GetDefaultValueWidget() OVERRIDE;
+	virtual FSlateColor GetPinColor() const override;
+	virtual TSharedRef<SWidget>	GetDefaultValueWidget() override;
 	// End SGraphPin interface
 
 	const FSlateBrush* GetPinBorder() const;
diff --git a/Engine/Source/Editor/EnvironmentQueryEditor/Private/SGraphNode_EnvironmentQuery.h b/Engine/Source/Editor/EnvironmentQueryEditor/Private/SGraphNode_EnvironmentQuery.h
--- a/Engine/Source/Editor/EnvironmentQueryEditor/Private/SGraphNode_EnvironmentQuery.h
+++ b/Engine/Source/Editor/EnvironmentQueryEditor/Private/SGraphNode_EnvironmentQuery.h
@@ -11,16 +11,16 @@ public:
 	void Construct(const FArguments& InArgs, UEnvironmentQueryGraphNode* InNode);
 
 	// SGraphNode interface
-	virtual void UpdateGraphNode() OVERRIDE;
-	virtual void CreatePinWidgets() OVERRIDE;
-	virtual void AddPin(const TSharedRef<SGraphPin>& PinToAdd) OVERRIDE;
-	virtual TSharedPtr<SToolTip> GetComplexTooltip() OVERRIDE;
-	virtual void OnDragEnter( const FGeometry& MyGeometry, const FDragDropEvent& DragDropEvent ) OVERRIDE;
-	virtual FReply OnDragOver( const FGeometry& MyGeometry, const FDragDropEvent& DragDropEvent ) OVERRIDE;
-	virtual FReply OnDrop( const FGeometry& MyGeometry, const FDragDropEvent& DragDropEvent ) OVERRIDE;
-	virtual void OnDragLeave( const FDragDropEvent& DragDropEvent ) OVERRIDE;
-	virtual FReply OnMouseMove(const FGeometry& SenderGeometry, const FPointerEvent& MouseEvent) OVERRIDE;
-	virtual void Tick( const FGeometry& AllottedGeometry, const double InCurrentTime, const float InDeltaTime ) OVERRIDE;
+	virtual void UpdateGraphNode() override;
+	virtual void CreatePinWidgets() override;
+	virtual void AddPin(const TSharedRef<SGraphPin>& PinToAdd) override;
+	virtual TSharedPtr<SToolTip> GetComplexTooltip() override;
+	virtual void OnDragEnter( const FGeometry& MyGeometry, const FDragDropEvent& DragDropEvent ) override;
+	virtual FReply OnDragOver( const FGeometry& MyGeometry, const FDragDropEvent& DragDropEvent ) override;
+	virtual FReply OnDrop( const FGeometry& MyGeometry, const FDragDropEvent& DragDropEvent ) override;
+	virtual void OnDragLeave( const FDragDropEvent& DragDropEvent ) override;
+	virtual FReply OnMouseMove(const FGeometry& SenderGeometry, const FPointerEvent& MouseEvent) override;
+	virtual void Tick( const FGeometry& AllottedGeometry, const double InCurrentTime, const float InDeltaTime ) override;
 	// End of SGraphNode interface
 
 	/** handle mouse down on the node */
diff --git a/Engine/Source/Editor/EnvironmentQueryEditor/Private/STestFunctionWidget.h b/Engine/Source/Editor/EnvironmentQueryEditor/Private/STestFunctionWidget.h
--- a/Engine/Source/Editor/EnvironmentQueryEditor/Private/STestFunctionWidget.h
+++ b/Engine/Source/Editor/EnvironmentQueryEditor/Private/STestFunctionWidget.h
@@ -15,8 +15,8 @@ public:
 	SLATE_END_ARGS()
 
 	void Construct(const FArguments& InArgs);
-	virtual FVector2D ComputeDesiredSize() const OVERRIDE;
-	virtual int32 OnPaint(const FGeometry& AllottedGeometry, const FSlateRect& MyClippingRect, FSlateWindowElementList& OutDrawElements, int32 LayerId, const FWidgetStyle& InWidgetStyle, bool bParentEnabled) const OVERRIDE;
+	virtual FVector2D ComputeDesiredSize() const override;
+	virtual int32 OnPaint(const FGeometry& AllottedGeometry, const FSlateRect& MyClippingRect, FSlateWindowElementList& OutDrawElements, int32 LayerId, const FWidgetStyle& InWidgetStyle, bool bParentEnabled) const override;
 
 	// normalized test scores for 11 normalized test values (11 samples: y(0), y(0.1), ..., y(1.0))
 	// clamp lines will be fixed at samples x=ClampMinX and x=ClampMaxX
diff --git a/Engine/Source/Editor/FoliageEdit/Private/FoliageEdMode.h b/Engine/Source/Editor/FoliageEdit/Private/FoliageEdMode.h
--- a/Engine/Source/Editor/FoliageEdit/Private/FoliageEdMode.h
+++ b/Engine/Source/Editor/FoliageEdit/Private/FoliageEdMode.h
@@ -154,16 +154,16 @@ public:
 	virtual ~FEdModeFoliage();
 
 	/** FGCObject interface */
-	virtual void AddReferencedObjects( FReferenceCollector& Collector ) OVERRIDE;
+	virtual void AddReferencedObjects( FReferenceCollector& Collector ) override;
 
 	/** FEdMode: Called when the mode is entered */
-	virtual void Enter() OVERRIDE;
+	virtual void Enter() override;
 
 	/** FEdMode: Called when the mode is exited */
-	virtual void Exit() OVERRIDE;
+	virtual void Exit() override;
 
 	/** FEdMode: Called after an Undo operation */
-	virtual void PostUndo() OVERRIDE;
+	virtual void PostUndo() override;
 
 	/** Called when the current level changes */
 	void NotifyNewCurrentLevel();
@@ -187,7 +187,7 @@ public:
 	 *
 	 * @return	true if input was handled
 	 */
-	virtual bool MouseMove( FLevelEditorViewportClient* ViewportClient, FViewport* Viewport, int32 x, int32 y ) OVERRIDE;
+	virtual bool MouseMove( FLevelEditorViewportClient* ViewportClient, FViewport* Viewport, int32 x, int32 y ) override;
 
 	/**
 	 * FEdMode: Called when the mouse is moved while a window input capture is in effect
@@ -199,49 +199,49 @@ public:
 	 *
 	 * @return	true if input was handled
 	 */
-	virtual bool CapturedMouseMove( FLevelEditorViewportClient* InViewportClient, FViewport* InViewport, int32 InMouseX, int32 InMouseY ) OVERRIDE;
+	virtual bool CapturedMouseMove( FLevelEditorViewportClient* InViewportClient, FViewport* InViewport, int32 InMouseX, int32 InMouseY ) override;
 
 	/** FEdMode: Called when a mouse button is pressed */
-	virtual bool StartTracking(FLevelEditorViewportClient* InViewportClient, FViewport* InViewport) OVERRIDE;
+	virtual bool StartTracking(FLevelEditorViewportClient* InViewportClient, FViewport* InViewport) override;
 
 	/** FEdMode: Called when a mouse button is released */
-	virtual bool EndTracking(FLevelEditorViewportClient* InViewportClient, FViewport* InViewport) OVERRIDE;
+	virtual bool EndTracking(FLevelEditorViewportClient* InViewportClient, FViewport* InViewport) override;
 
 	/** FEdMode: Called once per frame */
-	virtual void Tick(FLevelEditorViewportClient* ViewportClient,float DeltaTime) OVERRIDE;
+	virtual void Tick(FLevelEditorViewportClient* ViewportClient,float DeltaTime) override;
 
 	/** FEdMode: Called when a key is pressed */
-	virtual bool InputKey( FLevelEditorViewportClient* InViewportClient, FViewport* InViewport, FKey InKey, EInputEvent InEvent ) OVERRIDE;
+	virtual bool InputKey( FLevelEditorViewportClient* InViewportClient, FViewport* InViewport, FKey InKey, EInputEvent InEvent ) override;
 
 	/** FEdMode: Called when mouse drag input it applied */
-	virtual bool InputDelta( FLevelEditorViewportClient* InViewportClient, FViewport* InViewport, FVector& InDrag, FRotator& InRot, FVector& InScale ) OVERRIDE;
+	virtual bool InputDelta( FLevelEditorViewportClient* InViewportClient, FViewport* InViewport, FVector& InDrag, FRotator& InRot, FVector& InScale ) override;
 
 	/** FEdMode: Render elements for the Foliage tool */
-	virtual void Render( const FSceneView* View, FViewport* Viewport, FPrimitiveDrawInterface* PDI ) OVERRIDE;
+	virtual void Render( const FSceneView* View, FViewport* Viewport, FPrimitiveDrawInterface* PDI ) override;
 
 	/** FEdMode: Render HUD elements for this tool */
-	virtual void DrawHUD( FLevelEditorViewportClient* ViewportClient, FViewport* Viewport, const FSceneView* View, FCanvas* Canvas ) OVERRIDE;
+	virtual void DrawHUD( FLevelEditorViewportClient* ViewportClient, FViewport* Viewport, const FSceneView* View, FCanvas* Canvas ) override;
 
 	/** FEdMode: Handling SelectActor */
-	virtual bool Select( AActor* InActor, bool bInSelected ) OVERRIDE;
+	virtual bool Select( AActor* InActor, bool bInSelected ) override;
 
 	/** FEdMode: Check to see if an actor can be selected in this mode - no side effects */
-	virtual bool IsSelectionAllowed( AActor* InActor, bool bInSelection ) const OVERRIDE;
+	virtual bool IsSelectionAllowed( AActor* InActor, bool bInSelection ) const override;
 
 	/** FEdMode: Called when the currently selected actor has changed */
-	virtual void ActorSelectionChangeNotify() OVERRIDE;
+	virtual void ActorSelectionChangeNotify() override;
 
 	/** Notifies all active modes of mouse click messages. */
 	bool HandleClick(FLevelEditorViewportClient* InViewportClient,  HHitProxy *HitProxy, const FViewportClick &Click );
 
 	/** FEdMode: widget handling */
-	virtual FVector GetWidgetLocation() const OVERRIDE;
+	virtual FVector GetWidgetLocation() const override;
 	virtual bool AllowWidgetMove();
-	virtual bool ShouldDrawWidget() const OVERRIDE;
-	virtual bool UsesTransformWidget() const OVERRIDE ;
-	virtual EAxisList::Type GetWidgetAxisToDraw( FWidget::EWidgetMode InWidgetMode ) const OVERRIDE;
+	virtual bool ShouldDrawWidget() const override;
+	virtual bool UsesTransformWidget() const override ;
+	virtual EAxisList::Type GetWidgetAxisToDraw( FWidget::EWidgetMode InWidgetMode ) const override;
 
-	virtual bool DisallowMouseDeltaTracking() const OVERRIDE;
+	virtual bool DisallowMouseDeltaTracking() const override;
 
 	/** Forces real-time perspective viewports */
 	void ForceRealTimeViewports( const bool bEnable, const bool bStoreCurrentState );
diff --git a/Engine/Source/Editor/FoliageEdit/Private/FoliageEdModeToolkit.h b/Engine/Source/Editor/FoliageEdit/Private/FoliageEdModeToolkit.h
--- a/Engine/Source/Editor/FoliageEdit/Private/FoliageEdModeToolkit.h
+++ b/Engine/Source/Editor/FoliageEdit/Private/FoliageEdModeToolkit.h
@@ -10,17 +10,17 @@
 class FFoliageEdModeToolkit : public FModeToolkit
 {
 public:
-	virtual void RegisterTabSpawners(const TSharedRef<class FTabManager>& TabManager) OVERRIDE;
-	virtual void UnregisterTabSpawners(const TSharedRef<class FTabManager>& TabManager) OVERRIDE;
+	virtual void RegisterTabSpawners(const TSharedRef<class FTabManager>& TabManager) override;
+	virtual void UnregisterTabSpawners(const TSharedRef<class FTabManager>& TabManager) override;
 
 	/** Initializes the foliage mode toolkit */
 	virtual void Init(const TSharedPtr< class IToolkitHost >& InitToolkitHost);
 
 	/** IToolkit interface */
-	virtual FName GetToolkitFName() const OVERRIDE;
-	virtual FText GetBaseToolkitName() const OVERRIDE;
-	virtual class FEdMode* GetEditorMode() const OVERRIDE;
-	virtual TSharedPtr<class SWidget> GetInlineContent() const OVERRIDE;
+	virtual FName GetToolkitFName() const override;
+	virtual FText GetBaseToolkitName() const override;
+	virtual class FEdMode* GetEditorMode() const override;
+	virtual TSharedPtr<class SWidget> GetInlineContent() const override;
 
 	void PostUndo();
 
diff --git a/Engine/Source/Editor/FoliageEdit/Private/SFoliageEdit.cpp b/Engine/Source/Editor/FoliageEdit/Private/SFoliageEdit.cpp
--- a/Engine/Source/Editor/FoliageEdit/Private/SFoliageEdit.cpp
+++ b/Engine/Source/Editor/FoliageEdit/Private/SFoliageEdit.cpp
@@ -47,7 +47,7 @@ public:
 		return FReply::Handled();
 	}
 
-	virtual void OnDragEnter( const FGeometry& MyGeometry, const FDragDropEvent& DragDropEvent ) OVERRIDE
+	virtual void OnDragEnter( const FGeometry& MyGeometry, const FDragDropEvent& DragDropEvent ) override
 	{
 		TSharedPtr<SFoliageEdit> PinnedPtr = FoliageEditPtr.Pin();
 		if (PinnedPtr.IsValid())
@@ -57,7 +57,7 @@ public:
 		SCompoundWidget::OnDragEnter(MyGeometry, DragDropEvent);
 	}
 
-	virtual void OnDragLeave( const FDragDropEvent& DragDropEvent ) OVERRIDE
+	virtual void OnDragLeave( const FDragDropEvent& DragDropEvent ) override
 	{
 		TSharedPtr<SFoliageEdit> PinnedPtr = FoliageEditPtr.Pin();
 		if (PinnedPtr.IsValid())
diff --git a/Engine/Source/Editor/FoliageEdit/Private/SFoliageEditMeshDisplayItem.h b/Engine/Source/Editor/FoliageEdit/Private/SFoliageEditMeshDisplayItem.h
--- a/Engine/Source/Editor/FoliageEdit/Private/SFoliageEditMeshDisplayItem.h
+++ b/Engine/Source/Editor/FoliageEdit/Private/SFoliageEditMeshDisplayItem.h
@@ -66,7 +66,7 @@ public:
 
 private:
 	/** Notification when the cluster details are changed */
-	void NotifyPostChange( const FPropertyChangedEvent& PropertyChangedEvent, UProperty* PropertyThatChanged ) OVERRIDE;
+	void NotifyPostChange( const FPropertyChangedEvent& PropertyChangedEvent, UProperty* PropertyThatChanged ) override;
 
 	/** Callback to get the border color for the item. */
 	FSlateColor GetBorderColor() const;
diff --git a/Engine/Source/Editor/FontEditor/Private/FontEditor.h b/Engine/Source/Editor/FontEditor/Private/FontEditor.h
--- a/Engine/Source/Editor/FontEditor/Private/FontEditor.h
+++ b/Engine/Source/Editor/FontEditor/Private/FontEditor.h
@@ -12,8 +12,8 @@
 class FFontEditor : public IFontEditor, public FGCObject, public FNotifyHook, public FEditorUndoClient
 {
 public:
-	virtual void RegisterTabSpawners(const TSharedRef<class FTabManager>& TabManager) OVERRIDE;
-	virtual void UnregisterTabSpawners(const TSharedRef<class FTabManager>& TabManager) OVERRIDE;
+	virtual void RegisterTabSpawners(const TSharedRef<class FTabManager>& TabManager) override;
+	virtual void UnregisterTabSpawners(const TSharedRef<class FTabManager>& TabManager) override;
 
 	/** Destructor */
 	virtual ~FFontEditor();
@@ -22,20 +22,20 @@ public:
 	void InitFontEditor(const EToolkitMode::Type Mode, const TSharedPtr< class IToolkitHost >& InitToolkitHost, UObject* ObjectToEdit);
 
 	/** IFontEditor interface */
-	virtual UFont* GetFont() const OVERRIDE;
-	virtual void SetSelectedPage(int32 PageIdx) OVERRIDE;
+	virtual UFont* GetFont() const override;
+	virtual void SetSelectedPage(int32 PageIdx) override;
 	
 	/** IToolkit interface */
-	virtual FName GetToolkitFName() const OVERRIDE;
-	virtual FText GetBaseToolkitName() const OVERRIDE;
-	virtual FString GetWorldCentricTabPrefix() const OVERRIDE;
-	virtual FLinearColor GetWorldCentricTabColorScale() const OVERRIDE;
+	virtual FName GetToolkitFName() const override;
+	virtual FText GetBaseToolkitName() const override;
+	virtual FString GetWorldCentricTabPrefix() const override;
+	virtual FLinearColor GetWorldCentricTabColorScale() const override;
 
 	/** FGCObject interface */
-	virtual void AddReferencedObjects(FReferenceCollector& Collector) OVERRIDE;
+	virtual void AddReferencedObjects(FReferenceCollector& Collector) override;
 
 	/** Called to determine if the user should be prompted for a new file if one is missing during an asset reload*/
-	virtual bool ShouldPromptForNewFilesOnReload(const UObject& object) const OVERRIDE;
+	virtual bool ShouldPromptForNewFilesOnReload(const UObject& object) const override;
 
 protected:
 	/** Called when the preview text changes */
@@ -43,13 +43,13 @@ protected:
 	
 	// Begin FEditorUndoClient Interface
 	/** Handles any post undo cleanup of the GUI so that we don't have stale data being displayed. */
-	virtual void PostUndo(bool bSuccess) OVERRIDE;
-	virtual void PostRedo(bool bSuccess) OVERRIDE { PostUndo(bSuccess); }
+	virtual void PostUndo(bool bSuccess) override;
+	virtual void PostRedo(bool bSuccess) override { PostUndo(bSuccess); }
 	// End of FEditorUndoClient
 
 private:
 	/** FNotifyHook interface */
-	virtual void NotifyPostChange( const FPropertyChangedEvent& PropertyChangedEvent, class FEditPropertyChain* PropertyThatChanged) OVERRIDE;
+	virtual void NotifyPostChange( const FPropertyChangedEvent& PropertyChangedEvent, class FEditPropertyChain* PropertyThatChanged) override;
 
 	/** Creates all internal widgets for the tabs to point at */
 	void CreateInternalWidgets();
diff --git a/Engine/Source/Editor/FontEditor/Private/FontEditorModule.cpp b/Engine/Source/Editor/FontEditor/Private/FontEditorModule.cpp
--- a/Engine/Source/Editor/FontEditor/Private/FontEditorModule.cpp
+++ b/Engine/Source/Editor/FontEditor/Private/FontEditorModule.cpp
@@ -23,21 +23,21 @@ public:
 	}
 
 	/** Called right after the module DLL has been loaded and the module object has been created */
-	virtual void StartupModule() OVERRIDE
+	virtual void StartupModule() override
 	{
 		MenuExtensibilityManager = MakeShareable(new FExtensibilityManager);
 		ToolBarExtensibilityManager = MakeShareable(new FExtensibilityManager);
 	}
 
 	/** Called before the module is unloaded, right before the module object is destroyed */
-	virtual void ShutdownModule() OVERRIDE
+	virtual void ShutdownModule() override
 	{
 		MenuExtensibilityManager.Reset();
 		ToolBarExtensibilityManager.Reset();
 	}
 
 	/** Creates a new Font editor */
-	virtual TSharedRef<IFontEditor> CreateFontEditor( const EToolkitMode::Type Mode, const TSharedPtr< IToolkitHost >& InitToolkitHost, UFont* Font ) OVERRIDE
+	virtual TSharedRef<IFontEditor> CreateFontEditor( const EToolkitMode::Type Mode, const TSharedPtr< IToolkitHost >& InitToolkitHost, UFont* Font ) override
 	{
 		TSharedRef<FFontEditor> NewFontEditor(new FFontEditor());
 		NewFontEditor->InitFontEditor(Mode, InitToolkitHost, Font);
diff --git a/Engine/Source/Editor/FontEditor/Private/SFontEditorViewport.cpp b/Engine/Source/Editor/FontEditor/Private/SFontEditorViewport.cpp
--- a/Engine/Source/Editor/FontEditor/Private/SFontEditorViewport.cpp
+++ b/Engine/Source/Editor/FontEditor/Private/SFontEditorViewport.cpp
@@ -19,8 +19,8 @@ public:
 	FFontEditorViewportClient(TWeakPtr<SFontEditorViewport> InFontEditorViewport);
 
 	/** FLevelEditorViewportClient interface */
-	virtual void Draw(FViewport* Viewport, FCanvas* Canvas) OVERRIDE;
-	virtual bool InputKey(FViewport* Viewport, int32 ControllerId, FKey Key, EInputEvent Event, float AmountDepressed = 1.f, bool bGamepad = false) OVERRIDE;
+	virtual void Draw(FViewport* Viewport, FCanvas* Canvas) override;
+	virtual bool InputKey(FViewport* Viewport, int32 ControllerId, FKey Key, EInputEvent Event, float AmountDepressed = 1.f, bool bGamepad = false) override;
 	
 	/** Determines which texture page was selected. */
 	void UpdateSelectedPage(UObject* SelectedObject);
diff --git a/Engine/Source/Editor/GameProjectGeneration/Private/SNewClassDialog.cpp b/Engine/Source/Editor/GameProjectGeneration/Private/SNewClassDialog.cpp
--- a/Engine/Source/Editor/GameProjectGeneration/Private/SNewClassDialog.cpp
+++ b/Engine/Source/Editor/GameProjectGeneration/Private/SNewClassDialog.cpp
@@ -22,7 +22,7 @@ struct FParentClassItem
 class FNativeClassParentFilter : public IClassViewerFilter
 {
 public:
-	virtual bool IsClassAllowed(const FClassViewerInitializationOptions& InInitOptions, const UClass* InClass, TSharedRef< FClassViewerFilterFuncs > InFilterFuncs ) OVERRIDE
+	virtual bool IsClassAllowed(const FClassViewerInitializationOptions& InInitOptions, const UClass* InClass, TSharedRef< FClassViewerFilterFuncs > InFilterFuncs ) override
 	{
 		// You may not make native classes based on blueprint generated classes
 		const bool bIsBlueprintClass = (InClass->ClassGeneratedBy != NULL);
@@ -48,7 +48,7 @@ public:
 		return !bIsBlueprintClass && (!bNeedsAPI || bHasAPI) && !bIsInterface;
 	}
 
-	virtual bool IsUnloadedClassAllowed(const FClassViewerInitializationOptions& InInitOptions, const TSharedRef< const IUnloadedBlueprintData > InUnloadedClassData, TSharedRef< FClassViewerFilterFuncs > InFilterFuncs) OVERRIDE
+	virtual bool IsUnloadedClassAllowed(const FClassViewerInitializationOptions& InInitOptions, const TSharedRef< const IUnloadedBlueprintData > InUnloadedClassData, TSharedRef< FClassViewerFilterFuncs > InFilterFuncs) override
 	{
 		return false;
 	}
diff --git a/Engine/Source/Editor/GameProjectGeneration/Private/SProjectBrowser.h b/Engine/Source/Editor/GameProjectGeneration/Private/SProjectBrowser.h
--- a/Engine/Source/Editor/GameProjectGeneration/Private/SProjectBrowser.h
+++ b/Engine/Source/Editor/GameProjectGeneration/Private/SProjectBrowser.h
@@ -91,7 +91,7 @@ protected:
 	 *
 	 * @return  Returns whether the event was handled, along with other possible actions
 	 */
-	virtual FReply OnKeyDown( const FGeometry& MyGeometry, const FKeyboardEvent& InKeyboardEvent ) OVERRIDE;
+	virtual FReply OnKeyDown( const FGeometry& MyGeometry, const FKeyboardEvent& InKeyboardEvent ) override;
 
 private:
 
diff --git a/Engine/Source/Editor/GameProjectGeneration/Private/SVerbChoiceDialog.h b/Engine/Source/Editor/GameProjectGeneration/Private/SVerbChoiceDialog.h
--- a/Engine/Source/Editor/GameProjectGeneration/Private/SVerbChoiceDialog.h
+++ b/Engine/Source/Editor/GameProjectGeneration/Private/SVerbChoiceDialog.h
@@ -19,10 +19,10 @@ public:
 	void Construct( const FArguments& InArgs );
 
 	/** Override behavior when a key is pressed */
-	virtual	FReply OnKeyDown( const FGeometry& MyGeometry, const FKeyboardEvent& InKeyboardEvent ) OVERRIDE;
+	virtual	FReply OnKeyDown( const FGeometry& MyGeometry, const FKeyboardEvent& InKeyboardEvent ) override;
 
 	/** Override the base method to allow for keyboard focus */
-	virtual bool SupportsKeyboardFocus() const OVERRIDE;
+	virtual bool SupportsKeyboardFocus() const override;
 
 protected:
 	/** Copies the message text to the clipboard. */
diff --git a/Engine/Source/Editor/GameplayAbilitiesEditor/Classes/K2Node_LatentAbilityCall.h b/Engine/Source/Editor/GameplayAbilitiesEditor/Classes/K2Node_LatentAbilityCall.h
--- a/Engine/Source/Editor/GameplayAbilitiesEditor/Classes/K2Node_LatentAbilityCall.h
+++ b/Engine/Source/Editor/GameplayAbilitiesEditor/Classes/K2Node_LatentAbilityCall.h
@@ -9,6 +9,6 @@ class UK2Node_LatentAbilityCall : public UK2Node_BaseAsyncTask
 	GENERATED_UCLASS_BODY()
 
 	// UEdGraphNode interface
-	virtual void GetMenuEntries(FGraphContextMenuBuilder& ContextMenuBuilder) const OVERRIDE;
+	virtual void GetMenuEntries(FGraphContextMenuBuilder& ContextMenuBuilder) const override;
 	// End of UEdGraphNode interface
 };
diff --git a/Engine/Source/Editor/GameplayAbilitiesEditor/Private/AttributeDetails.h b/Engine/Source/Editor/GameplayAbilitiesEditor/Private/AttributeDetails.h
--- a/Engine/Source/Editor/GameplayAbilitiesEditor/Private/AttributeDetails.h
+++ b/Engine/Source/Editor/GameplayAbilitiesEditor/Private/AttributeDetails.h
@@ -18,7 +18,7 @@ public:
 
 private:
 	/** IDetailCustomization interface */
-	virtual void CustomizeDetails(IDetailLayoutBuilder& DetailLayout) OVERRIDE;
+	virtual void CustomizeDetails(IDetailLayoutBuilder& DetailLayout) override;
 
 	TSharedPtr<FString> GetPropertyType() const;
 
@@ -36,8 +36,8 @@ public:
 	static TSharedRef<IPropertyTypeCustomization> MakeInstance();
 
 	/** IPropertyTypeCustomization interface */
-	virtual void CustomizeHeader( TSharedRef<class IPropertyHandle> StructPropertyHandle, class FDetailWidgetRow& HeaderRow, IPropertyTypeCustomizationUtils& StructCustomizationUtils ) OVERRIDE;
-	virtual void CustomizeChildren( TSharedRef<class IPropertyHandle> StructPropertyHandle, class IDetailChildrenBuilder& StructBuilder, IPropertyTypeCustomizationUtils& StructCustomizationUtils ) OVERRIDE;              
+	virtual void CustomizeHeader( TSharedRef<class IPropertyHandle> StructPropertyHandle, class FDetailWidgetRow& HeaderRow, IPropertyTypeCustomizationUtils& StructCustomizationUtils ) override;
+	virtual void CustomizeChildren( TSharedRef<class IPropertyHandle> StructPropertyHandle, class IDetailChildrenBuilder& StructBuilder, IPropertyTypeCustomizationUtils& StructCustomizationUtils ) override;              
 
 private:
 
@@ -57,8 +57,8 @@ public:
 
 private:
 
-	virtual void CustomizeHeader( TSharedRef<class IPropertyHandle> StructPropertyHandle, class FDetailWidgetRow& HeaderRow, IPropertyTypeCustomizationUtils& StructCustomizationUtils ) OVERRIDE;
-	virtual void CustomizeChildren( TSharedRef<class IPropertyHandle> StructPropertyHandle, class IDetailChildrenBuilder& StructBuilder, IPropertyTypeCustomizationUtils& StructCustomizationUtils ) OVERRIDE;              
+	virtual void CustomizeHeader( TSharedRef<class IPropertyHandle> StructPropertyHandle, class FDetailWidgetRow& HeaderRow, IPropertyTypeCustomizationUtils& StructCustomizationUtils ) override;
+	virtual void CustomizeChildren( TSharedRef<class IPropertyHandle> StructPropertyHandle, class IDetailChildrenBuilder& StructBuilder, IPropertyTypeCustomizationUtils& StructCustomizationUtils ) override;              
 
 	bool IsEditable( ) const;
 
diff --git a/Engine/Source/Editor/GameplayTagsEditor/Classes/GameplayTagsK2Node_LiteralGameplayTag.h b/Engine/Source/Editor/GameplayTagsEditor/Classes/GameplayTagsK2Node_LiteralGameplayTag.h
--- a/Engine/Source/Editor/GameplayTagsEditor/Classes/GameplayTagsK2Node_LiteralGameplayTag.h
+++ b/Engine/Source/Editor/GameplayTagsEditor/Classes/GameplayTagsK2Node_LiteralGameplayTag.h
@@ -12,19 +12,19 @@ class UGameplayTagsK2Node_LiteralGameplayTag : public UK2Node
 
 #if WITH_EDITOR
 	// Begin UEdGraphNode interface
-	virtual void AllocateDefaultPins() OVERRIDE;
-	virtual FLinearColor GetNodeTitleColor() const OVERRIDE;
-	virtual FText GetNodeTitle(ENodeTitleType::Type TitleType) const OVERRIDE;
-	virtual bool CanDuplicateNode() const OVERRIDE { return false; }
-	virtual bool NodeCausesStructuralBlueprintChange() const OVERRIDE { return true; }
-	virtual bool CanCreateUnderSpecifiedSchema(const UEdGraphSchema* Schema) const OVERRIDE;
-	virtual void GetMenuEntries(FGraphContextMenuBuilder& ContextMenuBuilder) const OVERRIDE;
+	virtual void AllocateDefaultPins() override;
+	virtual FLinearColor GetNodeTitleColor() const override;
+	virtual FText GetNodeTitle(ENodeTitleType::Type TitleType) const override;
+	virtual bool CanDuplicateNode() const override { return false; }
+	virtual bool NodeCausesStructuralBlueprintChange() const override { return true; }
+	virtual bool CanCreateUnderSpecifiedSchema(const UEdGraphSchema* Schema) const override;
+	virtual void GetMenuEntries(FGraphContextMenuBuilder& ContextMenuBuilder) const override;
 	// End UEdGraphNode interface
 
 	// Begin UK2Node interface
-	virtual bool ShouldShowNodeProperties() const OVERRIDE { return true; }
-	virtual bool IsNodeSafeToIgnore() const OVERRIDE { return true; }
-	virtual void ExpandNode(class FKismetCompilerContext& CompilerContext, UEdGraph* SourceGraph) OVERRIDE;
+	virtual bool ShouldShowNodeProperties() const override { return true; }
+	virtual bool IsNodeSafeToIgnore() const override { return true; }
+	virtual void ExpandNode(class FKismetCompilerContext& CompilerContext, UEdGraph* SourceGraph) override;
 	// End UK2Node interface
 #endif
 
diff --git a/Engine/Source/Editor/GameplayTagsEditor/Private/GameplayTagContainerCustomization.h b/Engine/Source/Editor/GameplayTagsEditor/Private/GameplayTagContainerCustomization.h
--- a/Engine/Source/Editor/GameplayTagsEditor/Private/GameplayTagContainerCustomization.h
+++ b/Engine/Source/Editor/GameplayTagsEditor/Private/GameplayTagContainerCustomization.h
@@ -17,14 +17,14 @@ public:
 	~FGameplayTagContainerCustomization();
 
 	/** Overridden to show an edit button to launch the gameplay tag editor */
-	virtual void CustomizeHeader(TSharedRef<class IPropertyHandle> StructPropertyHandle, class FDetailWidgetRow& HeaderRow, IPropertyTypeCustomizationUtils& StructCustomizationUtils) OVERRIDE;
+	virtual void CustomizeHeader(TSharedRef<class IPropertyHandle> StructPropertyHandle, class FDetailWidgetRow& HeaderRow, IPropertyTypeCustomizationUtils& StructCustomizationUtils) override;
 	
 	/** Overridden to do nothing */
-	virtual void CustomizeChildren(TSharedRef<IPropertyHandle> InStructPropertyHandle, class IDetailChildrenBuilder& ChildBuilder, IPropertyTypeCustomizationUtils& StructCustomizationUtils) OVERRIDE {}
+	virtual void CustomizeChildren(TSharedRef<IPropertyHandle> InStructPropertyHandle, class IDetailChildrenBuilder& ChildBuilder, IPropertyTypeCustomizationUtils& StructCustomizationUtils) override {}
 
 	// Begin FEditorUndoClient Interface
-	virtual void PostUndo( bool bSuccess ) OVERRIDE;
-	virtual void PostRedo( bool bSuccess ) OVERRIDE;	
+	virtual void PostUndo( bool bSuccess ) override;
+	virtual void PostRedo( bool bSuccess ) override;	
 	// End FEditorUndoClient Interface
 
 private:
diff --git a/Engine/Source/Editor/GameplayTagsEditor/Private/GameplayTagsGraphPanelPinFactory.h b/Engine/Source/Editor/GameplayTagsEditor/Private/GameplayTagsGraphPanelPinFactory.h
--- a/Engine/Source/Editor/GameplayTagsEditor/Private/GameplayTagsGraphPanelPinFactory.h
+++ b/Engine/Source/Editor/GameplayTagsEditor/Private/GameplayTagsGraphPanelPinFactory.h
@@ -9,7 +9,7 @@
 
 class FGameplayTagsGraphPanelPinFactory: public FGraphPanelPinFactory
 {
-	virtual TSharedPtr<class SGraphPin> CreatePin(class UEdGraphPin* InPin) const OVERRIDE
+	virtual TSharedPtr<class SGraphPin> CreatePin(class UEdGraphPin* InPin) const override
 	{
 		const UEdGraphSchema_K2* K2Schema = GetDefault<UEdGraphSchema_K2>();
 		if (InPin->PinType.PinCategory == K2Schema->PC_Struct && InPin->PinType.PinSubCategoryObject == FGameplayTag::StaticStruct())
diff --git a/Engine/Source/Editor/GameplayTagsEditor/Private/SGameplayTagContainerGraphPin.h b/Engine/Source/Editor/GameplayTagsEditor/Private/SGameplayTagContainerGraphPin.h
--- a/Engine/Source/Editor/GameplayTagsEditor/Private/SGameplayTagContainerGraphPin.h
+++ b/Engine/Source/Editor/GameplayTagsEditor/Private/SGameplayTagContainerGraphPin.h
@@ -15,7 +15,7 @@ public:
 	void Construct(const FArguments& InArgs, UEdGraphPin* InGraphPinObj);
 
 	// Begin SGraphPin interface
-	virtual TSharedRef<SWidget>	GetDefaultValueWidget() OVERRIDE;
+	virtual TSharedRef<SWidget>	GetDefaultValueWidget() override;
 	// End SGraphPin interface
 
 private:
diff --git a/Engine/Source/Editor/GameplayTagsEditor/Private/SGameplayTagWidget.h b/Engine/Source/Editor/GameplayTagsEditor/Private/SGameplayTagWidget.h
--- a/Engine/Source/Editor/GameplayTagsEditor/Private/SGameplayTagWidget.h
+++ b/Engine/Source/Editor/GameplayTagsEditor/Private/SGameplayTagWidget.h
@@ -44,7 +44,7 @@ public:
 		struct FGameplayTagContainer* TagContainer; 
 
 		/** Overridden to emit tag container owner reference */
-		virtual void AddReferencedObjects(FReferenceCollector& Collector) OVERRIDE
+		virtual void AddReferencedObjects(FReferenceCollector& Collector) override
 		{
 			Collector.AddReferencedObject(TagContainerOwner);
 		}
diff --git a/Engine/Source/Editor/GameplayTagsEditor/Public/AssetTypeActions_GameplayTagAssetBase.h b/Engine/Source/Editor/GameplayTagsEditor/Public/AssetTypeActions_GameplayTagAssetBase.h
--- a/Engine/Source/Editor/GameplayTagsEditor/Public/AssetTypeActions_GameplayTagAssetBase.h
+++ b/Engine/Source/Editor/GameplayTagsEditor/Public/AssetTypeActions_GameplayTagAssetBase.h
@@ -13,13 +13,13 @@ public:
 	FAssetTypeActions_GameplayTagAssetBase(FName InTagPropertyName);
 
 	/** Overridden to specify that the gameplay tag base has actions */
-	virtual bool HasActions(const TArray<UObject*>& InObjects) const OVERRIDE;
+	virtual bool HasActions(const TArray<UObject*>& InObjects) const override;
 
 	/** Overridden to offer the gameplay tagging options */
-	virtual void GetActions(const TArray<UObject*>& InObjects, FMenuBuilder& MenuBuilder) OVERRIDE;
+	virtual void GetActions(const TArray<UObject*>& InObjects, FMenuBuilder& MenuBuilder) override;
 
 	/** Overridden to specify misc category */
-	virtual uint32 GetCategories() OVERRIDE;
+	virtual uint32 GetCategories() override;
 
 private:
 	/**
diff --git a/Engine/Source/Editor/GeometryMode/Classes/GeomModifier_Clip.h b/Engine/Source/Editor/GeometryMode/Classes/GeomModifier_Clip.h
--- a/Engine/Source/Editor/GeometryMode/Classes/GeomModifier_Clip.h
+++ b/Engine/Source/Editor/GeometryMode/Classes/GeomModifier_Clip.h
@@ -29,14 +29,14 @@ class UGeomModifier_Clip : public UGeomModifier_Edit
 
 
 	// Begin UGeomModifier Interface
-	virtual bool Supports() OVERRIDE;
-	virtual bool InputKey(class FLevelEditorViewportClient* ViewportClient, FViewport* Viewport, FKey Key, EInputEvent Event) OVERRIDE;
-	virtual void Render(const FSceneView* View,FViewport* Viewport,FPrimitiveDrawInterface* PDI) OVERRIDE;
-	virtual void DrawHUD(FLevelEditorViewportClient* ViewportClient,FViewport* Viewport,const FSceneView* View,FCanvas* Canvas) OVERRIDE;
-	virtual void Tick(FLevelEditorViewportClient* ViewportClient,float DeltaTime) OVERRIDE;
-	virtual void WasActivated() OVERRIDE;
+	virtual bool Supports() override;
+	virtual bool InputKey(class FLevelEditorViewportClient* ViewportClient, FViewport* Viewport, FKey Key, EInputEvent Event) override;
+	virtual void Render(const FSceneView* View,FViewport* Viewport,FPrimitiveDrawInterface* PDI) override;
+	virtual void DrawHUD(FLevelEditorViewportClient* ViewportClient,FViewport* Viewport,const FSceneView* View,FCanvas* Canvas) override;
+	virtual void Tick(FLevelEditorViewportClient* ViewportClient,float DeltaTime) override;
+	virtual void WasActivated() override;
 protected:
-	virtual bool OnApply() OVERRIDE;
+	virtual bool OnApply() override;
 	// End UGeomModifier Interface
 
 private:
diff --git a/Engine/Source/Editor/GeometryMode/Classes/GeomModifier_Create.h b/Engine/Source/Editor/GeometryMode/Classes/GeomModifier_Create.h
--- a/Engine/Source/Editor/GeometryMode/Classes/GeomModifier_Create.h
+++ b/Engine/Source/Editor/GeometryMode/Classes/GeomModifier_Create.h
@@ -15,9 +15,9 @@ class UGeomModifier_Create : public UGeomModifier_Edit
 
 
 	// Begin UGeomModifier Interface
-	virtual bool Supports() OVERRIDE;
+	virtual bool Supports() override;
 protected:
-	virtual bool OnApply() OVERRIDE;
+	virtual bool OnApply() override;
 	// End UGeomModifier Interface
 };
 
diff --git a/Engine/Source/Editor/GeometryMode/Classes/GeomModifier_Delete.h b/Engine/Source/Editor/GeometryMode/Classes/GeomModifier_Delete.h
--- a/Engine/Source/Editor/GeometryMode/Classes/GeomModifier_Delete.h
+++ b/Engine/Source/Editor/GeometryMode/Classes/GeomModifier_Delete.h
@@ -15,9 +15,9 @@ class UGeomModifier_Delete : public UGeomModifier_Edit
 
 
 	// Begin UGeomModifier Interface
-	virtual bool Supports() OVERRIDE;
+	virtual bool Supports() override;
 protected:
-	virtual bool OnApply() OVERRIDE;
+	virtual bool OnApply() override;
 	// End UGeomModifier Interface
 };
 
diff --git a/Engine/Source/Editor/GeometryMode/Classes/GeomModifier_Edit.h b/Engine/Source/Editor/GeometryMode/Classes/GeomModifier_Edit.h
--- a/Engine/Source/Editor/GeometryMode/Classes/GeomModifier_Edit.h
+++ b/Engine/Source/Editor/GeometryMode/Classes/GeomModifier_Edit.h
@@ -21,7 +21,7 @@ class UGeomModifier_Edit : public UGeomModifier
 
 
 	// Begin UGeomModifier Interface
-	virtual bool InputDelta(class FLevelEditorViewportClient* InViewportClient,FViewport* InViewport,FVector& InDrag,FRotator& InRot,FVector& InScale) OVERRIDE;	
+	virtual bool InputDelta(class FLevelEditorViewportClient* InViewportClient,FViewport* InViewport,FVector& InDrag,FRotator& InRot,FVector& InScale) override;	
 	// End UGeomModifier Interface
 };
 
diff --git a/Engine/Source/Editor/GeometryMode/Classes/GeomModifier_Extrude.h b/Engine/Source/Editor/GeometryMode/Classes/GeomModifier_Extrude.h
--- a/Engine/Source/Editor/GeometryMode/Classes/GeomModifier_Extrude.h
+++ b/Engine/Source/Editor/GeometryMode/Classes/GeomModifier_Extrude.h
@@ -24,16 +24,16 @@ class UGeomModifier_Extrude : public UGeomModifier_Edit
 
 
 	// Begin UGeomModifier Interface
-	virtual bool Supports() OVERRIDE;
-	virtual void Initialize() OVERRIDE;
-	virtual void WasActivated() OVERRIDE;
-	virtual void WasDeactivated() OVERRIDE;
+	virtual bool Supports() override;
+	virtual void Initialize() override;
+	virtual void WasActivated() override;
+	virtual void WasDeactivated() override;
 
 	/* Check the coordinates mode is local and warn the user with a suppressible dialog if it is not */
 	void CheckCoordinatesMode();
 
 protected:
-	virtual bool OnApply() OVERRIDE;
+	virtual bool OnApply() override;
 	// End UGeomModifier Interface
 private:
 	void Apply(int32 InLength, int32 InSegments);
diff --git a/Engine/Source/Editor/GeometryMode/Classes/GeomModifier_Lathe.h b/Engine/Source/Editor/GeometryMode/Classes/GeomModifier_Lathe.h
--- a/Engine/Source/Editor/GeometryMode/Classes/GeomModifier_Lathe.h
+++ b/Engine/Source/Editor/GeometryMode/Classes/GeomModifier_Lathe.h
@@ -28,10 +28,10 @@ class UGeomModifier_Lathe : public UGeomModifier_Edit
 
 
 	// Begin UGeomModifier Interface
-	virtual bool Supports() OVERRIDE;
-	virtual void Initialize() OVERRIDE;
+	virtual bool Supports() override;
+	virtual void Initialize() override;
 protected:
-	virtual bool OnApply() OVERRIDE;
+	virtual bool OnApply() override;
 	// End UGeomModifier Interface
 private:
 	void Apply( int32 InTotalSegments, int32 InSegments, EAxis::Type InAxis );
diff --git a/Engine/Source/Editor/GeometryMode/Classes/GeomModifier_Pen.h b/Engine/Source/Editor/GeometryMode/Classes/GeomModifier_Pen.h
--- a/Engine/Source/Editor/GeometryMode/Classes/GeomModifier_Pen.h
+++ b/Engine/Source/Editor/GeometryMode/Classes/GeomModifier_Pen.h
@@ -41,13 +41,13 @@ class UGeomModifier_Pen : public UGeomModifier_Edit
 	FLevelEditorViewportClient*		UsingViewportClient;
 
 	// Begin UGeomModifier Interface
-	virtual bool InputKey(class FLevelEditorViewportClient* ViewportClient, FViewport* Viewport, FKey Key, EInputEvent Event) OVERRIDE;
-	virtual void Render(const FSceneView* View,FViewport* Viewport,FPrimitiveDrawInterface* PDI) OVERRIDE;
-	virtual void DrawHUD(FLevelEditorViewportClient* ViewportClient,FViewport* Viewport,const FSceneView* View,FCanvas* Canvas) OVERRIDE;
-	virtual void Tick(FLevelEditorViewportClient* ViewportClient,float DeltaTime) OVERRIDE;
-	virtual void WasActivated() OVERRIDE;
+	virtual bool InputKey(class FLevelEditorViewportClient* ViewportClient, FViewport* Viewport, FKey Key, EInputEvent Event) override;
+	virtual void Render(const FSceneView* View,FViewport* Viewport,FPrimitiveDrawInterface* PDI) override;
+	virtual void DrawHUD(FLevelEditorViewportClient* ViewportClient,FViewport* Viewport,const FSceneView* View,FCanvas* Canvas) override;
+	virtual void Tick(FLevelEditorViewportClient* ViewportClient,float DeltaTime) override;
+	virtual void WasActivated() override;
 protected:
-	virtual bool OnApply() OVERRIDE;
+	virtual bool OnApply() override;
 	// End UGeomModifier Interface
 private:
 	void Apply();
diff --git a/Engine/Source/Editor/GeometryMode/Private/GeometryMode.h b/Engine/Source/Editor/GeometryMode/Private/GeometryMode.h
--- a/Engine/Source/Editor/GeometryMode/Private/GeometryMode.h
+++ b/Engine/Source/Editor/GeometryMode/Private/GeometryMode.h
@@ -81,17 +81,17 @@ private:
 class FGeometryMode : public FModeToolkit
 {
 public:
-	virtual void RegisterTabSpawners(const TSharedRef<class FTabManager>& TabManager) OVERRIDE;
-	virtual void UnregisterTabSpawners(const TSharedRef<class FTabManager>& TabManager) OVERRIDE;
+	virtual void RegisterTabSpawners(const TSharedRef<class FTabManager>& TabManager) override;
+	virtual void UnregisterTabSpawners(const TSharedRef<class FTabManager>& TabManager) override;
 
 	/** Initializes the geometry mode toolkit */
 	virtual void Init(const TSharedPtr< class IToolkitHost >& InitToolkitHost);
 
 	/** IToolkit interface */
-	virtual FName GetToolkitFName() const OVERRIDE;
-	virtual FText GetBaseToolkitName() const OVERRIDE;
-	virtual class FEdMode* GetEditorMode() const OVERRIDE;
-	virtual TSharedPtr<class SWidget> GetInlineContent() const OVERRIDE;
+	virtual FName GetToolkitFName() const override;
+	virtual FText GetBaseToolkitName() const override;
+	virtual class FEdMode* GetEditorMode() const override;
+	virtual TSharedPtr<class SWidget> GetInlineContent() const override;
 
 	/** Method called when the selection */
 	virtual void SelectionChanged();
diff --git a/Engine/Source/Editor/GeometryMode/Public/EditorGeometry.h b/Engine/Source/Editor/GeometryMode/Public/EditorGeometry.h
--- a/Engine/Source/Editor/GeometryMode/Public/EditorGeometry.h
+++ b/Engine/Source/Editor/GeometryMode/Public/EditorGeometry.h
@@ -300,7 +300,7 @@ public:
 	 */
 	void ForceLastSelectionIndex( int32 InLastSelectionIndex ) { LastSelectionIndex = InLastSelectionIndex; }
 
-	virtual void AddReferencedObjects( FReferenceCollector& Collector ) OVERRIDE;
+	virtual void AddReferencedObjects( FReferenceCollector& Collector ) override;
 	
 	/**
 	 * Set the pivot position based on the 'highest' selected object (vertex/edge/poly) in the given selection array
diff --git a/Engine/Source/Editor/GeometryMode/Public/GeometryEdMode.h b/Engine/Source/Editor/GeometryMode/Public/GeometryEdMode.h
--- a/Engine/Source/Editor/GeometryMode/Public/GeometryEdMode.h
+++ b/Engine/Source/Editor/GeometryMode/Public/GeometryEdMode.h
@@ -12,8 +12,8 @@ private:
 	TSharedPtr<class FEdModeGeometry> EdModeGeometry;
 public:
 	// IModuleInterface
-	virtual void StartupModule() OVERRIDE;
-	virtual void ShutdownModule() OVERRIDE;
+	virtual void StartupModule() override;
+	virtual void ShutdownModule() override;
 	// End of IModuleInterface
 };
 
@@ -43,22 +43,22 @@ public:
 	virtual ~FEdModeGeometry();
 
 	// FEdMode interface
-	virtual void Render(const FSceneView* View,FViewport* Viewport,FPrimitiveDrawInterface* PDI) OVERRIDE;
-	virtual bool ShowModeWidgets() const OVERRIDE;
-	virtual bool UsesToolkits() const OVERRIDE;
-	virtual bool ShouldDrawBrushWireframe( AActor* InActor ) const OVERRIDE;
-	virtual bool GetCustomDrawingCoordinateSystem( FMatrix& InMatrix, void* InData ) OVERRIDE;
-	virtual bool GetCustomInputCoordinateSystem( FMatrix& InMatrix, void* InData ) OVERRIDE;
-	virtual void Enter() OVERRIDE;
-	virtual void Exit() OVERRIDE;
-	virtual void ActorSelectionChangeNotify() OVERRIDE;
-	virtual void MapChangeNotify() OVERRIDE;
-	virtual void SelectionChanged() OVERRIDE;
+	virtual void Render(const FSceneView* View,FViewport* Viewport,FPrimitiveDrawInterface* PDI) override;
+	virtual bool ShowModeWidgets() const override;
+	virtual bool UsesToolkits() const override;
+	virtual bool ShouldDrawBrushWireframe( AActor* InActor ) const override;
+	virtual bool GetCustomDrawingCoordinateSystem( FMatrix& InMatrix, void* InData ) override;
+	virtual bool GetCustomInputCoordinateSystem( FMatrix& InMatrix, void* InData ) override;
+	virtual void Enter() override;
+	virtual void Exit() override;
+	virtual void ActorSelectionChangeNotify() override;
+	virtual void MapChangeNotify() override;
+	virtual void SelectionChanged() override;
 	virtual FVector GetWidgetLocation() const;
 	// End of FEdMode interface
 
 	// FGCObject interface
-	virtual void AddReferencedObjects( FReferenceCollector& Collector ) OVERRIDE;
+	virtual void AddReferencedObjects( FReferenceCollector& Collector ) override;
 	// End of FGCObject interface
 
 	void UpdateModifierWindow();
@@ -245,7 +245,7 @@ public:
 	virtual void Tick(FLevelEditorViewportClient* ViewportClient,float DeltaTime);
 
 	/** @return		true if the key was handled by this editor mode tool. */
-	virtual bool InputKey(FLevelEditorViewportClient* ViewportClient, FViewport* Viewport, FKey Key, EInputEvent Event) OVERRIDE;
+	virtual bool InputKey(FLevelEditorViewportClient* ViewportClient, FViewport* Viewport, FKey Key, EInputEvent Event) override;
 
 	virtual void Render(const FSceneView* View,FViewport* Viewport,FPrimitiveDrawInterface* PDI);
 	virtual void DrawHUD(FLevelEditorViewportClient* ViewportClient,FViewport* Viewport,const FSceneView* View,FCanvas* Canvas);
diff --git a/Engine/Source/Editor/GraphEditor/Private/AnimationNodes/SGraphNodeAnimationResult.h b/Engine/Source/Editor/GraphEditor/Private/AnimationNodes/SGraphNodeAnimationResult.h
--- a/Engine/Source/Editor/GraphEditor/Private/AnimationNodes/SGraphNodeAnimationResult.h
+++ b/Engine/Source/Editor/GraphEditor/Private/AnimationNodes/SGraphNodeAnimationResult.h
@@ -12,6 +12,6 @@ public:
 
 protected:
 	// SGraphNode interface
-	virtual TSharedRef<SWidget> CreateNodeContentArea() OVERRIDE;
+	virtual TSharedRef<SWidget> CreateNodeContentArea() override;
 	// End of SGraphNode interface
 };
diff --git a/Engine/Source/Editor/GraphEditor/Private/AnimationNodes/SGraphNodeSequencePlayer.h b/Engine/Source/Editor/GraphEditor/Private/AnimationNodes/SGraphNodeSequencePlayer.h
--- a/Engine/Source/Editor/GraphEditor/Private/AnimationNodes/SGraphNodeSequencePlayer.h
+++ b/Engine/Source/Editor/GraphEditor/Private/AnimationNodes/SGraphNodeSequencePlayer.h
@@ -11,12 +11,12 @@ public:
 	void Construct(const FArguments& InArgs, UK2Node* InNode);
 
 	// SNodePanel::SNode interface
-	virtual void GetNodeInfoPopups(FNodeInfoContext* Context, TArray<FGraphInformationPopupInfo>& Popups) const OVERRIDE;
+	virtual void GetNodeInfoPopups(FNodeInfoContext* Context, TArray<FGraphInformationPopupInfo>& Popups) const override;
 	// End of SNodePanel::SNode interface
 
 	// SGraphNode interface
-	virtual void UpdateGraphNode() OVERRIDE;
-	virtual void CreateBelowWidgetControls(TSharedPtr<SVerticalBox> MainBox) OVERRIDE;
+	virtual void UpdateGraphNode() override;
+	virtual void CreateBelowWidgetControls(TSharedPtr<SVerticalBox> MainBox) override;
 	// End of SGraphNode interface
 
 protected:
diff --git a/Engine/Source/Editor/GraphEditor/Private/AnimationNodes/SGraphNodeStateMachineInstance.h b/Engine/Source/Editor/GraphEditor/Private/AnimationNodes/SGraphNodeStateMachineInstance.h
--- a/Engine/Source/Editor/GraphEditor/Private/AnimationNodes/SGraphNodeStateMachineInstance.h
+++ b/Engine/Source/Editor/GraphEditor/Private/AnimationNodes/SGraphNodeStateMachineInstance.h
@@ -11,6 +11,6 @@ public:
 	void Construct(const FArguments& InArgs, class UAnimGraphNode_StateMachineBase* InNode);
 protected:
 	// SGraphNodeK2Composite interface
-	virtual UEdGraph* GetInnerGraph() const OVERRIDE;
+	virtual UEdGraph* GetInnerGraph() const override;
 	// End of SGraphNodeK2Composite interface
 };
diff --git a/Engine/Source/Editor/GraphEditor/Private/AnimationPins/SGraphPinPose.h b/Engine/Source/Editor/GraphEditor/Private/AnimationPins/SGraphPinPose.h
--- a/Engine/Source/Editor/GraphEditor/Private/AnimationPins/SGraphPinPose.h
+++ b/Engine/Source/Editor/GraphEditor/Private/AnimationPins/SGraphPinPose.h
@@ -12,6 +12,6 @@ public:
 
 protected:
 	// Begin SGraphPin interface
-	virtual const FSlateBrush* GetPinIcon() const OVERRIDE;
+	virtual const FSlateBrush* GetPinIcon() const override;
 	// End SGraphPin interface
 };
diff --git a/Engine/Source/Editor/GraphEditor/Private/AnimationStateNodes/SGraphNodeAnimState.h b/Engine/Source/Editor/GraphEditor/Private/AnimationStateNodes/SGraphNodeAnimState.h
--- a/Engine/Source/Editor/GraphEditor/Private/AnimationStateNodes/SGraphNodeAnimState.h
+++ b/Engine/Source/Editor/GraphEditor/Private/AnimationStateNodes/SGraphNodeAnimState.h
@@ -18,14 +18,14 @@ public:
 	void Construct(const FArguments& InArgs, UAnimStateNodeBase* InNode);
 
 	// SNodePanel::SNode interface
-	virtual void GetNodeInfoPopups(FNodeInfoContext* Context, TArray<FGraphInformationPopupInfo>& Popups) const OVERRIDE;
+	virtual void GetNodeInfoPopups(FNodeInfoContext* Context, TArray<FGraphInformationPopupInfo>& Popups) const override;
 	// End of SNodePanel::SNode interface
 
 	// SGraphNode interface
-	virtual void UpdateGraphNode() OVERRIDE;
-	virtual void CreatePinWidgets() OVERRIDE;
-	virtual void AddPin(const TSharedRef<SGraphPin>& PinToAdd) OVERRIDE;
-	virtual TSharedPtr<SToolTip> GetComplexTooltip() OVERRIDE;
+	virtual void UpdateGraphNode() override;
+	virtual void CreatePinWidgets() override;
+	virtual void AddPin(const TSharedRef<SGraphPin>& PinToAdd) override;
+	virtual TSharedPtr<SToolTip> GetComplexTooltip() override;
 	// End of SGraphNode interface
 
 	static void GetStateInfoPopup(UEdGraphNode* GraphNode, TArray<FGraphInformationPopupInfo>& Popups);
@@ -46,9 +46,9 @@ public:
 	void Construct(const FArguments& InArgs, UAnimStateConduitNode* InNode);
 
 	// SNodePanel::SNode interface
-	virtual void GetNodeInfoPopups(FNodeInfoContext* Context, TArray<FGraphInformationPopupInfo>& Popups) const OVERRIDE;
+	virtual void GetNodeInfoPopups(FNodeInfoContext* Context, TArray<FGraphInformationPopupInfo>& Popups) const override;
 	// End of SNodePanel::SNode interface
 protected:
-	virtual FString GetPreviewCornerText() const OVERRIDE;
-	virtual const FSlateBrush* GetNameIcon() const OVERRIDE;
+	virtual FString GetPreviewCornerText() const override;
+	virtual const FSlateBrush* GetNameIcon() const override;
 };
diff --git a/Engine/Source/Editor/GraphEditor/Private/AnimationStateNodes/SGraphNodeAnimStateEntry.h b/Engine/Source/Editor/GraphEditor/Private/AnimationStateNodes/SGraphNodeAnimStateEntry.h
--- a/Engine/Source/Editor/GraphEditor/Private/AnimationStateNodes/SGraphNodeAnimStateEntry.h
+++ b/Engine/Source/Editor/GraphEditor/Private/AnimationStateNodes/SGraphNodeAnimStateEntry.h
@@ -16,12 +16,12 @@ public:
 	void Construct(const FArguments& InArgs, UAnimStateEntryNode* InNode);
 
 	// SNodePanel::SNode interface
-	virtual void GetNodeInfoPopups(FNodeInfoContext* Context, TArray<FGraphInformationPopupInfo>& Popups) const OVERRIDE;
+	virtual void GetNodeInfoPopups(FNodeInfoContext* Context, TArray<FGraphInformationPopupInfo>& Popups) const override;
 	// End of SNodePanel::SNode interface
 
 	// SGraphNode interface
-	virtual void UpdateGraphNode() OVERRIDE;
-	virtual void AddPin(const TSharedRef<SGraphPin>& PinToAdd) OVERRIDE;
+	virtual void UpdateGraphNode() override;
+	virtual void AddPin(const TSharedRef<SGraphPin>& PinToAdd) override;
 	
 	// End of SGraphNode interface
 
diff --git a/Engine/Source/Editor/GraphEditor/Private/AnimationStateNodes/SGraphNodeAnimTransition.h b/Engine/Source/Editor/GraphEditor/Private/AnimationStateNodes/SGraphNodeAnimTransition.h
--- a/Engine/Source/Editor/GraphEditor/Private/AnimationStateNodes/SGraphNodeAnimTransition.h
+++ b/Engine/Source/Editor/GraphEditor/Private/AnimationStateNodes/SGraphNodeAnimTransition.h
@@ -11,20 +11,20 @@ public:
 	void Construct(const FArguments& InArgs, UAnimStateTransitionNode* InNode);
 
 	// SNodePanel::SNode interface
-	virtual void GetNodeInfoPopups(FNodeInfoContext* Context, TArray<FGraphInformationPopupInfo>& Popups) const OVERRIDE;
-	virtual void MoveTo(const FVector2D& NewPosition, FNodeSet& NodeFilter) OVERRIDE;
-	virtual bool RequiresSecondPassLayout() const OVERRIDE;
-	virtual void PerformSecondPassLayout(const TMap< UObject*, TSharedRef<SNode> >& NodeToWidgetLookup) const OVERRIDE;
+	virtual void GetNodeInfoPopups(FNodeInfoContext* Context, TArray<FGraphInformationPopupInfo>& Popups) const override;
+	virtual void MoveTo(const FVector2D& NewPosition, FNodeSet& NodeFilter) override;
+	virtual bool RequiresSecondPassLayout() const override;
+	virtual void PerformSecondPassLayout(const TMap< UObject*, TSharedRef<SNode> >& NodeToWidgetLookup) const override;
 	// End of SNodePanel::SNode interface
 
 	// SGraphNode interface
-	virtual void UpdateGraphNode() OVERRIDE;
-	virtual TSharedPtr<SToolTip> GetComplexTooltip() OVERRIDE;
+	virtual void UpdateGraphNode() override;
+	virtual TSharedPtr<SToolTip> GetComplexTooltip() override;
 	// End of SGraphNode interface
 
 	// SWidget interface
-	void OnMouseEnter(const FGeometry& MyGeometry, const FPointerEvent& MouseEvent) OVERRIDE;
-	void OnMouseLeave(const FPointerEvent& MouseEvent) OVERRIDE;
+	void OnMouseEnter(const FGeometry& MyGeometry, const FPointerEvent& MouseEvent) override;
+	void OnMouseLeave(const FPointerEvent& MouseEvent) override;
 	// End of SWidget interface
 
 	static FLinearColor StaticGetTransitionColor(UAnimStateTransitionNode* TransNode, bool bIsHovered);
diff --git a/Engine/Source/Editor/GraphEditor/Private/ConnectionDrawingPolicy.h b/Engine/Source/Editor/GraphEditor/Private/ConnectionDrawingPolicy.h
--- a/Engine/Source/Editor/GraphEditor/Private/ConnectionDrawingPolicy.h
+++ b/Engine/Source/Editor/GraphEditor/Private/ConnectionDrawingPolicy.h
@@ -147,8 +147,8 @@ public:
 	void CalculateEnvelopeAlphas(double ExecutionTime, /*out*/ float& AttackAlpha, /*out*/ float& SustainAlpha) const;
 
 	// FConnectionDrawingPolicy interface
-	virtual void DetermineWiringStyle(UEdGraphPin* OutputPin, UEdGraphPin* InputPin, /*inout*/ float& Thickness, /*inout*/ FLinearColor& WireColor, /*inout*/bool& bDrawBubbles, /*inout*/ bool& bBidirectional) OVERRIDE;
-	virtual void Draw(TMap<TSharedRef<SWidget>, FArrangedWidget>& PinGeometries, FArrangedChildren& ArrangedNodes) OVERRIDE;
+	virtual void DetermineWiringStyle(UEdGraphPin* OutputPin, UEdGraphPin* InputPin, /*inout*/ float& Thickness, /*inout*/ FLinearColor& WireColor, /*inout*/bool& bDrawBubbles, /*inout*/ bool& bBidirectional) override;
+	virtual void Draw(TMap<TSharedRef<SWidget>, FArrangedWidget>& PinGeometries, FArrangedChildren& ArrangedNodes) override;
 	// End of FConnectionDrawingPolicy interface
 
 	virtual void DetermineStyleOfExecWire(float& Thickness, FLinearColor& WireColor, bool& bDrawBubbles, const FTimePair& Times);
@@ -183,8 +183,8 @@ public:
 	FStateMachineConnectionDrawingPolicy(int32 InBackLayerID, int32 InFrontLayerID, float ZoomFactor, const FSlateRect& InClippingRect, FSlateWindowElementList& InDrawElements, UEdGraph* InGraphObj);
 
 	// FConnectionDrawingPolicy interface
-	virtual void DetermineWiringStyle(UEdGraphPin* OutputPin, UEdGraphPin* InputPin, /*inout*/ float& Thickness, /*inout*/ FLinearColor& WireColor, /*inout*/ bool& bDrawBubbles, /*inout*/ bool& bBidirectional) OVERRIDE;
-	virtual void Draw(TMap<TSharedRef<SWidget>, FArrangedWidget>& PinGeometries, FArrangedChildren& ArrangedNodes) OVERRIDE;
+	virtual void DetermineWiringStyle(UEdGraphPin* OutputPin, UEdGraphPin* InputPin, /*inout*/ float& Thickness, /*inout*/ FLinearColor& WireColor, /*inout*/ bool& bDrawBubbles, /*inout*/ bool& bBidirectional) override;
+	virtual void Draw(TMap<TSharedRef<SWidget>, FArrangedWidget>& PinGeometries, FArrangedChildren& ArrangedNodes) override;
 	virtual void DetermineLinkGeometry(
 		TMap<TSharedRef<SWidget>,
 		FArrangedWidget>& PinGeometries,
@@ -194,11 +194,11 @@ public:
 		UEdGraphPin* InputPin,
 		/*out*/ FArrangedWidget*& StartWidgetGeometry,
 		/*out*/ FArrangedWidget*& EndWidgetGeometry
-		) OVERRIDE;
-	virtual void DrawSplineWithArrow(FGeometry& StartGeom, FGeometry& EndGeom, const FLinearColor& WireColor, float WireThickness, bool bDrawBubbles, bool Bidirectional) OVERRIDE;
-	virtual void DrawSplineWithArrow(const FVector2D& StartPoint, const FVector2D& EndPoint, const FLinearColor& WireColor, float WireThickness, bool bDrawBubbles, bool Bidirectional) OVERRIDE;
-	virtual void DrawPreviewConnector(const FGeometry& PinGeometry, const FVector2D& StartPoint, const FVector2D& EndPoint, UEdGraphPin* Pin) OVERRIDE;
-	virtual void DrawConnection(int32 LayerId, const FVector2D& Start, const FVector2D& End, const FLinearColor& InColor, float Thickness, bool bDrawBubbles) OVERRIDE;
+		) override;
+	virtual void DrawSplineWithArrow(FGeometry& StartGeom, FGeometry& EndGeom, const FLinearColor& WireColor, float WireThickness, bool bDrawBubbles, bool Bidirectional) override;
+	virtual void DrawSplineWithArrow(const FVector2D& StartPoint, const FVector2D& EndPoint, const FLinearColor& WireColor, float WireThickness, bool bDrawBubbles, bool Bidirectional) override;
+	virtual void DrawPreviewConnector(const FGeometry& PinGeometry, const FVector2D& StartPoint, const FVector2D& EndPoint, UEdGraphPin* Pin) override;
+	virtual void DrawConnection(int32 LayerId, const FVector2D& Start, const FVector2D& End, const FLinearColor& InColor, float Thickness, bool bDrawBubbles) override;
 	// End of FConnectionDrawingPolicy interface
 
 protected:
@@ -216,9 +216,9 @@ public:
 	FAnimGraphConnectionDrawingPolicy(int32 InBackLayerID, int32 InFrontLayerID, float ZoomFactor, const FSlateRect& InClippingRect, FSlateWindowElementList& InDrawElements, UEdGraph* InGraphObj);
 
 	// FKismetConnectionDrawingPolicy interface
-	virtual bool TreatWireAsExecutionPin(UEdGraphPin* InputPin, UEdGraphPin* OutputPin) const OVERRIDE;
-	virtual void BuildExecutionRoadmap() OVERRIDE;
-	virtual void DetermineStyleOfExecWire(float& Thickness, FLinearColor& WireColor, bool& bDrawBubbles, const FTimePair& Times) OVERRIDE;
+	virtual bool TreatWireAsExecutionPin(UEdGraphPin* InputPin, UEdGraphPin* OutputPin) const override;
+	virtual void BuildExecutionRoadmap() override;
+	virtual void DetermineStyleOfExecWire(float& Thickness, FLinearColor& WireColor, bool& bDrawBubbles, const FTimePair& Times) override;
 	// End of FKismetConnectionDrawingPolicy interface
 };
 
@@ -262,8 +262,8 @@ public:
 	void BuildAudioFlowRoadmap();
 
 	// FConnectionDrawingPolicy interface
-	virtual void DetermineWiringStyle(UEdGraphPin* OutputPin, UEdGraphPin* InputPin, /*inout*/ float& Thickness, /*inout*/ FLinearColor& WireColor, /*inout*/bool& bDrawBubbles, /*inout*/ bool& bBidirectional) OVERRIDE;
-	virtual void Draw(TMap<TSharedRef<SWidget>, FArrangedWidget>& PinGeometries, FArrangedChildren& ArrangedNodes) OVERRIDE;
+	virtual void DetermineWiringStyle(UEdGraphPin* OutputPin, UEdGraphPin* InputPin, /*inout*/ float& Thickness, /*inout*/ FLinearColor& WireColor, /*inout*/bool& bDrawBubbles, /*inout*/ bool& bBidirectional) override;
+	virtual void Draw(TMap<TSharedRef<SWidget>, FArrangedWidget>& PinGeometries, FArrangedChildren& ArrangedNodes) override;
 	// End of FConnectionDrawingPolicy interface
 };
 
@@ -281,6 +281,6 @@ public:
 	FMaterialGraphConnectionDrawingPolicy(int32 InBackLayerID, int32 InFrontLayerID, float ZoomFactor, const FSlateRect& InClippingRect, FSlateWindowElementList& InDrawElements, UEdGraph* InGraphObj);
 
 	// FConnectionDrawingPolicy interface
-	virtual void DetermineWiringStyle(UEdGraphPin* OutputPin, UEdGraphPin* InputPin, /*inout*/ float& Thickness, /*inout*/ FLinearColor& WireColor, /*inout*/bool& bDrawBubbles, /*inout*/ bool& bBidirectional) OVERRIDE;
+	virtual void DetermineWiringStyle(UEdGraphPin* OutputPin, UEdGraphPin* InputPin, /*inout*/ float& Thickness, /*inout*/ FLinearColor& WireColor, /*inout*/bool& bDrawBubbles, /*inout*/ bool& bBidirectional) override;
 	// End of FConnectionDrawingPolicy interface
 };
diff --git a/Engine/Source/Editor/GraphEditor/Private/DragConnection.h b/Engine/Source/Editor/GraphEditor/Private/DragConnection.h
--- a/Engine/Source/Editor/GraphEditor/Private/DragConnection.h
+++ b/Engine/Source/Editor/GraphEditor/Private/DragConnection.h
@@ -10,15 +10,15 @@ public:
 	static TSharedRef<FDragConnection> New(const TSharedRef<SGraphPanel>& InGraphPanel, const TArray< TSharedRef<SGraphPin> >& InStartingPins, bool bInShiftOperation);
 	
 	// FDragDropOperation interface
-	virtual void OnDrop( bool bDropWasHandled, const FPointerEvent& MouseEvent ) OVERRIDE;
+	virtual void OnDrop( bool bDropWasHandled, const FPointerEvent& MouseEvent ) override;
 	// End of FDragDropOperation interface
 
 	// FGraphEditorDragDropAction interface
-	virtual void HoverTargetChanged() OVERRIDE;
-	virtual FReply DroppedOnPin(FVector2D ScreenPosition, FVector2D GraphPosition) OVERRIDE;
-	virtual FReply DroppedOnPanel(const TSharedRef< SWidget >& Panel, FVector2D ScreenPosition, FVector2D GraphPosition, UEdGraph& Graph) OVERRIDE;
-	virtual void OnDragBegin(const TSharedRef<class SGraphPin>& InPin) OVERRIDE;
-	virtual void OnDragged(const class FDragDropEvent& DragDropEvent) OVERRIDE;
+	virtual void HoverTargetChanged() override;
+	virtual FReply DroppedOnPin(FVector2D ScreenPosition, FVector2D GraphPosition) override;
+	virtual FReply DroppedOnPanel(const TSharedRef< SWidget >& Panel, FVector2D ScreenPosition, FVector2D GraphPosition, UEdGraph& Graph) override;
+	virtual void OnDragBegin(const TSharedRef<class SGraphPin>& InPin) override;
+	virtual void OnDragged(const class FDragDropEvent& DragDropEvent) override;
 	// End of FGraphEditorDragDropAction interface
 
 	/*
diff --git a/Engine/Source/Editor/GraphEditor/Private/DragNode.h b/Engine/Source/Editor/GraphEditor/Private/DragNode.h
--- a/Engine/Source/Editor/GraphEditor/Private/DragNode.h
+++ b/Engine/Source/Editor/GraphEditor/Private/DragNode.h
@@ -12,11 +12,11 @@ public:
 	static TSharedRef<FDragNode> New(const TSharedRef<SGraphPanel>& InGraphPanel, const TArray< TSharedRef<SGraphNode> >& InDraggedNodes);
 	
 	// FGraphEditorDragDropAction interface
-	virtual void HoverTargetChanged() OVERRIDE;
-	virtual FReply DroppedOnNode(FVector2D ScreenPosition, FVector2D GraphPosition) OVERRIDE;
-	virtual FReply DroppedOnPanel( const TSharedRef< SWidget >& Panel, FVector2D ScreenPosition, FVector2D GraphPosition, UEdGraph& Graph) OVERRIDE;
+	virtual void HoverTargetChanged() override;
+	virtual FReply DroppedOnNode(FVector2D ScreenPosition, FVector2D GraphPosition) override;
+	virtual FReply DroppedOnPanel( const TSharedRef< SWidget >& Panel, FVector2D ScreenPosition, FVector2D GraphPosition, UEdGraph& Graph) override;
 	//virtual void OnDragBegin( const TSharedRef<class SGraphPin>& InPin) OVERRIDE;
-	virtual void OnDragged (const class FDragDropEvent& DragDropEvent ) OVERRIDE;
+	virtual void OnDragged (const class FDragDropEvent& DragDropEvent ) override;
 	// End of FGraphEditorDragDropAction interface
 	
 	const TArray< TSharedRef<SGraphNode> > & GetNodes() const;
