Only hide asset picker for component hierarchy for objects, not classes

diff --git a/Engine/Source/Editor/BlueprintGraph/Classes/EdGraphSchema_K2.h b/Engine/Source/Editor/BlueprintGraph/Classes/EdGraphSchema_K2.h
--- a/Engine/Source/Editor/BlueprintGraph/Classes/EdGraphSchema_K2.h
+++ b/Engine/Source/Editor/BlueprintGraph/Classes/EdGraphSchema_K2.h
@@ -17,9 +17,9 @@ struct BLUEPRINTGRAPH_API FEdGraphSchemaAction_K2Struct : public FEdGraphSchemaA
 	UStruct* Struct;
 
 	void AddReferencedObjects( FReferenceCollector& Collector ) override
-{
-		if( Struct )
 	{
+		if( Struct )
+		{
 			Collector.AddReferencedObject(Struct);
 		}
 	}
@@ -469,6 +469,7 @@ public:
 	virtual void TrySetDefaultObject(UEdGraphPin& Pin, UObject* NewDefaultObject) const override;
 	virtual void TrySetDefaultText(UEdGraphPin& InPin, const FText& InNewDefaultText) const override;
 	virtual bool ShouldHidePinDefaultValue(UEdGraphPin* Pin) const override;
+	virtual bool ShouldShowAssetPickerForPin(UEdGraphPin* Pin) const override;
 	virtual FLinearColor GetPinTypeColor(const FEdGraphPinType& PinType) const override;
 	virtual FString GetPinDisplayName(const UEdGraphPin* Pin) const override;
 	virtual void ConstructBasicPinTooltip(const UEdGraphPin& Pin, const FString& PinDescription, FString& TooltipOut) const override;
diff --git a/Engine/Source/Editor/BlueprintGraph/Private/EdGraphSchema_K2.cpp b/Engine/Source/Editor/BlueprintGraph/Private/EdGraphSchema_K2.cpp
--- a/Engine/Source/Editor/BlueprintGraph/Private/EdGraphSchema_K2.cpp
+++ b/Engine/Source/Editor/BlueprintGraph/Private/EdGraphSchema_K2.cpp
@@ -3008,6 +3008,21 @@ bool UEdGraphSchema_K2::ShouldHidePinDefaultValue(UEdGraphPin* Pin) const
 	return false;
 }
 
+bool UEdGraphSchema_K2::ShouldShowAssetPickerForPin(UEdGraphPin* Pin) const
+{
+	bool bShow = true;
+	if (Pin->PinType.PinCategory == PC_Object)
+	{
+		UClass* ObjectClass = Cast<UClass>(Pin->PinType.PinSubCategoryObject.Get());
+		if (ObjectClass)
+		{
+			// Don't show literal buttons for component type objects
+			bShow = !ObjectClass->IsChildOf(UActorComponent::StaticClass());
+		}
+	}
+	return bShow;
+}
+
 void UEdGraphSchema_K2::SetPinDefaultValue(UEdGraphPin* Pin, const UFunction* Function, const UProperty* Param) const
 {
 	if (Function && Param)
diff --git a/Engine/Source/Editor/GraphEditor/Private/KismetPins/SGraphPinObject.cpp b/Engine/Source/Editor/GraphEditor/Private/KismetPins/SGraphPinObject.cpp
--- a/Engine/Source/Editor/GraphEditor/Private/KismetPins/SGraphPinObject.cpp
+++ b/Engine/Source/Editor/GraphEditor/Private/KismetPins/SGraphPinObject.cpp
@@ -47,92 +47,75 @@ TSharedRef<SWidget>	SGraphPinObject::GetDefaultValueWidget()
 		}
 	}
 	// Don't show literal buttons for component type objects
-	UClass* ObjectClass = Cast<UClass>(GraphPinObj->PinType.PinSubCategoryObject.Get());
-	if (ObjectClass != NULL && ObjectClass->IsChildOf(UActorComponent::StaticClass()))
+	if (GraphPinObj->GetSchema()->ShouldShowAssetPickerForPin(GraphPinObj))
 	{
 		return
 			SNew(SHorizontalBox)
 			.Visibility( this, &SGraphPin::GetDefaultValueVisibility )
-			// Icon to show object assigned or not
 			+SHorizontalBox::Slot()
 			.AutoWidth()
-			.Padding(1,0)
-			.MaxWidth(120)
-			[
-				SNew(SEditableTextBox)
-				.Style(FEditorStyle::Get(), "Graph.EditableTextBox")
-				.Text( this, &SGraphPinObject::GetObjectName )
-				.ForegroundColor( FLinearColor(0.004f,0.5f,1.0f) )
-				.IsReadOnly(true)
-				.ToolTipText(this, &SGraphPinObject::GetObjectToolTipAsString)
-				.BackgroundColor( FLinearColor(0.004f,0.5f,1.0f) )	
-			];
-	}
-
-	return
-		SNew(SHorizontalBox)
-		.Visibility( this, &SGraphPin::GetDefaultValueVisibility )
-		+SHorizontalBox::Slot()
-		.AutoWidth()
-		.Padding(2,0)
-		.MaxWidth(100.0f)
-		[
-			SAssignNew(AssetPickerAnchor, SComboButton)
-			.ButtonStyle( FEditorStyle::Get(), "PropertyEditor.AssetComboStyle" )
-			.ForegroundColor( this, &SGraphPinObject::OnGetComboForeground)
-			.ContentPadding( FMargin(2,2,2,1) )
-			.ButtonColorAndOpacity( this, &SGraphPinObject::OnGetWidgetBackground )
-			.MenuPlacement(MenuPlacement_BelowAnchor)
-			.ButtonContent()
+			.Padding(2,0)
+			.MaxWidth(100.0f)
 			[
-				SNew(STextBlock)
-				.ColorAndOpacity( this, &SGraphPinObject::OnGetComboForeground )
-				.TextStyle( FEditorStyle::Get(), "PropertyEditor.AssetClass" )
-				.Font( FEditorStyle::GetFontStyle( "PropertyWindow.NormalFont" ) )
-				.Text( this, &SGraphPinObject::OnGetComboTextValue )
-				.ToolTipText( this, &SGraphPinObject::GetObjectToolTip )
+				SAssignNew(AssetPickerAnchor, SComboButton)
+				.ButtonStyle( FEditorStyle::Get(), "PropertyEditor.AssetComboStyle" )
+				.ForegroundColor( this, &SGraphPinObject::OnGetComboForeground)
+				.ContentPadding( FMargin(2,2,2,1) )
+				.ButtonColorAndOpacity( this, &SGraphPinObject::OnGetWidgetBackground )
+				.MenuPlacement(MenuPlacement_BelowAnchor)
+				.ButtonContent()
+				[
+					SNew(STextBlock)
+					.ColorAndOpacity( this, &SGraphPinObject::OnGetComboForeground )
+					.TextStyle( FEditorStyle::Get(), "PropertyEditor.AssetClass" )
+					.Font( FEditorStyle::GetFontStyle( "PropertyWindow.NormalFont" ) )
+					.Text( this, &SGraphPinObject::OnGetComboTextValue )
+					.ToolTipText( this, &SGraphPinObject::GetObjectToolTip )
+				]
+				.MenuContent()
+				[
+					GenerateAssetPicker()
+				]
 			]
-			.MenuContent()
-			[
-				GenerateAssetPicker()
-			]
-		]
-		// Use button
-		+SHorizontalBox::Slot()
-		.AutoWidth()
-		.Padding(1,0)
-		.VAlign(VAlign_Center)
-		[
-			SAssignNew(UseButton, SButton)
-			.ButtonStyle( FEditorStyle::Get(), "NoBorder" )
-			.ButtonColorAndOpacity( this, &SGraphPinObject::OnGetWidgetBackground )
-			.OnClicked( GetOnUseButtonDelegate() )
-			.ContentPadding(1.f)
-			.ToolTipText(NSLOCTEXT("GraphEditor", "ObjectGraphPin_Use", "Use asset browser selection").ToString())
+			// Use button
+			+SHorizontalBox::Slot()
+			.AutoWidth()
+			.Padding(1,0)
+			.VAlign(VAlign_Center)
 			[
-				SNew(SImage)
-				.ColorAndOpacity( this, &SGraphPinObject::OnGetWidgetForeground )
-				.Image( FEditorStyle::GetBrush(TEXT("PropertyWindow.Button_Use")) )
+				SAssignNew(UseButton, SButton)
+				.ButtonStyle( FEditorStyle::Get(), "NoBorder" )
+				.ButtonColorAndOpacity( this, &SGraphPinObject::OnGetWidgetBackground )
+				.OnClicked( GetOnUseButtonDelegate() )
+				.ContentPadding(1.f)
+				.ToolTipText(NSLOCTEXT("GraphEditor", "ObjectGraphPin_Use", "Use asset browser selection").ToString())
+				[
+					SNew(SImage)
+					.ColorAndOpacity( this, &SGraphPinObject::OnGetWidgetForeground )
+					.Image( FEditorStyle::GetBrush(TEXT("PropertyWindow.Button_Use")) )
+				]
 			]
-		]
-		// Browse button
-		+SHorizontalBox::Slot()
-		.AutoWidth()
-		.Padding(1,0)
-		.VAlign(VAlign_Center)
-		[
-			SAssignNew(BrowseButton, SButton)
-			.ButtonStyle( FEditorStyle::Get(), "NoBorder" )
-			.ButtonColorAndOpacity( this, &SGraphPinObject::OnGetWidgetBackground )
-			.OnClicked( GetOnBrowseButtonDelegate() )
-			.ContentPadding(0)
-			.ToolTipText(NSLOCTEXT("GraphEditor", "ObjectGraphPin_Browse", "Browse").ToString())
+			// Browse button
+			+SHorizontalBox::Slot()
+			.AutoWidth()
+			.Padding(1,0)
+			.VAlign(VAlign_Center)
 			[
-				SNew(SImage)
-				.ColorAndOpacity( this, &SGraphPinObject::OnGetWidgetForeground )
-				.Image( FEditorStyle::GetBrush(TEXT("PropertyWindow.Button_Browse")) )
-			]
-		];
+				SAssignNew(BrowseButton, SButton)
+				.ButtonStyle( FEditorStyle::Get(), "NoBorder" )
+				.ButtonColorAndOpacity( this, &SGraphPinObject::OnGetWidgetBackground )
+				.OnClicked( GetOnBrowseButtonDelegate() )
+				.ContentPadding(0)
+				.ToolTipText(NSLOCTEXT("GraphEditor", "ObjectGraphPin_Browse", "Browse").ToString())
+				[
+					SNew(SImage)
+					.ColorAndOpacity( this, &SGraphPinObject::OnGetWidgetForeground )
+					.Image( FEditorStyle::GetBrush(TEXT("PropertyWindow.Button_Browse")) )
+				]
+			];
+	}
+
+	return SNullWidget::NullWidget;
 }
 
 FOnClicked SGraphPinObject::GetOnUseButtonDelegate()
diff --git a/Engine/Source/Runtime/Engine/Classes/EdGraph/EdGraphSchema.h b/Engine/Source/Runtime/Engine/Classes/EdGraph/EdGraphSchema.h
--- a/Engine/Source/Runtime/Engine/Classes/EdGraph/EdGraphSchema.h
+++ b/Engine/Source/Runtime/Engine/Classes/EdGraph/EdGraphSchema.h
@@ -551,6 +551,9 @@ class ENGINE_API UEdGraphSchema : public UObject
 	/** If we should disallow viewing and editing of the supplied pin */
 	virtual bool ShouldHidePinDefaultValue(UEdGraphPin* Pin) const { return false; }
 
+	/** Should the Pin in question display an asset picker */
+	virtual bool ShouldShowAssetPickerForPin(UEdGraphPin* Pin) const { return true; }
+
 	/**
 	 * Gets the draw color of a pin based on it's type.
 	 *
