Modified the dragdrop operations in the blueprint graphs to respect read only attributes set on either graphs or nodes.
#TTP 334976 - BLUEPRINTS: Drag-dropping can still create nodes in read-only graphs (e.g., during PIE or when viewing an anim parent graph)
#Branch UE4
#Proj GraphEditor

#Add added SGraphPanel::IsGraphEditable so nodes and objects placed in the graph can easily determine if the graph is read only.
#Change Modified SGraphPanel::OnDragLeave to restore the tooltip if a FDecoratedDragDropOp has been modified to indicate that the graph is read only.
#Change Modified SGraphPanel::OnDragOver to change the icon on the tooltip to a lined circle if the graph is read only.
#Change Modified SGraphPanel::OnDrop to exit without changes if the graph is read only.
#Change Modified FGraphEditorDragDropAction to present feedback to the user if the drag drop target is invalid. this is activated by setting FGraphEditorDragDropAction::SetDropTargetValid, this just collpases the active icon and displays a lined circle with the same tooltip when set to false.
#Add added SGraphNode::IsNodeEditable to determine if the node is currently editable, based on an evaluation of the IsEditable attribute and the parent graphs IsGraphEditable.
#Change modified SGraphNode::OnDragOver to display a lined circle if the parent graph or the node is read only.
#Change modified SGraphNode::OnDrop to exit without changes if the parent graph or the node is read only.
#Change modified SGraphPin::OnDrop to exit without changes if the parent graph or the owner node is read only.

diff --git a/Engine/Source/Editor/GraphEditor/Private/GraphEditorDragDropAction.cpp b/Engine/Source/Editor/GraphEditor/Private/GraphEditorDragDropAction.cpp
--- a/Engine/Source/Editor/GraphEditor/Private/GraphEditorDragDropAction.cpp
+++ b/Engine/Source/Editor/GraphEditor/Private/GraphEditorDragDropAction.cpp
@@ -118,8 +118,18 @@ void FGraphEditorDragDropAction::SetSimpleFeedbackMessage(const FSlateBrush* Ico
 		+SHorizontalBox::Slot()
 		.AutoWidth()
 		.Padding(3.0f)
+		[
+			SNew(SImage)
+			.Visibility( this, &FGraphEditorDragDropAction::GetErrorIconVisible )
+			.Image( FEditorStyle::GetBrush( TEXT("Graph.ConnectorFeedback.Error") ))
+			.ColorAndOpacity( FLinearColor::White )
+		]
+		+SHorizontalBox::Slot()
+		.AutoWidth()
+		.Padding(3.0f)
 		[
 			SNew(SImage) 
+			.Visibility( this, &FGraphEditorDragDropAction::GetIconVisible )
 			.Image( Icon )
 			.ColorAndOpacity( IconColor )
 		]
@@ -134,6 +144,16 @@ void FGraphEditorDragDropAction::SetSimpleFeedbackMessage(const FSlateBrush* Ico
 	);
 }
 
+EVisibility FGraphEditorDragDropAction::GetIconVisible() const
+{
+	return bDropTargetValid ? EVisibility::Visible : EVisibility::Collapsed;
+}
+
+EVisibility FGraphEditorDragDropAction::GetErrorIconVisible() const
+{
+	return bDropTargetValid ? EVisibility::Collapsed : EVisibility::Visible;
+}
+
 ////////////////////////////////////////////////////////////
 
 void FGraphSchemaActionDragDropAction::HoverTargetChanged()
diff --git a/Engine/Source/Editor/GraphEditor/Private/SGraphNode.cpp b/Engine/Source/Editor/GraphEditor/Private/SGraphNode.cpp
--- a/Engine/Source/Editor/GraphEditor/Private/SGraphNode.cpp
+++ b/Engine/Source/Editor/GraphEditor/Private/SGraphNode.cpp
@@ -110,6 +110,12 @@ void SGraphNode::SetIsEditable(TAttribute<bool> InIsEditable)
 	IsEditable = InIsEditable;
 }
 
+bool SGraphNode::IsNodeEditable() const
+{
+	bool bReadOnly = OwnerGraphPanelPtr.IsValid() ? OwnerGraphPanelPtr.Pin()->IsGraphEditable() : false;
+	return IsEditable.Get() && !bReadOnly;
+}
+
 /** Set event when node is double clicked */
 void SGraphNode::SetDoubleClickEvent(FSingleNodeEvent InDoubleClickEvent)
 {
@@ -221,6 +227,8 @@ FReply SGraphNode::OnDragOver( const FGeometry& MyGeometry, const FDragDropEvent
 			bool bOkIcon = false;
 			FString TooltipText;
 			GraphNode->GetSchema()->GetAssetsNodeHoverMessage(AssetOp->AssetData, GraphNode, TooltipText, bOkIcon);
+			bool bReadOnly = OwnerGraphPanelPtr.IsValid() ? !OwnerGraphPanelPtr.Pin()->IsGraphEditable() : false;
+			bOkIcon = bReadOnly ? false : bOkIcon;
 			const FSlateBrush* TooltipIcon = bOkIcon ? FEditorStyle::GetBrush(TEXT("Graph.ConnectorFeedback.OK")) : FEditorStyle::GetBrush(TEXT("Graph.ConnectorFeedback.Error"));;
 			AssetOp->SetToolTip(FText::FromString(TooltipText), TooltipIcon);
 		}
@@ -247,8 +255,9 @@ FVector2D SGraphNode::NodeCoordToGraphCoord( const FVector2D& NodeSpaceCoordinat
 
 FReply SGraphNode::OnDrop( const FGeometry& MyGeometry, const FDragDropEvent& DragDropEvent )
 {
+	bool bReadOnly = OwnerGraphPanelPtr.IsValid() ? !OwnerGraphPanelPtr.Pin()->IsGraphEditable() : false;
 	TSharedPtr<FDragDropOperation> Operation = DragDropEvent.GetOperation();
-	if (!Operation.IsValid())
+	if (!Operation.IsValid() || bReadOnly)
 	{
 		return FReply::Unhandled();
 	}
diff --git a/Engine/Source/Editor/GraphEditor/Private/SGraphPanel.cpp b/Engine/Source/Editor/GraphEditor/Private/SGraphPanel.cpp
--- a/Engine/Source/Editor/GraphEditor/Private/SGraphPanel.cpp
+++ b/Engine/Source/Editor/GraphEditor/Private/SGraphPanel.cpp
@@ -702,15 +702,15 @@ void SGraphPanel::OnDragEnter( const FGeometry& MyGeometry, const FDragDropEvent
 
 void SGraphPanel::OnDragLeave( const FDragDropEvent& DragDropEvent )
 {
-	TSharedPtr<FGraphEditorDragDropAction> DragConnectionOp = DragDropEvent.GetOperationAs<FGraphEditorDragDropAction>();
-	if (DragConnectionOp.IsValid())
+	TSharedPtr<FGraphEditorDragDropAction> Operation = DragDropEvent.GetOperationAs<FGraphEditorDragDropAction>();
+	if( Operation.IsValid() )
 	{
-		DragConnectionOp->SetHoveredGraph( TSharedPtr<SGraphPanel>(NULL) );
+		Operation->SetHoveredGraph( TSharedPtr<SGraphPanel>(NULL) );
 	}
 	else
 	{
-		TSharedPtr<FAssetDragDropOp> AssetOp = DragDropEvent.GetOperationAs<FAssetDragDropOp>();
-		if (AssetOp.IsValid())
+		TSharedPtr<FDecoratedDragDropOp> AssetOp = DragDropEvent.GetOperationAs<FDecoratedDragDropOp>();
+		if( AssetOp.IsValid()  )
 		{
 			AssetOp->ResetToDefaultToolTip();
 		}
@@ -725,6 +725,31 @@ FReply SGraphPanel::OnDragOver( const FGeometry& MyGeometry, const FDragDropEven
 		return FReply::Unhandled();
 	}
 
+	// Handle Read only graphs
+	if( !IsEditable.Get() )
+	{
+		TSharedPtr<FGraphEditorDragDropAction> GraphDragDropOp = DragDropEvent.GetOperationAs<FGraphEditorDragDropAction>();
+
+		if( GraphDragDropOp.IsValid() )
+		{
+			GraphDragDropOp->SetDropTargetValid( false );
+		}
+		else
+		{
+			TSharedPtr<FDecoratedDragDropOp> AssetOp = DragDropEvent.GetOperationAs<FDecoratedDragDropOp>();
+			if( AssetOp.IsValid() )
+			{
+				FText Tooltip = AssetOp->GetHoverText();
+				if( Tooltip.IsEmpty() )
+				{
+					Tooltip = NSLOCTEXT( "GraphPanel", "DragDropOperation", "Graph is Read-Only" );
+				}
+				AssetOp->SetToolTip( Tooltip, FEditorStyle::GetBrush(TEXT("Graph.ConnectorFeedback.Error")));
+			}
+		}
+		return FReply::Handled();
+	}
+
 	if( Operation->IsOfType<FGraphEditorDragDropAction>() )
 	{
 		PreviewConnectorEndpoint = MyGeometry.AbsoluteToLocal( DragDropEvent.GetScreenSpacePosition() );
@@ -760,7 +785,7 @@ FReply SGraphPanel::OnDrop( const FGeometry& MyGeometry, const FDragDropEvent& D
 	FSlateApplication::Get().SetKeyboardFocus(AsShared(), EKeyboardFocusCause::SetDirectly);
 
 	TSharedPtr<FDragDropOperation> Operation = DragDropEvent.GetOperation();
-	if (!Operation.IsValid())
+	if (!Operation.IsValid() || !IsEditable.Get())
 	{
 		return FReply::Unhandled();
 	}
diff --git a/Engine/Source/Editor/GraphEditor/Private/SGraphPanel.h b/Engine/Source/Editor/GraphEditor/Private/SGraphPanel.h
--- a/Engine/Source/Editor/GraphEditor/Private/SGraphPanel.h
+++ b/Engine/Source/Editor/GraphEditor/Private/SGraphPanel.h
@@ -152,6 +152,9 @@ public:
 	/** helper to attach graph events to sub node, which won't be placed directly on the graph */
 	void AttachGraphEvents(TSharedPtr<SGraphNode> CreatedSubNode);
 
+	/** Returns if this graph is editable */
+	bool IsGraphEditable() const { return IsEditable.Get(); }
+
 protected:
 
 	void NotifyGraphChanged ( const struct FEdGraphEditAction& InAction);
diff --git a/Engine/Source/Editor/GraphEditor/Private/SGraphPin.cpp b/Engine/Source/Editor/GraphEditor/Private/SGraphPin.cpp
--- a/Engine/Source/Editor/GraphEditor/Private/SGraphPin.cpp
+++ b/Engine/Source/Editor/GraphEditor/Private/SGraphPin.cpp
@@ -666,8 +666,11 @@ bool SGraphPin::TryHandlePinConnection(SGraphPin& OtherSPin)
 
 FReply SGraphPin::OnDrop( const FGeometry& MyGeometry, const FDragDropEvent& DragDropEvent )
 {
+	TSharedPtr<SGraphNode> NodeWidget = OwnerNodePtr.Pin();
+	bool bReadOnly = NodeWidget.IsValid() ? !NodeWidget->IsNodeEditable() : false;
+
 	TSharedPtr<FDragDropOperation> Operation = DragDropEvent.GetOperation();
-	if (!Operation.IsValid())
+	if (!Operation.IsValid() || bReadOnly)
 	{
 		return FReply::Unhandled();
 	}
@@ -707,18 +710,14 @@ FReply SGraphPin::OnDrop( const FGeometry& MyGeometry, const FDragDropEvent& Dra
 		return DragConnectionOp->DroppedOnPin(DragDropEvent.GetScreenSpacePosition(), NodeAddPosition);
 	}
 	// handle dropping an asset on the pin
-	else if (Operation->IsOfType<FAssetDragDropOp>())
+	else if (Operation->IsOfType<FAssetDragDropOp>() && NodeWidget.IsValid())
 	{
-		TSharedPtr<SGraphNode> NodeWidget = OwnerNodePtr.Pin();
-		if (NodeWidget.IsValid())
+		UEdGraphNode* Node = NodeWidget->GetNodeObj();
+		if(Node != NULL && Node->GetSchema() != NULL)
 		{
-			UEdGraphNode* Node = NodeWidget->GetNodeObj();
-			if(Node != NULL && Node->GetSchema() != NULL)
-			{
-				TSharedPtr<FAssetDragDropOp> AssetOp = StaticCastSharedPtr<FAssetDragDropOp>(Operation);
+			TSharedPtr<FAssetDragDropOp> AssetOp = StaticCastSharedPtr<FAssetDragDropOp>(Operation);
 
-				Node->GetSchema()->DroppedAssetsOnPin(AssetOp->AssetData, DragDropEvent.GetScreenSpacePosition(), GraphPinObj);
-			}
+			Node->GetSchema()->DroppedAssetsOnPin(AssetOp->AssetData, DragDropEvent.GetScreenSpacePosition(), GraphPinObj);
 		}
 		return FReply::Handled();
 	}
diff --git a/Engine/Source/Editor/GraphEditor/Public/GraphEditorDragDropAction.h b/Engine/Source/Editor/GraphEditor/Public/GraphEditorDragDropAction.h
--- a/Engine/Source/Editor/GraphEditor/Public/GraphEditorDragDropAction.h
+++ b/Engine/Source/Editor/GraphEditor/Public/GraphEditorDragDropAction.h
@@ -13,6 +13,7 @@ public:
 	void SetHoveredGraph(const TSharedPtr<class SGraphPanel>& InGraph);
 	void SetHoveredCategoryName(const FString& InHoverCategoryName);
 	void SetHoveredAction(TSharedPtr<struct FEdGraphSchemaAction> Action);
+	void SetDropTargetValid( bool bValid ) { bDropTargetValid = bValid; }
 
 	// Interface to override
 	virtual void HoverTargetChanged() {}
@@ -36,6 +37,10 @@ protected:
 	virtual void Construct();
 
 private:
+
+	EVisibility GetIconVisible() const;
+	EVisibility GetErrorIconVisible() const;
+
 	// The pin that the drag action is currently hovering over
 	TWeakObjectPtr<class UEdGraphPin> HoveredPin;
 
@@ -52,6 +57,9 @@ protected:
 
 	// Action we are hovering over
 	TWeakPtr<struct FEdGraphSchemaAction> HoveredAction;
+
+	// drop target status
+	bool bDropTargetValid;
 };
 
 // Drag-drop action where an FEdGraphSchemaAction should be performed when dropped
diff --git a/Engine/Source/Editor/GraphEditor/Public/SGraphNode.h b/Engine/Source/Editor/GraphEditor/Public/SGraphNode.h
--- a/Engine/Source/Editor/GraphEditor/Public/SGraphNode.h
+++ b/Engine/Source/Editor/GraphEditor/Public/SGraphNode.h
@@ -90,6 +90,9 @@ public:
 	/** Set attribute for determining if widget is editable */
 	void SetIsEditable(TAttribute<bool> InIsEditable);
 
+	/** Returns if widget is editable, additionally considers if the owning graph is read only */
+	bool IsNodeEditable() const;
+
 	/** Set event when node is double clicked */
 	void SetDoubleClickEvent(FSingleNodeEvent InDoubleClickEvent);
 
