Fix issues with state machine transition creation caused by CL 2109030

diff --git a/Engine/Source/Editor/GraphEditor/Private/GraphEditorDragDropAction.cpp b/Engine/Source/Editor/GraphEditor/Private/GraphEditorDragDropAction.cpp
--- a/Engine/Source/Editor/GraphEditor/Private/GraphEditorDragDropAction.cpp
+++ b/Engine/Source/Editor/GraphEditor/Private/GraphEditorDragDropAction.cpp
@@ -113,14 +113,20 @@ void FGraphEditorDragDropAction::SetFeedbackMessage(const TSharedPtr<SWidget>& M
 void FGraphEditorDragDropAction::SetSimpleFeedbackMessage(const FSlateBrush* Icon, const FSlateColor& IconColor, const FText& Message)
 {
 	// Let the user know the status of making this connection.
-	SetFeedbackMessage(
+
+	// Use CreateRaw as we cannot using anything that will create a shared ptr from within an objects construction, this should be
+	// safe though as we will destroy our window before we get destroyed.
+	TAttribute<EVisibility> ErrorIconVisibility = TAttribute<EVisibility>::Create(TAttribute<EVisibility>::FGetter::CreateRaw(this, &FGraphEditorDragDropAction::GetErrorIconVisible));
+	TAttribute<EVisibility> IconVisibility = TAttribute<EVisibility>::Create(TAttribute<EVisibility>::FGetter::CreateRaw(this, &FGraphEditorDragDropAction::GetIconVisible));
+	
+		SetFeedbackMessage(
 		SNew(SHorizontalBox)
 		+SHorizontalBox::Slot()
 		.AutoWidth()
 		.Padding(3.0f)
 		[
 			SNew(SImage)
-			.Visibility( this, &FGraphEditorDragDropAction::GetErrorIconVisible )
+			.Visibility(ErrorIconVisibility)
 			.Image( FEditorStyle::GetBrush( TEXT("Graph.ConnectorFeedback.Error") ))
 			.ColorAndOpacity( FLinearColor::White )
 		]
@@ -129,7 +135,7 @@ void FGraphEditorDragDropAction::SetSimpleFeedbackMessage(const FSlateBrush* Ico
 		.Padding(3.0f)
 		[
 			SNew(SImage) 
-			.Visibility( this, &FGraphEditorDragDropAction::GetIconVisible )
+			.Visibility(IconVisibility)
 			.Image( Icon )
 			.ColorAndOpacity( IconColor )
 		]
diff --git a/Engine/Source/Editor/GraphEditor/Private/SGraphNode.cpp b/Engine/Source/Editor/GraphEditor/Private/SGraphNode.cpp
--- a/Engine/Source/Editor/GraphEditor/Private/SGraphNode.cpp
+++ b/Engine/Source/Editor/GraphEditor/Private/SGraphNode.cpp
@@ -112,8 +112,8 @@ void SGraphNode::SetIsEditable(TAttribute<bool> InIsEditable)
 
 bool SGraphNode::IsNodeEditable() const
 {
-	bool bReadOnly = OwnerGraphPanelPtr.IsValid() ? OwnerGraphPanelPtr.Pin()->IsGraphEditable() : false;
-	return IsEditable.Get() && !bReadOnly;
+	bool bIsEditable = OwnerGraphPanelPtr.IsValid() ? OwnerGraphPanelPtr.Pin()->IsGraphEditable() : true;
+	return IsEditable.Get() && bIsEditable;
 }
 
 /** Set event when node is double clicked */
