Stop multiple transition nodes between 2 state nodes being placed on top of each other.

#TTP 334329 - ANIMATION: Multiple transitions in the same direction between the same two nodes draw on top of each other

diff --git a/Engine/Source/Editor/GraphEditor/Private/AnimationStateNodes/SGraphNodeAnimTransition.cpp b/Engine/Source/Editor/GraphEditor/Private/AnimationStateNodes/SGraphNodeAnimTransition.cpp
--- a/Engine/Source/Editor/GraphEditor/Private/AnimationStateNodes/SGraphNodeAnimTransition.cpp
+++ b/Engine/Source/Editor/GraphEditor/Private/AnimationStateNodes/SGraphNodeAnimTransition.cpp
@@ -50,6 +50,9 @@ void SGraphNodeAnimTransition::PerformSecondPassLayout(const TMap< UObject*, TSh
 	FGeometry StartGeom;
 	FGeometry EndGeom;
 
+	int32 TransIndex = 0;
+	int32 NumOfTrans = 1;
+
 	UAnimStateNodeBase* PrevState = TransNode->GetPreviousState();
 	UAnimStateNodeBase* NextState = TransNode->GetNextState();
 	if ((PrevState != NULL) && (NextState != NULL))
@@ -63,24 +66,22 @@ void SGraphNodeAnimTransition::PerformSecondPassLayout(const TMap< UObject*, TSh
 
 			StartGeom = FGeometry(FVector2D(PrevState->NodePosX, PrevState->NodePosY), FVector2D::ZeroVector, PrevNodeWidget->GetDesiredSize(), 1.0f);
 			EndGeom = FGeometry(FVector2D(NextState->NodePosX, NextState->NodePosY), FVector2D::ZeroVector, NextNodeWidget->GetDesiredSize(), 1.0f);
+
+			TArray<UAnimStateTransitionNode*> Transitions;
+			PrevState->GetTransitionList(Transitions);
+
+			Transitions = Transitions.FilterByPredicate([NextState](const UAnimStateTransitionNode* InTransition) -> bool
+			{
+				return InTransition->GetNextState() == NextState;
+			});
+
+			TransIndex = Transitions.IndexOfByKey(TransNode);
+			NumOfTrans = Transitions.Num();
 		}
 	}
 
-	//@TODO: Duplicated code from ConnectionDrawingPolicy.cpp
-	{
-		// Get a reasonable seed point (halfway between the boxes)
-		const FVector2D StartCenter = FGeometryHelper::CenterOf(StartGeom);
-		const FVector2D EndCenter = FGeometryHelper::CenterOf(EndGeom);
-		const FVector2D SeedPoint = (StartCenter + EndCenter) * 0.5f;
-
-		// Find the (approximate) closest points between the two boxes
-		const FVector2D StartAnchorPoint = FGeometryHelper::FindClosestPointOnGeom(StartGeom, SeedPoint);
-		const FVector2D EndAnchorPoint = FGeometryHelper::FindClosestPointOnGeom(EndGeom, SeedPoint);
-
-		// Position ourselves halfway along the connecting line between the nodes, elevated away perpendicular to the direction of the line
-		const float Height = 30.0f;
-		PositionThisNodeBetweenOtherNodes(StartAnchorPoint, EndAnchorPoint, Height);
-	}
+	//Position Node
+	PositionBetweenTwoNodesWithOffset(StartGeom, EndGeom, TransIndex, NumOfTrans);
 }
 
 TSharedRef<SWidget> SGraphNodeAnimTransition::GenerateRichTooltip()
@@ -331,6 +332,52 @@ void SGraphNodeAnimTransition::OnMouseLeave(const FPointerEvent& MouseEvent)
 	SGraphNode::OnMouseLeave(MouseEvent);
 }
 
+void SGraphNodeAnimTransition::PositionBetweenTwoNodesWithOffset(const FGeometry& StartGeom, const FGeometry& EndGeom, int32 NodeIndex, int32 MaxNodes) const
+{
+	// Get a reasonable seed point (halfway between the boxes)
+	const FVector2D StartCenter = FGeometryHelper::CenterOf(StartGeom);
+	const FVector2D EndCenter = FGeometryHelper::CenterOf(EndGeom);
+	const FVector2D SeedPoint = (StartCenter + EndCenter) * 0.5f;
+
+	// Find the (approximate) closest points between the two boxes
+	const FVector2D StartAnchorPoint = FGeometryHelper::FindClosestPointOnGeom(StartGeom, SeedPoint);
+	const FVector2D EndAnchorPoint = FGeometryHelper::FindClosestPointOnGeom(EndGeom, SeedPoint);
+
+	// Position ourselves halfway along the connecting line between the nodes, elevated away perpendicular to the direction of the line
+	const float Height = 30.0f;
+
+	const FVector2D DesiredSize = GetDesiredSize();
+
+	FVector2D DeltaPos(EndAnchorPoint - StartAnchorPoint);
+
+	if (DeltaPos.IsNearlyZero())
+	{
+		DeltaPos = FVector2D(10.0f, 0.0f);
+	}
+
+	const FVector2D Normal = FVector2D(DeltaPos.Y, -DeltaPos.X).SafeNormal();
+
+	const FVector2D NewCenter = StartAnchorPoint + (0.5f * DeltaPos) + (Height * Normal);
+
+	FVector2D DeltaNormal = DeltaPos.SafeNormal();
+	
+	// Calculate node offset in the case of multiple transitions between the same two nodes
+	// MultiNodeOffset: the offset where 0 is the centre of the transition, -1 is 1 <size of node>
+	// towards the PrevStateNode and +1 is 1 <size of node> towards the NextStateNode.
+
+	const float MutliNodeSpace = 0.2f; // Space between multiple transition nodes (in units of <size of node> )
+	const float MultiNodeStep = (1.f + MutliNodeSpace); //Step between node centres (Size of node + size of node spacer)
+
+	const float MultiNodeStart = -((MaxNodes - 1) * MultiNodeStep) / 2.f;
+	const float MultiNodeOffset = MultiNodeStart + (NodeIndex * MultiNodeStep);
+
+	// Now we need to adjust the new center by the node size, zoom factor and multi node offset
+	const FVector2D NewCorner = NewCenter - (0.5f * DesiredSize) + (DeltaNormal * MultiNodeOffset * DesiredSize.Size());
+
+	GraphNode->NodePosX = NewCorner.X;
+	GraphNode->NodePosY = NewCorner.Y;
+}
+
 /////////////////////////////////////////////////////
 
 #undef LOCTEXT_NAMESPACE
\ No newline at end of file
diff --git a/Engine/Source/Editor/GraphEditor/Private/AnimationStateNodes/SGraphNodeAnimTransition.h b/Engine/Source/Editor/GraphEditor/Private/AnimationStateNodes/SGraphNodeAnimTransition.h
--- a/Engine/Source/Editor/GraphEditor/Private/AnimationStateNodes/SGraphNodeAnimTransition.h
+++ b/Engine/Source/Editor/GraphEditor/Private/AnimationStateNodes/SGraphNodeAnimTransition.h
@@ -27,6 +27,9 @@ public:
 	void OnMouseLeave(const FPointerEvent& MouseEvent) override;
 	// End of SWidget interface
 
+	// Calculate position for multiple nodes to be placed between a start and end point, by providing this nodes index and max expected nodes 
+	void PositionBetweenTwoNodesWithOffset(const FGeometry& StartGeom, const FGeometry& EndGeom, int32 NodeIndex, int32 MaxNodes) const;
+
 	static FLinearColor StaticGetTransitionColor(UAnimStateTransitionNode* TransNode, bool bIsHovered);
 private:
 	TSharedPtr<STextEntryPopup> TextEntryWidget;
