Editor: Moved Construct() call outside of FDragConnection onstructor

diff --git a/Engine/Source/Editor/GraphEditor/Private/DragConnection.cpp b/Engine/Source/Editor/GraphEditor/Private/DragConnection.cpp
--- a/Engine/Source/Editor/GraphEditor/Private/DragConnection.cpp
+++ b/Engine/Source/Editor/GraphEditor/Private/DragConnection.cpp
@@ -8,6 +8,7 @@
 TSharedRef<FDragConnection> FDragConnection::New(const TSharedRef<SGraphPanel>& InGraphPanel, const TArray< TSharedRef<SGraphPin> >& InStartingPins, bool bInShiftOperation)
 {
 	TSharedRef<FDragConnection> Operation = MakeShareable(new FDragConnection(InGraphPanel, InStartingPins, bInShiftOperation));
+	Operation->Construct();
 
 	return Operation;
 }
@@ -147,8 +148,6 @@ FDragConnection::FDragConnection(const TSharedRef<SGraphPanel>& InGraphPanel, co
 	{
 		InGraphPanel->OnBeginMakingConnection(PinRef->GetPinObj());
 	}
-
-	Construct();
 }
 
 FReply FDragConnection::DroppedOnPin(FVector2D ScreenPosition, FVector2D GraphPosition)
diff --git a/Engine/Source/Editor/GraphEditor/Private/DragConnection.h b/Engine/Source/Editor/GraphEditor/Private/DragConnection.h
--- a/Engine/Source/Editor/GraphEditor/Private/DragConnection.h
+++ b/Engine/Source/Editor/GraphEditor/Private/DragConnection.h
@@ -29,6 +29,7 @@ public:
 protected:
 	typedef FGraphEditorDragDropAction Super;
 
+	// Constructor: Make sure to call Construct() after factorying one of these
 	FDragConnection(const TSharedRef<SGraphPanel>& InGraphPanel, const TArray< TSharedRef<SGraphPin> >& InStartingPins, bool bInShiftOperation);
 
 protected:
diff --git a/Engine/Source/Editor/GraphEditor/Private/SGraphNodeKnot.cpp b/Engine/Source/Editor/GraphEditor/Private/SGraphNodeKnot.cpp
--- a/Engine/Source/Editor/GraphEditor/Private/SGraphNodeKnot.cpp
+++ b/Engine/Source/Editor/GraphEditor/Private/SGraphNodeKnot.cpp
@@ -10,15 +10,6 @@
 class FAmbivalentDirectionDragConnection : public FDragConnection
 {
 public:
-	FAmbivalentDirectionDragConnection(UK2Node_Knot* InKnot, const TSharedRef<SGraphPanel>& InGraphPanel, const TArray< TSharedRef<SGraphPin> >& InStartingPins, bool bInShiftOperation)
-		: FDragConnection(InGraphPanel, InStartingPins, bInShiftOperation)
-		, KnotPtr(InKnot)
-		, StartScreenPos(FVector2D::ZeroVector)
-		, MostRecentScreenPos(FVector2D::ZeroVector)
-		, bLatchedStartScreenPos(false)
-	{
-	}
-		
 	// FDragDropOperation interface
 	virtual void OnDragged(const class FDragDropEvent& DragDropEvent) override;
 	// End of FDragDropOperation interface
@@ -29,6 +20,25 @@ public:
 	virtual void ValidateGraphPinList(TArray<UEdGraphPin*>& OutValidPins) override;
 	// End of FDragConnection interface
 
+
+	static TSharedRef<FAmbivalentDirectionDragConnection> New(UK2Node_Knot* InKnot, const TSharedRef<SGraphPanel>& InGraphPanel, const TArray< TSharedRef<SGraphPin> >& InStartingPins, bool bInShiftOperation)
+	{
+		TSharedRef<FAmbivalentDirectionDragConnection> Operation = MakeShareable(new FAmbivalentDirectionDragConnection(InKnot, InGraphPanel, InStartingPins, bInShiftOperation));
+		Operation->Construct();
+
+		return Operation;
+	}
+
+protected:
+	FAmbivalentDirectionDragConnection(UK2Node_Knot* InKnot, const TSharedRef<SGraphPanel>& InGraphPanel, const TArray< TSharedRef<SGraphPin> >& InStartingPins, bool bInShiftOperation)
+		: FDragConnection(InGraphPanel, InStartingPins, bInShiftOperation)
+		, KnotPtr(InKnot)
+		, StartScreenPos(FVector2D::ZeroVector)
+		, MostRecentScreenPos(FVector2D::ZeroVector)
+		, bLatchedStartScreenPos(false)
+	{
+	}
+
 protected:
 	TWeakObjectPtr<UK2Node_Knot> KnotPtr;
 
@@ -200,8 +210,7 @@ TSharedRef<SWidget>	SGraphPinKnot::GetDefaultValueWidget()
 
 TSharedRef<FDragDropOperation> SGraphPinKnot::SpawnPinDragEvent(const TSharedRef<SGraphPanel>& InGraphPanel, const TArray< TSharedRef<SGraphPin> >& InStartingPins, bool bInShiftOperation)
 {
-	TSharedRef<FAmbivalentDirectionDragConnection> Operation = MakeShareable(new FAmbivalentDirectionDragConnection(CastChecked<UK2Node_Knot>(GetPinObj()->GetOwningNode()), InGraphPanel, InStartingPins, bInShiftOperation));
-
+	TSharedRef<FAmbivalentDirectionDragConnection> Operation = FAmbivalentDirectionDragConnection::New(CastChecked<UK2Node_Knot>(GetPinObj()->GetOwningNode()), InGraphPanel, InStartingPins, bInShiftOperation);
 	return Operation;
 }
 
