Fixed name conflicts for slate render targets

diff --git a/Engine/Source/Editor/GraphEditor/Private/MaterialNodes/SGraphNodeMaterialBase.cpp b/Engine/Source/Editor/GraphEditor/Private/MaterialNodes/SGraphNodeMaterialBase.cpp
--- a/Engine/Source/Editor/GraphEditor/Private/MaterialNodes/SGraphNodeMaterialBase.cpp
+++ b/Engine/Source/Editor/GraphEditor/Private/MaterialNodes/SGraphNodeMaterialBase.cpp
@@ -10,7 +10,7 @@
 * Simple representation of the backbuffer that the preview canvas renders to
 * This class may only be accessed from the render thread
 */
-class FSlateBackBufferTarget : public FRenderTarget
+class FSlateMaterialPreviewRenderTarget : public FRenderTarget
 {
 public:
 	/** FRenderTarget interface */
@@ -119,7 +119,7 @@ FIntPoint FPreviewViewport::GetSize() const
 
 FPreviewElement::FPreviewElement()
 	: ExpressionPreview(NULL)
-	, RenderTarget(new FSlateBackBufferTarget)
+	, RenderTarget(new FSlateMaterialPreviewRenderTarget)
 	, bIsRealtime(false)
 {
 }
diff --git a/Engine/Source/Editor/GraphEditor/Private/MaterialNodes/SGraphNodeMaterialBase.h b/Engine/Source/Editor/GraphEditor/Private/MaterialNodes/SGraphNodeMaterialBase.h
--- a/Engine/Source/Editor/GraphEditor/Private/MaterialNodes/SGraphNodeMaterialBase.h
+++ b/Engine/Source/Editor/GraphEditor/Private/MaterialNodes/SGraphNodeMaterialBase.h
@@ -48,7 +48,7 @@ private:
 
 private:
 	/** Render target that the canvas renders to */
-	class FSlateBackBufferTarget* RenderTarget;
+	class FSlateMaterialPreviewRenderTarget* RenderTarget;
 	/** Render proxy for the expression preview */
 	FMaterialRenderProxy* ExpressionPreview;
 	/** Whether preview is using realtime values */
diff --git a/Engine/Source/Runtime/Engine/Private/Slate/DebugCanvas.cpp b/Engine/Source/Runtime/Engine/Private/Slate/DebugCanvas.cpp
--- a/Engine/Source/Runtime/Engine/Private/Slate/DebugCanvas.cpp
+++ b/Engine/Source/Runtime/Engine/Private/Slate/DebugCanvas.cpp
@@ -50,7 +50,7 @@ FProxyCounter FCanvasProxy::Counter;
  * Simple representation of the backbuffer that the debug canvas renders to
  * This class may only be accessed from the render thread
  */
-class FSlateBackBufferTarget : public FRenderTarget
+class FSlateCanvasRenderTarget : public FRenderTarget
 {
 public:
 	/** FRenderTarget interface */
@@ -89,7 +89,7 @@ private:
 FDebugCanvasDrawer::FDebugCanvasDrawer()
 	: GameThreadCanvas( NULL )
 	, RenderThreadCanvas( NULL )
-	, RenderTarget( new FSlateBackBufferTarget )
+	, RenderTarget( new FSlateCanvasRenderTarget )
 {}
 
 FDebugCanvasDrawer::~FDebugCanvasDrawer()
