Fixed Behavior Tree diffs being backwards

Re-reversed the comparison order of passed-in Blueprint diffs & internally fixed the order of graph comparisons. Now other graphs compare correctly, as well as Blueprint comparisons still being correct.

TTP# 337672 - Behaviour Tree Editor: Visual Diffing of Behavior Trees displays some changes backwards and others not at all.

diff --git a/Engine/Source/Editor/GraphEditor/Private/GraphDiffControl.cpp b/Engine/Source/Editor/GraphEditor/Private/GraphDiffControl.cpp
--- a/Engine/Source/Editor/GraphEditor/Private/GraphDiffControl.cpp
+++ b/Engine/Source/Editor/GraphEditor/Private/GraphDiffControl.cpp
@@ -524,17 +524,17 @@ bool FGraphDiffControl::DiffGraphs(UEdGraph* const LhsGraph, UEdGraph* const Rhs
 		TArray<FGraphDiffControl::FNodeMatch> NodeMatches;
 		TSet<UEdGraphNode const*> MatchedRhsNodes;
 
-		// march through the all the nodes in the lhs graph and look for matches 
-		for (auto NodeIt(LhsGraph->Nodes.CreateConstIterator()); NodeIt; ++NodeIt)
+		// march through the all the nodes in the rhs graph and look for matches 
+		for (auto NodeIt(RhsGraph->Nodes.CreateConstIterator()); NodeIt; ++NodeIt)
 		{
-			UEdGraphNode* const LhsNode = *NodeIt;
-			if (LhsNode == NULL)
+			UEdGraphNode* const RhsNode = *NodeIt;
+			if (RhsNode == NULL)
 			{
 				continue;
 			}
 
-			FGraphDiffControl::FNodeMatch NodeMatch = FGraphDiffControl::FindNodeMatch(RhsGraph, LhsNode, NodeMatches);
-			// if we found a corresponding node in the rhs graph, track it (so we
+			FGraphDiffControl::FNodeMatch NodeMatch = FGraphDiffControl::FindNodeMatch(LhsGraph, RhsNode, NodeMatches);
+			// if we found a corresponding node in the lhs graph, track it (so we
 			// can prevent future matches with the same nodes)
 			if (NodeMatch.IsValid())
 			{
@@ -545,17 +545,17 @@ bool FGraphDiffControl::DiffGraphs(UEdGraph* const LhsGraph, UEdGraph* const Rhs
 			bFoundDifferences |= NodeMatch.Diff(&DiffsOut);
 		}
 
-		// go through the rhs nodes to catch ones that may have been missing from the lhs graph
-		for (auto NodeIt(RhsGraph->Nodes.CreateConstIterator()); NodeIt; ++NodeIt)
+		// go through the lhs nodes to catch ones that may have been missing from the rhs graph
+		for (auto NodeIt(LhsGraph->Nodes.CreateConstIterator()); NodeIt; ++NodeIt)
 		{
-			UEdGraphNode* const RhsNode = *NodeIt;
+			UEdGraphNode* const LhsNode = *NodeIt;
 			// if this node has already been matched, move on
-			if ((RhsNode == NULL) || MatchedRhsNodes.Find(RhsNode))
+			if ((LhsNode == NULL) || MatchedRhsNodes.Find(LhsNode))
 			{
 				continue;
 			}
 
-			FGraphDiffControl::FNodeMatch NodeMatch = FGraphDiffControl::FindNodeMatch(LhsGraph, RhsNode, NodeMatches);
+			FGraphDiffControl::FNodeMatch NodeMatch = FGraphDiffControl::FindNodeMatch(RhsGraph, LhsNode, NodeMatches);
 
 			TArray<FDiffSingleResult> RhsDiffs;
 			if (NodeMatch.Diff(&RhsDiffs))
diff --git a/Engine/Source/Editor/Kismet/Private/SBlueprintDiff.cpp b/Engine/Source/Editor/Kismet/Private/SBlueprintDiff.cpp
--- a/Engine/Source/Editor/Kismet/Private/SBlueprintDiff.cpp
+++ b/Engine/Source/Editor/Kismet/Private/SBlueprintDiff.cpp
@@ -266,7 +266,7 @@ FText FListItemGraphToDiff::GetToolTip()
 void FListItemGraphToDiff::BuildDiffSourceArray()
 {
 	TArray<FDiffSingleResult> FoundDiffs;
-	FGraphDiffControl::DiffGraphs(GraphNew, GraphOld, FoundDiffs);
+	FGraphDiffControl::DiffGraphs(GraphOld, GraphNew, FoundDiffs);
 
 	DiffListSource.Empty();
 	for (auto DiffIt(FoundDiffs.CreateIterator()); DiffIt; ++DiffIt)
