Making it so routing exec nodes shows flow of execution when debugging.
(TTP #338146)

diff --git a/Engine/Source/Editor/GraphEditor/Private/ConnectionDrawingPolicy.cpp b/Engine/Source/Editor/GraphEditor/Private/ConnectionDrawingPolicy.cpp
--- a/Engine/Source/Editor/GraphEditor/Private/ConnectionDrawingPolicy.cpp
+++ b/Engine/Source/Editor/GraphEditor/Private/ConnectionDrawingPolicy.cpp
@@ -13,6 +13,7 @@
 #include "AnimationGraphSchema.h"
 #include "AnimGraphNode_Base.h"
 #include "Sound/SoundNode.h"
+#include "K2Node_Knot.h"
 
 DEFINE_LOG_CATEGORY_STATIC(LogConnectionDrawingPolicy, Log, All);
 
@@ -831,6 +832,31 @@ void FKismetConnectionDrawingPolicy::DetermineStyleOfExecWire(float& Thickness,
 	}
 }
 
+FKismetConnectionDrawingPolicy::FTimePair const* FKismetConnectionDrawingPolicy::BackTraceExecPath(UEdGraphPin const* const OutputPin, FExecPairingMap const* const NodeExecutionList)
+{
+	FTimePair const* FoundExecPath = nullptr;
+
+	UEdGraphNode const* const OwningNode = OutputPin->GetOwningNode();
+	if (UK2Node_Knot const* const KnotNode = Cast<UK2Node_Knot>(OwningNode))
+	{
+		UEdGraphPin const* const KnotInputPin = KnotNode->GetInputPin();
+		for (UEdGraphPin const* KnotInput : KnotInputPin->LinkedTo)
+		{
+			FoundExecPath = BackTraceExecPath(KnotInput, NodeExecutionList);
+			if (FoundExecPath != nullptr)
+			{
+				break;
+			}
+		}
+	}
+	else
+	{
+		FoundExecPath = NodeExecutionList->Find(OutputPin);
+	}
+
+	return FoundExecPath;
+}
+
 // Give specific editor modes a chance to highlight this connection or darken non-interesting connections
 void FKismetConnectionDrawingPolicy::DetermineWiringStyle(UEdGraphPin* OutputPin, UEdGraphPin* InputPin, /*inout*/ float& Thickness, /*inout*/ FLinearColor& WireColor, /*inout*/bool& bDrawBubbles, /*inout*/bool& bBidirectional)
 {
@@ -851,15 +877,30 @@ void FKismetConnectionDrawingPolicy::DetermineWiringStyle(UEdGraphPin* OutputPin
 		{
 			if (CanBuildRoadmap())
 			{
+				UEdGraphNode* InputNode = InputPin->GetOwningNode();
+				// knot nodes are removed from the graph at compile time, so we 
+				// have to follow them until we find something that would have 
+				// actually executed
+				while (UK2Node_Knot* KnotNode = Cast<UK2Node_Knot>(InputNode))
+				{
+					InputNode = nullptr;
+
+					UEdGraphPin* OutPin = KnotNode->GetOutputPin();
+					if (OutPin->LinkedTo.Num() > 0)
+					{
+						check(OutPin->LinkedTo.Num() == 1);
+						InputNode = OutPin->LinkedTo[0]->GetOwningNode();
+					}	
+				}
+
 				// track if this node connection was ran or not
 				bool bExecuted = false;
 
-				UEdGraphNode* InputNode = InputPin->GetOwningNode();
 				// if the node belonging to InputPin was actually executed
 				if (FExecPairingMap* ExecPaths = PredecessorPins.Find(InputNode))
 				{
 					// if the output pin is one of the pins that lead to InputNode being ran
-					if (FTimePair* ExecTiming = ExecPaths->Find(OutputPin))
+					if (FTimePair const* ExecTiming = BackTraceExecPath(OutputPin, ExecPaths))
 					{
 						bExecuted = true;
 						DetermineStyleOfExecWire(/*inout*/ Thickness, /*inout*/ WireColor, /*inout*/ bDrawBubbles, *ExecTiming);
diff --git a/Engine/Source/Editor/GraphEditor/Private/ConnectionDrawingPolicy.h b/Engine/Source/Editor/GraphEditor/Private/ConnectionDrawingPolicy.h
--- a/Engine/Source/Editor/GraphEditor/Private/ConnectionDrawingPolicy.h
+++ b/Engine/Source/Editor/GraphEditor/Private/ConnectionDrawingPolicy.h
@@ -166,6 +166,22 @@ private:
 	* @return True if the two nodes are connected in the graph, otherwise false.
 	*/
 	bool AreNodesGraphicallySequential(UEdGraphNode* InputNode, UEdGraphNode* OutputNode) const;
+
+	/**
+	 * With the introduction of "reroute" (UK2Node_Knot) nodes, we have to 
+	 * backtrace to find a node that was actually executed before the one in 
+	 * question (knot nodes are removed at compile time, and therefore never
+	 * executed).
+	 *
+	 * This function recursively searches for any pins feeding the OutputPin (or 
+	 * the OutputPin itself); it looks for a pin on a non-reroute node (one that
+	 * is listed in NodeExecutionList).
+	 * 
+	 * @param  OutputPin			An exec pin, leading into the node in question (an output pin on some other node).
+	 * @param  NodeExecutionList	A list of pins that feed the execution of a specific node.
+	 * @return Null if OutputPin did not cause the invocation of the node in question, otherwise a valid time pair denoting when the node was executed.
+	 */
+	FTimePair const* BackTraceExecPath(UEdGraphPin const* const OutputPin, FExecPairingMap const* const NodeExecutionList);
 };
 
 /////////////////////////////////////////////////////
