Don't show an edit box for pin types that don't need it

diff --git a/Engine/Source/Editor/GraphEditor/Private/KismetPins/SGraphPinNum.cpp b/Engine/Source/Editor/GraphEditor/Private/KismetPins/SGraphPinNum.cpp
--- a/Engine/Source/Editor/GraphEditor/Private/KismetPins/SGraphPinNum.cpp
+++ b/Engine/Source/Editor/GraphEditor/Private/KismetPins/SGraphPinNum.cpp
@@ -18,7 +18,7 @@ void SGraphPinNum::SetTypeInValue(const FText& NewTypeInValue, ETextCommit::Type
 	{
 		const UEdGraphSchema_K2* K2Schema = Cast<const UEdGraphSchema_K2>(GraphPinObj->GetSchema());
 
-		if(GraphPinObj->PinType.PinCategory == K2Schema->PC_Int)
+		if (GraphPinObj->PinType.PinCategory == K2Schema->PC_Int || GraphPinObj->PinType.PinCategory == K2Schema->PC_Byte)
 		{
 			int32 IntValue = FCString::Atoi(*TypeValueString);
 			GraphPinObj->GetSchema()->TrySetDefaultValue(*GraphPinObj, FString::FromInt(IntValue));
diff --git a/Engine/Source/Editor/GraphEditor/Private/NodeFactory.cpp b/Engine/Source/Editor/GraphEditor/Private/NodeFactory.cpp
--- a/Engine/Source/Editor/GraphEditor/Private/NodeFactory.cpp
+++ b/Engine/Source/Editor/GraphEditor/Private/NodeFactory.cpp
@@ -284,6 +284,10 @@ TSharedPtr<SGraphPin> FNodeFactory::CreatePinWidget(UEdGraphPin* InPin)
 		{
 			return SNew(SGraphPinNum, InPin);
 		}
+		else if (InPin->PinType.PinCategory == K2Schema->PC_String || InPin->PinType.PinCategory == K2Schema->PC_Name)
+		{
+			return SNew(SGraphPinString, InPin);
+		}
 		else if (InPin->PinType.PinCategory == K2Schema->PC_Struct)
 		{
 			UScriptStruct* ColorStruct = FindObjectChecked<UScriptStruct>(UObject::StaticClass(), TEXT("LinearColor"));
@@ -319,6 +323,10 @@ TSharedPtr<SGraphPin> FNodeFactory::CreatePinWidget(UEdGraphPin* InPin)
 			{
 				return SNew(SGraphPinEnum, InPin);
 			}
+			else
+			{
+				return SNew(SGraphPinNum, InPin);
+			}
 		}
 		else if ((InPin->PinType.PinCategory == K2Schema->PC_Wildcard) && (InPin->PinType.PinSubCategory == K2Schema->PSC_Index))
 		{
@@ -351,6 +359,6 @@ TSharedPtr<SGraphPin> FNodeFactory::CreatePinWidget(UEdGraphPin* InPin)
 		}
 	}
 
-	// If we didn't pick a custom pin widget, use the base (string) one.
-	return SNew(SGraphPinString, InPin);
+	// If we didn't pick a custom pin widget, use an uncustomized basic pin
+	return SNew(SGraphPin, InPin);
 }
