Differentiate between rotator and vector in labelling on "vector" graph pins

diff --git a/Engine/Source/Editor/GraphEditor/Private/KismetPins/SGraphPinVector.cpp b/Engine/Source/Editor/GraphEditor/Private/KismetPins/SGraphPinVector.cpp
--- a/Engine/Source/Editor/GraphEditor/Private/KismetPins/SGraphPinVector.cpp
+++ b/Engine/Source/Editor/GraphEditor/Private/KismetPins/SGraphPinVector.cpp
@@ -13,7 +13,7 @@ class SVectorTextBox : public SCompoundWidget
 {
 public:
 
-	SLATE_BEGIN_ARGS( SVectorTextBox ){}	
+	SLATE_BEGIN_ARGS( SVectorTextBox ) {}	
 		SLATE_ATTRIBUTE( FString, VisibleText_0 )
 		SLATE_ATTRIBUTE( FString, VisibleText_1 )
 		SLATE_ATTRIBUTE( FString, VisibleText_2 )
@@ -23,8 +23,9 @@ public:
 	SLATE_END_ARGS()
 
 	//Construct editable text boxes with the appropriate getter & setter functions along with tool tip text
-	void Construct( const FArguments& InArgs )
+	void Construct( const FArguments& InArgs, const bool bInIsRotator )
 	{
+		bIsRotator = bInIsRotator;
 		VisibleText_0 = InArgs._VisibleText_0;
 		VisibleText_1 = InArgs._VisibleText_1;
 		VisibleText_2 = InArgs._VisibleText_2;
@@ -48,14 +49,14 @@ public:
 					[
 						SNew( STextBlock )
 						.Font( FEditorStyle::GetFontStyle( "Graph.VectorEditableTextBox" ) )
-						.Text( LOCTEXT("VectorNodeXAxisValueLabel", "X") )
+						.Text( bIsRotator ? LOCTEXT("VectorNodeRollValueLabel", "R") : LOCTEXT("VectorNodeXAxisValueLabel", "X") )
 						.ColorAndOpacity( LabelClr )
 					]
 					.Value( this, &SVectorTextBox::GetTypeInValue_0 )
 					.OnValueCommitted( InArgs._OnFloatCommitted_Box_0 )
 					.Font( FEditorStyle::GetFontStyle( "Graph.VectorEditableTextBox" ) )
 					.UndeterminedString( LOCTEXT("MultipleValues", "Multiple Values") )
-					.ToolTipText(LOCTEXT("VectorNodeXAxisValueLabel_ToolTip", "X value"))
+					.ToolTipText( bIsRotator ? LOCTEXT("VectorNodeRollValueLabel_ToolTip", "Roll value") : LOCTEXT("VectorNodeXAxisValueLabel_ToolTip", "X value"))
 					.EditableTextBoxStyle( &FEditorStyle::GetWidgetStyle<FEditableTextBoxStyle>( "Graph.VectorEditableTextBox" ))
 					.BorderForegroundColor( FLinearColor::White )
 					.BorderBackgroundColor( FLinearColor::White )
@@ -70,14 +71,14 @@ public:
 					[
 						SNew( STextBlock )
 						.Font( FEditorStyle::GetFontStyle( "Graph.VectorEditableTextBox" ) )
-						.Text( LOCTEXT("VectorNodeYAxisValueLabel", "Y") )
+						.Text( bIsRotator ? LOCTEXT("VectorNodePitchValueLabel", "P") : LOCTEXT("VectorNodeYAxisValueLabel", "Y") )
 						.ColorAndOpacity( LabelClr )
 					]
 					.Value( this, &SVectorTextBox::GetTypeInValue_1 )
 					.OnValueCommitted( InArgs._OnFloatCommitted_Box_1 )
 					.Font( FEditorStyle::GetFontStyle( "Graph.VectorEditableTextBox" ) )
 					.UndeterminedString( LOCTEXT("MultipleValues", "Multiple Values") )
-					.ToolTipText(LOCTEXT("VectorNodeYAxisValueLabel_ToolTip", "Y value"))
+					.ToolTipText(bIsRotator ? LOCTEXT("VectorNodePitchValueLabel_ToolTip", "Pitch value") : LOCTEXT("VectorNodeYAxisValueLabel_ToolTip", "Y value"))
 					.EditableTextBoxStyle( &FEditorStyle::GetWidgetStyle<FEditableTextBoxStyle>( "Graph.VectorEditableTextBox" ))
 					.BorderForegroundColor( FLinearColor::White )
 					.BorderBackgroundColor( FLinearColor::White )
@@ -92,14 +93,14 @@ public:
 					[
 						SNew( STextBlock )
 						.Font( FEditorStyle::GetFontStyle( "Graph.VectorEditableTextBox" ) )
-						.Text( LOCTEXT("VectorNodeZAxisValueLabel", "Z") )
+						.Text( bIsRotator ? LOCTEXT("VectorNodeYawValueLabel", "Y") : LOCTEXT("VectorNodeZAxisValueLabel", "Z") )
 						.ColorAndOpacity( LabelClr )
 					]
 					.Value( this, &SVectorTextBox::GetTypeInValue_2 )
 					.OnValueCommitted( InArgs._OnFloatCommitted_Box_2 )
 					.Font( FEditorStyle::GetFontStyle( "Graph.VectorEditableTextBox" ) )
 					.UndeterminedString( LOCTEXT("MultipleValues", "Multiple Values") )
-					.ToolTipText(LOCTEXT("VectorNodeZAxisValueLabel_ToolTip", "Z value"))
+					.ToolTipText(bIsRotator ? LOCTEXT("VectorNodeYawValueLabel_Tooltip", "Yaw value") : LOCTEXT("VectorNodeZAxisValueLabel_ToolTip", "Z value"))
 					.EditableTextBoxStyle( &FEditorStyle::GetWidgetStyle<FEditableTextBoxStyle>( "Graph.VectorEditableTextBox" ))
 					.BorderForegroundColor( FLinearColor::White )
 					.BorderBackgroundColor( FLinearColor::White )
@@ -131,6 +132,8 @@ private:
 	TAttribute< FString >	VisibleText_0;
 	TAttribute< FString >	VisibleText_1;
 	TAttribute< FString >	VisibleText_2;
+
+	bool bIsRotator;
 };
 
 
@@ -151,7 +154,7 @@ TSharedRef<SWidget>	SGraphPinVector::GetDefaultValueWidget()
 	bIsRotator = (GraphPinObj->PinType.PinSubCategoryObject == RotatorStruct) ? true : false;
 	
 	//Create widget
-	return	SNew( SVectorTextBox )
+	return	SNew( SVectorTextBox, bIsRotator )
 			.VisibleText_0( this, &SGraphPinVector::GetCurrentValue_0 )
 			.VisibleText_1( this, &SGraphPinVector::GetCurrentValue_1 )
 			.VisibleText_2( this, &SGraphPinVector::GetCurrentValue_2 )
