Removed code to clamp the graph focus position to the bounding box of all graph nodes, as it creates a jarring visual glitch, particularly in an empty graph.
#ttp 332013 - The scroll wheel doesnt work properly in the material editor when a large comment box is created
#ttp 339978 - Graph Editor's Zoom to Cursor is super weird (bad) compared to ortho viewport's (good)
#branch UE4
#proj Editor.GraphEditor

diff --git a/Engine/Source/Editor/GraphEditor/Private/SNodePanel.cpp b/Engine/Source/Editor/GraphEditor/Private/SNodePanel.cpp
--- a/Engine/Source/Editor/GraphEditor/Private/SNodePanel.cpp
+++ b/Engine/Source/Editor/GraphEditor/Private/SNodePanel.cpp
@@ -1490,18 +1490,7 @@ void SNodePanel::ChangeZoomLevel(int32 ZoomLevelDelta, const FVector2D& WidgetSp
 	ZoomLevelFade.Play();
 
 	// Re-center the screen so that it feels like zooming around the cursor.
-	{
-		FSlateRect GraphBounds = ComputeSensibleGraphBounds();
-
-		// Make sure we are not zooming into/out into emptiness; otherwise the user will get lost..
-		const FVector2D ClampedPointToMaintainGraphSpace(
-			FMath::Clamp(PointToMaintainGraphSpace.X, GraphBounds.Left, GraphBounds.Right),
-			FMath::Clamp(PointToMaintainGraphSpace.Y, GraphBounds.Top, GraphBounds.Bottom)
-			);
-
-		this->ViewOffset = ClampedPointToMaintainGraphSpace - WidgetSpaceZoomOrigin / GetZoomAmount();
-	}
-
+	this->ViewOffset = PointToMaintainGraphSpace - WidgetSpaceZoomOrigin / GetZoomAmount();
 }
 
 bool SNodePanel::GetBoundsForSelectedNodes( class FSlateRect& Rect, float Padding )
