Improved trackpad magnify gesture responsiveness in SNodePanel

diff --git a/Engine/Source/Editor/GraphEditor/Private/SNodePanel.cpp b/Engine/Source/Editor/GraphEditor/Private/SNodePanel.cpp
--- a/Engine/Source/Editor/GraphEditor/Private/SNodePanel.cpp
+++ b/Engine/Source/Editor/GraphEditor/Private/SNodePanel.cpp
@@ -306,6 +306,7 @@ void SNodePanel::Construct()
 	OldViewOffset = ViewOffset;
 	OldZoomAmount = GetZoomAmount();
 	ZoomStartOffset = FVector2D::ZeroVector;
+	TotalGestureMagnify = 0.0f;
 
 	ScopedTransactionPtr.Reset();
 }
@@ -920,10 +921,15 @@ FReply SNodePanel::OnTouchGesture( const FGeometry& MyGeometry, const FPointerEv
 	const FVector2D& GestureDelta = GestureEvent.GetGestureDelta();
 	if (GestureType == EGestureEvent::Magnify)
 	{
-		// We want to zoom into this point; i.e. keep it the same fraction offset into the panel
-		const FVector2D WidgetSpaceCursorPos = MyGeometry.AbsoluteToLocal(GestureEvent.GetScreenSpacePosition());
-		const int32 ZoomLevelDelta = FMath::FloorToInt(GestureDelta.X * 10);
-		ChangeZoomLevel(ZoomLevelDelta, WidgetSpaceCursorPos, GestureEvent.IsControlDown());
+		TotalGestureMagnify += GestureDelta.X;
+		if (FMath::Abs(TotalGestureMagnify) > 0.07f)
+		{
+			// We want to zoom into this point; i.e. keep it the same fraction offset into the panel
+			const FVector2D WidgetSpaceCursorPos = MyGeometry.AbsoluteToLocal(GestureEvent.GetScreenSpacePosition());
+			const int32 ZoomLevelDelta = TotalGestureMagnify > 0.0f ? 1 : -1;
+			ChangeZoomLevel(ZoomLevelDelta, WidgetSpaceCursorPos, GestureEvent.IsControlDown());
+			TotalGestureMagnify = 0.0f;
+		}
 		return FReply::Handled();
 	}
 	else if (GestureType == EGestureEvent::Scroll)
@@ -935,6 +941,12 @@ FReply SNodePanel::OnTouchGesture( const FGeometry& MyGeometry, const FPointerEv
 	return FReply::Unhandled();
 }
 
+FReply SNodePanel::OnTouchEnded( const FGeometry& MyGeometry, const FPointerEvent& InTouchEvent )
+{
+	TotalGestureMagnify = 0.0f;
+	return FReply::Unhandled();
+}
+
 void SNodePanel::FindNodesAffectedByMarquee( FGraphPanelSelectionSet& OutAffectedNodes ) const
 {
 	OutAffectedNodes.Empty();
diff --git a/Engine/Source/Editor/GraphEditor/Public/SNodePanel.h b/Engine/Source/Editor/GraphEditor/Public/SNodePanel.h
--- a/Engine/Source/Editor/GraphEditor/Public/SNodePanel.h
+++ b/Engine/Source/Editor/GraphEditor/Public/SNodePanel.h
@@ -355,6 +355,7 @@ public:
 	virtual FReply OnKeyUp( const FGeometry& MyGeometry, const FKeyboardEvent& InKeyboardEvent ) override;
 	virtual void OnKeyboardFocusLost( const FKeyboardFocusEvent& InKeyboardFocusEvent ) override;
 	virtual FReply OnTouchGesture( const FGeometry& MyGeometry, const FPointerEvent& GestureEvent ) override;
+	virtual FReply OnTouchEnded( const FGeometry& MyGeometry, const FPointerEvent& InTouchEvent ) override;
 
 	// End of SWidget interface
 public:
@@ -568,6 +569,9 @@ protected:
 
 	/** Offset in the panel the user started the LMB+RMB zoom from */
 	FVector2D ZoomStartOffset;
+
+	/** Cumulative magnify delta from trackpad gesture */
+	float TotalGestureMagnify;
 public:
 	/** Nodes selected in this instance of the editor; the selection is per-instance of the GraphEditor */
 	FGraphSelectionManager SelectionManager;
