Fixed issues with zooming blueprint and material editor viewports with Magic Mouse

diff --git a/Engine/Source/Editor/GraphEditor/Private/SNodePanel.cpp b/Engine/Source/Editor/GraphEditor/Private/SNodePanel.cpp
--- a/Engine/Source/Editor/GraphEditor/Private/SNodePanel.cpp
+++ b/Engine/Source/Editor/GraphEditor/Private/SNodePanel.cpp
@@ -485,7 +485,7 @@ FReply SNodePanel::OnMouseButtonDown( const FGeometry& MyGeometry, const FPointe
 	TotalMouseDelta = 0;
 
 	if ((bIsLeftMouseButtonEffecting && bIsRightMouseButtonDown)
-	||  (bIsRightMouseButtonEffecting && bIsLeftMouseButtonDown))
+	||  (bIsRightMouseButtonEffecting && (bIsLeftMouseButtonDown || FSlateApplication::Get().IsUsingTrackpad())))
 	{
 		// Starting zoom by holding LMB+RMB
 		FReply ReplyState = FReply::Handled();
@@ -750,12 +750,12 @@ FReply SNodePanel::OnMouseButtonUp( const FGeometry& MyGeometry, const FPointerE
 	bool bRemoveSoftwareCursor = false;
 
 	if ((bIsLeftMouseButtonEffecting && bIsRightMouseButtonDown)
-	||  (bIsRightMouseButtonEffecting && bIsLeftMouseButtonDown))
+	||  (bIsRightMouseButtonEffecting && (bIsLeftMouseButtonDown || FSlateApplication::Get().IsUsingTrackpad())))
 	{
 		// Ending zoom by releasing LMB or RMB
 		ReplyState = FReply::Handled();
 
-		if (bIsLeftMouseButtonDown)
+		if (bIsLeftMouseButtonDown || FSlateApplication::Get().IsUsingTrackpad())
 		{
 			// If we released the right mouse button first, we need to cancel the software cursor display
 			bRemoveSoftwareCursor = true;
