Improved RMB handling in SNodePanel for Magic Mouse and trackpad

diff --git a/Engine/Source/Editor/GraphEditor/Private/SNodePanel.cpp b/Engine/Source/Editor/GraphEditor/Private/SNodePanel.cpp
--- a/Engine/Source/Editor/GraphEditor/Private/SNodePanel.cpp
+++ b/Engine/Source/Editor/GraphEditor/Private/SNodePanel.cpp
@@ -290,6 +290,7 @@ void SNodePanel::Construct()
 	DeferredMovementTargetObject = NULL;
 
 	bIsPanning = false;
+	bIsZoomingWithTrackpad = false;
 	IsEditable.Set(true);
 
 	ZoomLevelFade = FCurveSequence( 0.0f, 1.0f );
@@ -494,7 +495,11 @@ FReply SNodePanel::OnMouseButtonDown( const FGeometry& MyGeometry, const FPointe
 
 		DeferredMovementTargetObject = NULL; // clear any interpolation when you manually zoom
 		TotalMouseDeltaY = 0;
-		bShowSoftwareCursor = true;
+
+		if (!FSlateApplication::Get().IsUsingTrackpad()) // on trackpad we don't know yet if user wants to zoom or bring up the context menu
+		{
+			bShowSoftwareCursor = true;
+		}
 
 		if (bIsLeftMouseButtonEffecting)
 		{
@@ -619,6 +624,12 @@ FReply SNodePanel::OnMouseMove(const FGeometry& MyGeometry, const FPointerEvent&
 			ChangeZoomLevel(ZoomLevelDelta, ZoomStartOffset, MouseEvent.IsControlDown());
 
 			this->bIsPanning = false;
+
+			if (FSlateApplication::Get().IsUsingTrackpad() && ZoomLevelDelta != 0)
+			{
+				this->bIsZoomingWithTrackpad = true;
+				bShowSoftwareCursor = true;
+			}
 			return ReplyState;
 		}
 		else if (bIsRightMouseButtonDown)
@@ -750,7 +761,7 @@ FReply SNodePanel::OnMouseButtonUp( const FGeometry& MyGeometry, const FPointerE
 	bool bRemoveSoftwareCursor = false;
 
 	if ((bIsLeftMouseButtonEffecting && bIsRightMouseButtonDown)
-	||  (bIsRightMouseButtonEffecting && (bIsLeftMouseButtonDown || FSlateApplication::Get().IsUsingTrackpad())))
+	||  (bIsRightMouseButtonEffecting && (bIsLeftMouseButtonDown || (FSlateApplication::Get().IsUsingTrackpad() && bIsZoomingWithTrackpad))))
 	{
 		// Ending zoom by releasing LMB or RMB
 		ReplyState = FReply::Handled();
@@ -759,6 +770,7 @@ FReply SNodePanel::OnMouseButtonUp( const FGeometry& MyGeometry, const FPointerE
 		{
 			// If we released the right mouse button first, we need to cancel the software cursor display
 			bRemoveSoftwareCursor = true;
+			bIsZoomingWithTrackpad = false;
 			ReplyState.ReleaseMouseCapture();
 		}
 	}
diff --git a/Engine/Source/Editor/GraphEditor/Public/SNodePanel.h b/Engine/Source/Editor/GraphEditor/Public/SNodePanel.h
--- a/Engine/Source/Editor/GraphEditor/Public/SNodePanel.h
+++ b/Engine/Source/Editor/GraphEditor/Public/SNodePanel.h
@@ -551,6 +551,9 @@ protected:
 	/** Are we panning the view at the moment? */
 	bool bIsPanning;
 
+	/** Are we zooming the view with trackpad at the moment? */
+	bool bIsZoomingWithTrackpad;
+
 	/** The graph node widgets owned by this panel */
 	TSlotlessChildren<SNode> Children;
 	TSlotlessChildren<SNode> VisibleChildren;
