Modifying an item selected in My Blueprints using the details panel will no longer deselect the item.

Submitting a category in the BP details panel will no longer mark the Blueprint as structurally modified if nothing has changed.

#ttp 341162 - Regression: BP: After entering a tooltip for a variable, it gets deselected

diff --git a/Engine/Source/Editor/GraphEditor/Private/GraphActionNode.h b/Engine/Source/Editor/GraphEditor/Private/GraphActionNode.h
--- a/Engine/Source/Editor/GraphEditor/Private/GraphActionNode.h
+++ b/Engine/Source/Editor/GraphEditor/Private/GraphActionNode.h
@@ -361,6 +361,19 @@ public:
 		}
 	}
 
+	/** Returns the node's string used for display purposes */
+	FString GetNodeString() const
+	{
+		if( IsCategoryNode() )
+		{
+			return Category;
+		}
+		else if(Actions.Num() > 0 && Actions[0].IsValid())
+		{
+			return Actions[0]->MenuDescription.ToString();
+		}
+		return FString();
+	}
 protected:
 
 	FGraphActionNode( FString InCategory, int32 InGrouping, TSharedPtr<FEdGraphSchemaAction> InAction )
diff --git a/Engine/Source/Editor/GraphEditor/Private/SGraphActionMenu.cpp b/Engine/Source/Editor/GraphEditor/Private/SGraphActionMenu.cpp
--- a/Engine/Source/Editor/GraphEditor/Private/SGraphActionMenu.cpp
+++ b/Engine/Source/Editor/GraphEditor/Private/SGraphActionMenu.cpp
@@ -235,7 +235,7 @@ void SGraphActionMenu::RefreshAllActions(bool bPreserveExpansion, bool bHandleOn
 {
 	// Save Selection (of only the first selected thing)
 	TArray< TSharedPtr<FGraphActionNode> > SelectedNodes = TreeView->GetSelectedItems();
-	TSharedPtr<FEdGraphSchemaAction> SelectedAction = SelectedNodes.Num() > 0 ? SelectedNodes[0]->Actions[0] : NULL;
+	TSharedPtr<FGraphActionNode> SelectedAction = SelectedNodes.Num() > 0 ? SelectedNodes[0] : NULL;
 
 	AllActions.Empty();
 	OnCollectAllActions.ExecuteIfBound(AllActions);
@@ -244,28 +244,17 @@ void SGraphActionMenu::RefreshAllActions(bool bPreserveExpansion, bool bHandleOn
 	// Re-apply selection #0 if possible
 	if (SelectedAction.IsValid())
 	{
-		TArray<TSharedPtr<FGraphActionNode>> GraphNodes;
-		FilteredRootAction->GetAllNodes(GraphNodes, false);
-		for (int32 i = 0; i < GraphNodes.Num(); ++i)
+		// Clear the selection, we will be re-selecting the previous action
+		TreeView->ClearSelection();
+
+		if(bHandleOnSelectionEvent)
 		{
-			TSharedPtr<FEdGraphSchemaAction> GraphAction = GraphNodes[i]->Actions[0];
-			if (GraphAction.IsValid() && GraphAction->MenuDescription.ToString() == SelectedAction->MenuDescription.ToString())
-			{
-				// Clear the selection (if this node is already selected then setting it will have no effect)
-				TreeView->ClearSelection();
-				// Now set the selection
-				if(bHandleOnSelectionEvent)
-				{
-					TreeView->SetSelection(GraphNodes[i], ESelectInfo::OnMouseClick);
-				}
-				else
-				{
-					// If we do not want to handle the selection, set it directly so it will reselect the item but not handle the event.
-					TreeView->SetSelection(GraphNodes[i], ESelectInfo::Direct);
-				}
-				
-				break;
-			}
+			SelectItemByName(*SelectedAction->GetNodeString(), ESelectInfo::OnMouseClick, SelectedAction->SectionID, SelectedNodes[0]->IsCategoryNode());
+		}
+		else
+		{
+			// If we do not want to handle the selection, set it directly so it will reselect the item but not handle the event.
+			SelectItemByName(*SelectedAction->GetNodeString(), ESelectInfo::Direct, SelectedAction->SectionID, SelectedNodes[0]->IsCategoryNode());
 		}
 	}
 }
@@ -358,7 +347,7 @@ void SGraphActionMenu::GetCategorySubActions(TWeakPtr<FGraphActionNode> InAction
 	}
 }
 
-bool SGraphActionMenu::SelectItemByName(const FName& ItemName, ESelectInfo::Type SelectInfo )
+bool SGraphActionMenu::SelectItemByName(const FName& ItemName, ESelectInfo::Type SelectInfo, int32 SectionId/* = INDEX_NONE */, bool bIsCategory/* = false*/)
 {
 	if (ItemName != NAME_None)
 	{
@@ -369,15 +358,28 @@ bool SGraphActionMenu::SelectItemByName(const FName& ItemName, ESelectInfo::Type
 		for (int32 i = 0; i < GraphNodes.Num() && !SelectionNode.IsValid(); ++i)
 		{
 			TSharedPtr<FGraphActionNode> CurrentGraphNode = GraphNodes[i];
-
 			FEdGraphSchemaAction* GraphAction = CurrentGraphNode->Actions[0].Get();
-			if (GraphAction)
+
+			// If the user is attempting to select a category, make sure it's a category
+			if( CurrentGraphNode->IsCategoryNode() == bIsCategory )
 			{
-				if ((OnActionMatchesName.IsBound() && OnActionMatchesName.Execute(GraphAction, ItemName)) || GraphActionMenuHelpers::ActionMatchesName(GraphAction, ItemName))
+				if(SectionId == INDEX_NONE || CurrentGraphNode->SectionID == SectionId)
 				{
-					SelectionNode = GraphNodes[i];
+					if (GraphAction)
+					{
+						if ((OnActionMatchesName.IsBound() && OnActionMatchesName.Execute(GraphAction, ItemName)) || GraphActionMenuHelpers::ActionMatchesName(GraphAction, ItemName))
+						{
+							SelectionNode = GraphNodes[i];
 
-					break;
+							break;
+						}
+					}
+					else if( CurrentGraphNode->Category == FName::NameToDisplayString(ItemName.ToString(), false) )
+					{
+						SelectionNode = CurrentGraphNode;
+
+						break;
+					}
 				}
 			}
 
@@ -390,13 +392,26 @@ bool SGraphActionMenu::SelectItemByName(const FName& ItemName, ESelectInfo::Type
 				{
 					FEdGraphSchemaAction* ChildGraphAction = CurrentChildNode->Actions[ActionIndex].Get();
 
-					if(ChildGraphAction)
+					// If the user is attempting to select a category, make sure it's a category
+					if( CurrentChildNode->IsCategoryNode() == bIsCategory )
 					{
-						if ((OnActionMatchesName.IsBound() && OnActionMatchesName.Execute(GraphAction, ItemName)) || GraphActionMenuHelpers::ActionMatchesName(ChildGraphAction, ItemName))
+						if(SectionId == INDEX_NONE || CurrentChildNode->SectionID == SectionId)
 						{
-							SelectionNode = GraphNodes[i]->Children[ChildIdx];
-
-							break;
+							if(ChildGraphAction)
+							{
+								if ((OnActionMatchesName.IsBound() && OnActionMatchesName.Execute(GraphAction, ItemName)) || GraphActionMenuHelpers::ActionMatchesName(ChildGraphAction, ItemName))
+								{
+									SelectionNode = GraphNodes[i]->Children[ChildIdx];
+
+									break;
+								}
+							}
+							else if( CurrentChildNode->Category == FName::NameToDisplayString(ItemName.ToString(), false) )
+							{
+								SelectionNode = CurrentChildNode;
+
+								break;
+							}
 						}
 					}
 				}
diff --git a/Engine/Source/Editor/GraphEditor/Public/SGraphActionMenu.h b/Engine/Source/Editor/GraphEditor/Public/SGraphActionMenu.h
--- a/Engine/Source/Editor/GraphEditor/Public/SGraphActionMenu.h
+++ b/Engine/Source/Editor/GraphEditor/Public/SGraphActionMenu.h
@@ -210,9 +210,11 @@ public:
 	 *
 	 * @param	ItemName		The name of the item to select
 	 * @param	SelectInfo		The selection type
+	 * @param	SectionId		If known, the section Id to restrict the selection to, useful in the case of categories where they can exist multiple times
+	 * @param	bIsCategory		TRUE if the selection is a category, categories obey different rules and it's hard to re-select properly without this knowledge
 	 * @return					TRUE if the item was successfully selected or the tree cleared, FALSE if unsuccessful
 	 */
-	bool SelectItemByName(const FName& ItemName, ESelectInfo::Type SelectInfo = ESelectInfo::Direct );
+	bool SelectItemByName(const FName& ItemName, ESelectInfo::Type SelectInfo = ESelectInfo::Direct, int32 SectionId = INDEX_NONE, bool bIsCategory = false );
 
 	/** Expands any category with the associated name */
 	void ExpandCategory(const FString& CategoryName);
diff --git a/Engine/Source/Editor/Kismet/Private/SMyBlueprint.cpp b/Engine/Source/Editor/Kismet/Private/SMyBlueprint.cpp
--- a/Engine/Source/Editor/Kismet/Private/SMyBlueprint.cpp
+++ b/Engine/Source/Editor/Kismet/Private/SMyBlueprint.cpp
@@ -354,6 +354,7 @@ void SMyBlueprint::OnCategoryNameCommitted(const FText& InNewText, ETextCommit::
 		}
 		Refresh();
 		FBlueprintEditorUtils::MarkBlueprintAsModified(GetBlueprintObj());
+		BlueprintEditorPtr.Pin()->GetMyBlueprintWidget()->SelectItemByName(FName(*CategoryName.ToString()), ESelectInfo::OnMouseClick, InAction.Pin()->SectionID, true);
 	}
 }
 
@@ -1066,7 +1067,7 @@ FReply SMyBlueprint::OnCategoryDragged(const FString& InCategory, const FPointer
 void SMyBlueprint::OnGlobalActionSelected(const TArray< TSharedPtr<FEdGraphSchemaAction> >& InActions)
 {
 	// If an action is being selected, clear the LocalGraphActionMenu of any selection it has, this keeps it so that only one menu will ever have selection at a time
-	if(InActions.Num() && LocalGraphActionMenu.IsValid())
+	if(InActions.Num() && LocalGraphActionMenu.IsValid() && InActions.Num())
 	{
 		LocalGraphActionMenu->SelectItemByName(NAME_None);
 	}
@@ -1076,7 +1077,7 @@ void SMyBlueprint::OnGlobalActionSelected(const TArray< TSharedPtr<FEdGraphSchem
 void SMyBlueprint::OnLocalActionSelected(const TArray< TSharedPtr<FEdGraphSchemaAction> >& InActions)
 {
 	// If an action is being selected, clear the GraphActionMenu of any selection it has, this keeps it so that only one menu will ever have selection at a time
-	if(InActions.Num())
+	if(InActions.Num() && InActions.Num())
 	{
 		GraphActionMenu->SelectItemByName(NAME_None);
 	}
@@ -1993,7 +1994,7 @@ bool SMyBlueprint::CanRequestRenameOnActionNode() const
 	return false;
 }
 
-void SMyBlueprint::SelectItemByName(const FName& ItemName, ESelectInfo::Type SelectInfo)
+void SMyBlueprint::SelectItemByName(const FName& ItemName, ESelectInfo::Type SelectInfo, int32 SectionId/* = INDEX_NONE*/, bool bIsCategory/* = false*/)
 {
 	// Check if the graph action menu is being told to clear
 	if(ItemName == NAME_None)
@@ -2001,9 +2002,9 @@ void SMyBlueprint::SelectItemByName(const FName& ItemName, ESelectInfo::Type Sel
 		ClearGraphActionMenuSelection();
 	}
 	// Attempt to select the item in the main graph action menu, if that fails, attempt the same in LocalGraphActionMenu
-	else if( !GraphActionMenu->SelectItemByName(ItemName, SelectInfo) && LocalGraphActionMenu.IsValid() )
+	else if( !GraphActionMenu->SelectItemByName(ItemName, SelectInfo, SectionId, bIsCategory) && LocalGraphActionMenu.IsValid() )
 	{
-		LocalGraphActionMenu->SelectItemByName(ItemName, SelectInfo);
+		LocalGraphActionMenu->SelectItemByName(ItemName, SelectInfo, SectionId, bIsCategory);
 	}
 }
 
diff --git a/Engine/Source/Editor/Kismet/Private/SMyBlueprint.h b/Engine/Source/Editor/Kismet/Private/SMyBlueprint.h
--- a/Engine/Source/Editor/Kismet/Private/SMyBlueprint.h
+++ b/Engine/Source/Editor/Kismet/Private/SMyBlueprint.h
@@ -109,7 +109,7 @@ public:
 	void OnResetItemFilter();
 
 	/** Selects an item by name in either the main graph action menu or the local one */
-	void SelectItemByName(const FName& ItemName, ESelectInfo::Type SelectInfo = ESelectInfo::Direct);
+	void SelectItemByName(const FName& ItemName, ESelectInfo::Type SelectInfo = ESelectInfo::Direct, int32 SectionId = INDEX_NONE, bool bIsCategory = false);
 
 	/** Clears the selection in the graph action menus */
 	void ClearGraphActionMenuSelection();
diff --git a/Engine/Source/Editor/UnrealEd/Private/Kismet2/BlueprintEditorUtils.cpp b/Engine/Source/Editor/UnrealEd/Private/Kismet2/BlueprintEditorUtils.cpp
--- a/Engine/Source/Editor/UnrealEd/Private/Kismet2/BlueprintEditorUtils.cpp
+++ b/Engine/Source/Editor/UnrealEd/Private/Kismet2/BlueprintEditorUtils.cpp
@@ -2870,24 +2870,36 @@ void FBlueprintEditorUtils::SetBlueprintVariableCategory(UBlueprint* Blueprint,
 		UClass* OuterClass = CastChecked<UClass>(TargetProperty->GetOuter());
 		const bool bIsNativeVar = (OuterClass->ClassGeneratedBy == NULL);
 
+		// If the category does not change, we will not recompile the Blueprint
+		bool bIsCategoryChanged = false;
 		if (!bIsNativeVar)
 		{
 			TargetProperty->SetMetaData(TEXT("Category"), *SetCategory.ToString());
 			const int32 VarIndex = FBlueprintEditorUtils::FindNewVariableIndex(Blueprint, VarName);
 			if (VarIndex != INDEX_NONE)
 			{
-				Blueprint->NewVariables[VarIndex].Category = SetCategory;
+				bIsCategoryChanged = Blueprint->NewVariables[VarIndex].Category != SetCategory;
+				
+				if(bIsCategoryChanged)
+				{
+					Blueprint->NewVariables[VarIndex].Category = SetCategory;
+				}
 			}
 			else
 			{
 				const int32 SCS_NodeIndex = FBlueprintEditorUtils::FindSCS_Node(Blueprint, VarName);
 				if (SCS_NodeIndex != INDEX_NONE)
 				{
-					Blueprint->SimpleConstructionScript->GetAllNodes()[SCS_NodeIndex]->CategoryName = SetCategory;
+					bIsCategoryChanged = Blueprint->SimpleConstructionScript->GetAllNodes()[SCS_NodeIndex]->CategoryName != SetCategory;
+					
+					if(bIsCategoryChanged)
+					{
+						Blueprint->SimpleConstructionScript->GetAllNodes()[SCS_NodeIndex]->CategoryName = SetCategory;
+					}
 				}
 			}
 
-			if (bDontRecompile == false)
+			if (bDontRecompile == false && bIsCategoryChanged)
 			{
 				FBlueprintEditorUtils::MarkBlueprintAsStructurallyModified(Blueprint);
 			}
@@ -2897,10 +2909,16 @@ void FBlueprintEditorUtils::SetBlueprintVariableCategory(UBlueprint* Blueprint,
 	{
 		if(FBPVariableDescription* LocalVariable = FindLocalVariable(Blueprint, InLocalVarScope, VarName))
 		{
-			LocalVariable->SetMetaData(TEXT("Category"), *SetCategory.ToString());
-			LocalVariable->Category = SetCategory;
+			// If the category does not change, we will not recompile the Blueprint
+			bool bIsCategoryChanged = LocalVariable->Category != SetCategory;
+
+			if(bIsCategoryChanged)
+			{
+				LocalVariable->SetMetaData(TEXT("Category"), *SetCategory.ToString());
+				LocalVariable->Category = SetCategory;
+			}
 
-			if (bDontRecompile == false)
+			if (bDontRecompile == false && bIsCategoryChanged)
 			{
 				FBlueprintEditorUtils::MarkBlueprintAsStructurallyModified(Blueprint);
 			}
