Added a new documentation node to the blueprints to display udn documentation excerpts in the grapheditor.
#TTP 312311 - ROCKET: TASK: Add a "Documentation Node"
#Branch UE4
#Proj  BlueprintGraph, EdGraph, Kismet, KismetCompiler, GraphEditor, Documentation, EditorStyle

#Change Updated UEdGraphSchema to include an interface call to retrieve an FEdGraphSchemaAction to create documentation nodes. At this point only the K2 interfaces implement this.
#Change Updated UEdGraphSchema_K2 to include a call to retrieve an FEdGraphSchemaAction to create documentation nodes. This is used to add the actions to the blueprint palette and context menus.
#Add Added FEdGraphSchemaAction_K2AddDocumentation in EdGraphSchema_K2_Actions.h/cpp to implement the calls in UEdGraphSchema to create documenation nodes from palette and context menus.
#Change Modified FK2ActionMenuBuilder to provide a static call to create a documentation action in the same style as the comment. Additionally added calls to FK2ActionMenuBuilder::GetPaletteActions and FK2ActionMenuBuilder::GetContextAllowedNodeTypes to call this to add entries in the palette and context menus.
#Add Added a new Brush GraphEditor.Documentation_16x, for the context menu icon in SlateEditorStyle.cpp.
#Change Modified GetPaletteItemIcon in SBlueprintPalette.cpp to return the new icon for the DocumentationNode

#Change Modified FKismetCompilerContext::IsNodePure to include the Documentaton node in the drop through ensure test to prevent asserts on compling if a documentation node is present.

#Change Added an entry for Documentation node in FNodeFactory::CreateNodeWidget.

#Change Modified IDocumentationPage interface to provide the ability to provide a TextWrapAt Attribute so this can be set before creating excerpt content if desired.
#Change Modified the UDNParser to control text wrap at values in the created widgets using a float Attribute WrapAt, Added a set call in the DocumentationPage and made the default behaviour mimic the code it replaced.

#Add Added the class UEdGraphNode_Documentation implemented in EdGraphNode_Documentation.h and UEdGraph.cpp, this is the UObject implementation for the documentation nodes.
#Add Added the class SGraphNodeDocumentation as the GraphPanel implementation of the EdGraphNode_Documentation.
#Change Moved the resizable code from SGraphNodeComment into a SGraphNodeResizable and changed SGraphNodeComment and SGraphNodeDocumentation inherit from it to avoid duplicating code.

#Change Added a documentation specific details customisation so the excerpts can be displayed as combo button.
#Change Added FBlueprintDocumentationDetails into BlueprintDetailsCustomization.h/cpp to handle the user interaction with the documentation node in the BP Editor.

diff --git a/Engine/Source/Editor/BlueprintGraph/Classes/EdGraphSchema_K2.h b/Engine/Source/Editor/BlueprintGraph/Classes/EdGraphSchema_K2.h
--- a/Engine/Source/Editor/BlueprintGraph/Classes/EdGraphSchema_K2.h
+++ b/Engine/Source/Editor/BlueprintGraph/Classes/EdGraphSchema_K2.h
@@ -474,6 +474,7 @@ public:
 	virtual UEdGraphNode* CreateSubstituteNode(UEdGraphNode* Node, const UEdGraph* Graph, FObjectInstancingGraph* InstanceGraph) const override;
 	virtual int32 GetNodeSelectionCount(const UEdGraph* Graph) const override;
 	virtual TSharedPtr<FEdGraphSchemaAction> GetCreateCommentAction() const override;
+	virtual TSharedPtr<FEdGraphSchemaAction> GetCreateDocumentNodeAction() const override;
 	virtual bool FadeNodeWhenDraggingOffPin(const UEdGraphNode* Node, const UEdGraphPin* Pin) const override;
 	virtual void BackwardCompatibilityNodeConversion(UEdGraph* Graph, bool bOnlySafeChanges) const override;
 	virtual void SplitPin(UEdGraphPin* Pin) const override;
diff --git a/Engine/Source/Editor/BlueprintGraph/Classes/EdGraphSchema_K2_Actions.h b/Engine/Source/Editor/BlueprintGraph/Classes/EdGraphSchema_K2_Actions.h
--- a/Engine/Source/Editor/BlueprintGraph/Classes/EdGraphSchema_K2_Actions.h
+++ b/Engine/Source/Editor/BlueprintGraph/Classes/EdGraphSchema_K2_Actions.h
@@ -382,6 +382,41 @@ struct FEdGraphSchemaAction_K2AddComment : public FEdGraphSchemaAction
 	// End of FEdGraphSchemaAction interface
 };
 
+/*******************************************************************************
+* FEdGraphSchemaAction_K2AddDocumentation
+*******************************************************************************/
+
+/** Action to add a 'documentation' node to the graph */
+USTRUCT()
+struct FEdGraphSchemaAction_K2AddDocumentation : public FEdGraphSchemaAction
+{
+	GENERATED_USTRUCT_BODY()
+
+	// Simple type info
+	static FString StaticGetTypeId() {static FString Type = TEXT("FEdGraphSchemaAction_K2AddDocumentation"); return Type;}
+	virtual FString GetTypeId() const override { return StaticGetTypeId(); } 
+
+	FEdGraphSchemaAction_K2AddDocumentation()
+		: FEdGraphSchemaAction()
+	{
+		Category = TEXT("");
+		MenuDescription = NSLOCTEXT("K2AddDocumentation", "AddDocumentation", "Add Documentation ...");
+		TooltipDescription = NSLOCTEXT("K2AddDocumentation", "AddDocumentation_Tooltip", "Creates a Documentation Node.").ToString();
+	}
+
+	FEdGraphSchemaAction_K2AddDocumentation(const FString& InDescription, const FString& InToolTip)
+		: FEdGraphSchemaAction()
+	{
+		Category = TEXT("");
+		MenuDescription = FText::FromString( InDescription );
+		TooltipDescription = InToolTip;
+	}
+	
+	// FEdGraphSchemaAction interface
+	BLUEPRINTGRAPH_API virtual UEdGraphNode* PerformAction(class UEdGraph* ParentGraph, UEdGraphPin* FromPin, const FVector2D Location, bool bSelectNewNode = true) override;
+	// End of FEdGraphSchemaAction interface
+};
+
 /*******************************************************************************
 * FEdGraphSchemaAction_K2TargetNode
 *******************************************************************************/
diff --git a/Engine/Source/Editor/BlueprintGraph/Private/EdGraphSchema_K2.cpp b/Engine/Source/Editor/BlueprintGraph/Private/EdGraphSchema_K2.cpp
--- a/Engine/Source/Editor/BlueprintGraph/Private/EdGraphSchema_K2.cpp
+++ b/Engine/Source/Editor/BlueprintGraph/Private/EdGraphSchema_K2.cpp
@@ -4302,6 +4302,11 @@ TSharedPtr<FEdGraphSchemaAction> UEdGraphSchema_K2::GetCreateCommentAction() con
 	return TSharedPtr<FEdGraphSchemaAction>(static_cast<FEdGraphSchemaAction*>(new FEdGraphSchemaAction_K2AddComment));
 }
 
+TSharedPtr<FEdGraphSchemaAction> UEdGraphSchema_K2::GetCreateDocumentNodeAction() const
+{
+	return TSharedPtr<FEdGraphSchemaAction>(static_cast<FEdGraphSchemaAction*>(new FEdGraphSchemaAction_K2AddDocumentation));
+}
+
 bool UEdGraphSchema_K2::CanDuplicateGraph(UEdGraph* InSourceGraph) const
 {
 	if(GetGraphType(InSourceGraph) == GT_Function)
diff --git a/Engine/Source/Editor/BlueprintGraph/Private/EdGraphSchema_K2_Actions.cpp b/Engine/Source/Editor/BlueprintGraph/Private/EdGraphSchema_K2_Actions.cpp
--- a/Engine/Source/Editor/BlueprintGraph/Private/EdGraphSchema_K2_Actions.cpp
+++ b/Engine/Source/Editor/BlueprintGraph/Private/EdGraphSchema_K2_Actions.cpp
@@ -11,6 +11,7 @@
 #include "ComponentAssetBroker.h"
 #include "Kismet2/KismetEditorUtilities.h"
 #include "Editor/UnrealEd/Public/EdGraphUtilities.h"
+#include "EdGraph/EdGraphNode_Documentation.h"
 
 #define SNAP_GRID (16) // @todo ensure this is the same as SNodePanel::GetSnapGridSize()
 
@@ -608,6 +609,35 @@ UEdGraphNode* FEdGraphSchemaAction_K2AddComment::PerformAction(class UEdGraph* P
 	return NewNode;
 }
 
+/////////////////////////////////////////////////////
+// FEdGraphSchemaAction_K2AddDocumentation
+
+UEdGraphNode* FEdGraphSchemaAction_K2AddDocumentation::PerformAction(class UEdGraph* ParentGraph, UEdGraphPin* FromPin, const FVector2D Location, bool bSelectNewNode/* = true*/)
+{
+	// Add menu item for creating document nodes
+	UEdGraphNode_Documentation* DocumentTemplate = NewObject<UEdGraphNode_Documentation>();
+
+	UBlueprint* Blueprint = FBlueprintEditorUtils::FindBlueprintForGraph(ParentGraph);
+
+	FVector2D SpawnLocation = Location;
+
+	FSlateRect Bounds;
+	if ((Blueprint != NULL) && FKismetEditorUtilities::GetBoundsForSelectedNodes(Blueprint, Bounds, 50.0f))
+	{
+		DocumentTemplate->SetBounds(Bounds);
+		SpawnLocation.X = DocumentTemplate->NodePosX;
+		SpawnLocation.Y = DocumentTemplate->NodePosY;
+	}
+
+	UEdGraphNode* NewNode = FEdGraphSchemaAction_NewNode::SpawnNodeFromTemplate<UEdGraphNode_Documentation>(ParentGraph, DocumentTemplate, SpawnLocation, bSelectNewNode);
+
+	// Mark Blueprint as structurally modified since
+	// UK2Node_Documentation::NodeCausesStructuralBlueprintChange used to return true
+	FBlueprintEditorUtils::MarkBlueprintAsStructurallyModified(Blueprint);
+
+	return NewNode;
+}
+
 /////////////////////////////////////////////////////
 // FEdGraphSchemaAction_K2TargetNode
 
diff --git a/Engine/Source/Editor/BlueprintGraph/Private/K2ActionMenuBuilder.cpp b/Engine/Source/Editor/BlueprintGraph/Private/K2ActionMenuBuilder.cpp
--- a/Engine/Source/Editor/BlueprintGraph/Private/K2ActionMenuBuilder.cpp
+++ b/Engine/Source/Editor/BlueprintGraph/Private/K2ActionMenuBuilder.cpp
@@ -297,6 +297,16 @@ static void GetCommentAction(UBlueprint const* BlueprintIn, FGraphActionListBuil
 	ActionMenuBuilder.AddAction( NewAction );
 }
 
+/** Gets the add documentation action */
+static void GetDocumentationAction(UBlueprint const* BlueprintIn, FGraphActionListBuilderBase& ActionMenuBuilder)
+{
+	const FString MenuDescription = LOCTEXT("AddDocumentationNode", "Add Documentation Node...").ToString();
+	const FString ToolTip = LOCTEXT("AddDocumentationToolTip", "Adds a node to display UDN documentation excerpts.").ToString();
+	TSharedPtr<FEdGraphSchemaAction_K2AddDocumentation> NewAction = TSharedPtr<FEdGraphSchemaAction_K2AddDocumentation>(new FEdGraphSchemaAction_K2AddDocumentation( MenuDescription, ToolTip ));
+
+	ActionMenuBuilder.AddAction( NewAction );
+}
+
 /** */
 static void GetEnumUtilitiesNodes(FBlueprintGraphActionListBuilder& ContextMenuBuilder, bool bNumEnum, bool bForEach, bool bCastFromByte, bool bLiteralByte, UEnum* OnlyEnum = NULL)
 {
@@ -766,6 +776,8 @@ void FK2ActionMenuBuilder::GetPaletteActions(FBlueprintPaletteListBuilder& Actio
 
 	GetCommentAction(ActionMenuBuilder.Blueprint, ActionMenuBuilder);
 
+	GetDocumentationAction(ActionMenuBuilder.Blueprint, ActionMenuBuilder);
+
 	GetMacroTools(ActionMenuBuilder, true, NULL);
 
 	GetRequestedAction<UK2Node_FormatText>(ActionMenuBuilder, K2ActionCategories::TextCategory);
@@ -1021,6 +1033,7 @@ void FK2ActionMenuBuilder::GetContextAllowedNodeTypes(FBlueprintGraphActionListB
 			}
 		}
 		GetCommentAction(ContextMenuBuilder.Blueprint, ContextMenuBuilder);
+		GetDocumentationAction(ContextMenuBuilder.Blueprint, ContextMenuBuilder);
 
 		{
 			UK2Node* NodeTemplate = ContextMenuBuilder.CreateTemplateNode<UK2Node_Select>();
diff --git a/Engine/Source/Editor/Documentation/Private/DocumentationPage.cpp b/Engine/Source/Editor/Documentation/Private/DocumentationPage.cpp
--- a/Engine/Source/Editor/Documentation/Private/DocumentationPage.cpp
+++ b/Engine/Source/Editor/Documentation/Private/DocumentationPage.cpp
@@ -75,6 +75,11 @@ void FDocumentationPage::Reload()
 	Parser->Parse( Link, StoredExcerpts, StoredMetadata );
 }
 
+void FDocumentationPage::SetTextWrapAt( TAttribute<float> WrapAt )
+{
+	Parser->SetWrapAt( WrapAt );
+}
+
 FDocumentationPage::FDocumentationPage( const FString& InLink, const TSharedRef< FUDNParser >& InParser ) 
 	: Link( InLink )
 	, Parser( InParser )
diff --git a/Engine/Source/Editor/Documentation/Private/DocumentationPage.h b/Engine/Source/Editor/Documentation/Private/DocumentationPage.h
--- a/Engine/Source/Editor/Documentation/Private/DocumentationPage.h
+++ b/Engine/Source/Editor/Documentation/Private/DocumentationPage.h
@@ -9,6 +9,7 @@ class FDocumentationPage : public IDocumentationPage
 {
 public:
 
+	/** Returns an instance of class containing content from the link argument */
 	static TSharedRef< IDocumentationPage > Create( const FString& Link, const TSharedRef< FUDNParser >& Parser );
 
 public:
@@ -24,16 +25,23 @@ public:
 
 	virtual void Reload() override;
 
+	virtual void SetTextWrapAt( TAttribute<float> WrapAt ) override;
+
 private:
 
 	FDocumentationPage( const FString& InLink, const TSharedRef< FUDNParser >& InParser );
 
 private:
 
+	/** The string representing the UDN page location */
 	FString Link;
+	/** The UDN parser instance used to create this page */
 	TSharedRef< FUDNParser > Parser;
 
+	/** The excerpts contained in this UDN page */
 	TArray<FExcerpt> StoredExcerpts;
+	/** The UDN meta data contained in this page */
 	FUDNPageMetadata StoredMetadata;
+	/** Signals if the page has been loaded */
 	bool IsLoaded;
 };
\ No newline at end of file
diff --git a/Engine/Source/Editor/Documentation/Private/UDNParser.cpp b/Engine/Source/Editor/Documentation/Private/UDNParser.cpp
--- a/Engine/Source/Editor/Documentation/Private/UDNParser.cpp
+++ b/Engine/Source/Editor/Documentation/Private/UDNParser.cpp
@@ -43,6 +43,8 @@ TSharedRef< FUDNParser > FUDNParser::Create( const TSharedPtr< FParserConfigurat
 FUDNParser::FUDNParser( const TSharedRef< FParserConfiguration >& InConfiguration, const FDocumentationStyle& InStyle )
 	: Configuration( InConfiguration )
 	, Style(InStyle)
+	, WrapAt(600.f)
+	, ContentWidth(600.f)
 {
 
 }
@@ -339,14 +341,9 @@ bool FUDNParser::GetExcerptContent( const FString& Link, FExcerpt& Excerpt )
 	return false;
 }
 
-float FUDNParser::GetTextWrapSize() const
+void FUDNParser::SetWrapAt( TAttribute<float> InWrapAt )
 {
-	return 600.f;
-}
-
-FOptionalSize FUDNParser::GetTutorialWidth() const 
-{
-	return GetTextWrapSize();
+	WrapAt = InWrapAt;
 }
 
 int32 FUDNParser::FTokenConfiguration::CalculatedExpectedContentStrings()
@@ -678,8 +675,8 @@ void FUDNParser::AppendExcerpt(TSharedPtr<SVerticalBox> Box, TSharedRef<SWidget>
 	.HAlign(HAlign_Center)
 	[
 		SNew(SBox)
-		.WidthOverride(this, &FUDNParser::GetTutorialWidth)
 		.HAlign(HAlign_Left)
+		.WidthOverride(ContentWidth)
 		.Padding(FMargin(0,0,0,8.f))
 		[
 			SNew(SHorizontalBox)
@@ -700,7 +697,7 @@ void FUDNParser::AddContentToExcerpt(TSharedPtr<SVerticalBox> Box, const FString
 			SNew(STextBlock)
 			.Text(ContentSource)
 			.TextStyle(FEditorStyle::Get(), Style.ContentStyleName)
-			.WrapTextAt(this, &FUDNParser::GetTextWrapSize)
+			.WrapTextAt(WrapAt)
 		);
 	}
 }
@@ -852,7 +849,7 @@ TSharedRef< SWidget > FUDNParser::GenerateExcerptContent( const FString& Link, F
 					.HAlign(HAlign_Center)
 					[
 						SNew(SBox)
-						.WidthOverride(this, &FUDNParser::GetTutorialWidth)
+						.WidthOverride(ContentWidth)
 						.Padding(FMargin(0,0,0,10))
 						[
 							SNew(SSeparator)
diff --git a/Engine/Source/Editor/Documentation/Private/UDNParser.h b/Engine/Source/Editor/Documentation/Private/UDNParser.h
--- a/Engine/Source/Editor/Documentation/Private/UDNParser.h
+++ b/Engine/Source/Editor/Documentation/Private/UDNParser.h
@@ -130,10 +130,11 @@ public:
 
 	bool GetExcerptContent( const FString& Link, FExcerpt& Excerpt );
 
+	/** Allows a TAttribute to be set to control excerpt wrapat values from outside the parser. */
+	void SetWrapAt( TAttribute<float> InWrapAt );
+
 private:
 	/** UI Callbacks for the widgets */
-	float GetTextWrapSize() const;
-	FOptionalSize GetTutorialWidth() const;
 	FReply OnImageLinkClicked(FString Payload);
 	
 	/** Adds the content text source to the scrollbox */
@@ -229,5 +230,9 @@ private:
 	};
 
 	TArray<FTokenConfiguration> LineLibrary;
+	/** Documentation text wrapping control attribute */
+	TAttribute<float> WrapAt;
+	/** Documentation optional width control attribute */
+	TAttribute<FOptionalSize> ContentWidth;
 };
 
diff --git a/Engine/Source/Editor/Documentation/Public/IDocumentationPage.h b/Engine/Source/Editor/Documentation/Public/IDocumentationPage.h
--- a/Engine/Source/Editor/Documentation/Public/IDocumentationPage.h
+++ b/Engine/Source/Editor/Documentation/Public/IDocumentationPage.h
@@ -5,10 +5,12 @@
 
 struct FExcerpt
 {
+	/** FExcerpt constructor */
 	FExcerpt()
 		: LineNumber( INDEX_NONE )
 	{}
 
+	/** FExcerpt overloaded constructor */
 	FExcerpt( const FString& InName, const TSharedPtr< SWidget >& InContent, const TMap< FString, FString >& InVariables, int32 InLineNumber )
 		: Name( InName )
 		, Content( InContent )
@@ -16,7 +18,9 @@ struct FExcerpt
 		, LineNumber( InLineNumber )
 	{}
 
+	/** Excerpt name */
 	FString Name;
+	/** Slate content for excerpt */
 	TSharedPtr<SWidget> Content;
 	TMap< FString, FString > Variables;
 	int32 LineNumber;
@@ -25,13 +29,25 @@ struct FExcerpt
 class IDocumentationPage
 {
 public:
+
+	/** Returns true if this page contains a match for the ExcerptName */
 	virtual bool HasExcerpt( const FString& ExcerptName ) = 0;
+	/** Return the number of excerpts this page holds */
 	virtual int32 GetNumExcerpts() const = 0;
+	/** Populates the argument excerpt with this content found using the excer */
 	virtual bool GetExcerpt( const FString& ExcerptName, FExcerpt& Excerpt) = 0;
+	/** Populates the argument TArray with Excerpts this page contains */
 	virtual void GetExcerpts( /*OUT*/ TArray< FExcerpt >& Excerpts ) = 0;
+	/** Builds the Excerpt content using the excerpt name from the argument */
 	virtual bool GetExcerptContent( FExcerpt& Excerpt ) = 0;
 
+	/** Returns the title of the excerpt */
 	virtual FText GetTitle() = 0;
 
+	/** Rebuilds the excerpt content */
 	virtual void Reload() = 0;
+
+	/** Sets the argument as the width control for text wrap in the excerpt widgets */
+	virtual void SetTextWrapAt( TAttribute<float> WrapAt ) = 0;
+
 };
\ No newline at end of file
diff --git a/Engine/Source/Editor/EditorStyle/Private/SlateEditorStyle.cpp b/Engine/Source/Editor/EditorStyle/Private/SlateEditorStyle.cpp
--- a/Engine/Source/Editor/EditorStyle/Private/SlateEditorStyle.cpp
+++ b/Engine/Source/Editor/EditorStyle/Private/SlateEditorStyle.cpp
@@ -3259,6 +3259,7 @@ void FSlateEditorStyle::FStyle::SetupGraphEditorStyles()
 		Set( "GraphEditor.CallInEditorEvent_16x", new IMAGE_BRUSH( "Icons/icon_Blueprint_CallInEditor_16x", Icon16x16 ) );
 		Set( "GraphEditor.Timeline_16x", new IMAGE_BRUSH("Icons/icon_Blueprint_Timeline_16x", Icon16x16));
 		Set( "GraphEditor.Comment_16x", new IMAGE_BRUSH("Icons/icon_Blueprint_Comment_16x", Icon16x16));
+		Set( "GraphEditor.Documentation_16x", new IMAGE_BRUSH("Icons/icon_Blueprint_Documentation_16x", Icon16x16));
 		Set( "GraphEditor.Switch_16x", new IMAGE_BRUSH("Icons/icon_Blueprint_Switch_16x", Icon16x16));
 		Set( "GraphEditor.BreakStruct_16x", new IMAGE_BRUSH("Icons/icon_Blueprint_BreakStruct_16x", Icon16x16));
 		Set( "GraphEditor.MakeStruct_16x", new IMAGE_BRUSH("Icons/icon_Blueprint_MakeStruct_16x", Icon16x16));
diff --git a/Engine/Source/Editor/GraphEditor/Private/GraphEditorCommon.h b/Engine/Source/Editor/GraphEditor/Private/GraphEditorCommon.h
--- a/Engine/Source/Editor/GraphEditor/Private/GraphEditorCommon.h
+++ b/Engine/Source/Editor/GraphEditor/Private/GraphEditorCommon.h
@@ -13,6 +13,7 @@
 #include "SNodePanel.h"
 #include "SGraphPanel.h"
 #include "SGraphNode.h"
+#include "SGraphNodeResizable.h"
 #include "SGraphPin.h"
 #include "BlueprintUtilities.h"
 #include "GraphEditorDragDropAction.h"
diff --git a/Engine/Source/Editor/GraphEditor/Private/NodeFactory.cpp b/Engine/Source/Editor/GraphEditor/Private/NodeFactory.cpp
--- a/Engine/Source/Editor/GraphEditor/Private/NodeFactory.cpp
+++ b/Engine/Source/Editor/GraphEditor/Private/NodeFactory.cpp
@@ -9,6 +9,8 @@
 
 #include "SGraphNodeDefault.h"
 #include "SGraphNodeComment.h"
+#include "SGraphNodeDocumentation.h"
+#include "EdGraph/EdGraphNode_Documentation.h"
 #include "SGraphNodeKnot.h"
 
 #include "KismetNodes/SGraphNodeK2Base.h"
@@ -219,6 +221,10 @@ TSharedPtr<SGraphNode> FNodeFactory::CreateNodeWidget(UEdGraphNode* InNode)
 	{
 		return SNew(SGraphNodeAnimStateEntry, EntryNode);
 	}
+	else if(UEdGraphNode_Documentation* DocNode = Cast<UEdGraphNode_Documentation>(InNode))
+	{
+		return SNew(SGraphNodeDocumentation, DocNode);
+	}
 	else if (InNode->ShouldDrawNodeAsComment())
 	{
 		if (UMaterialGraphNode_Comment* MaterialCommentNode = Cast<UMaterialGraphNode_Comment>(InNode))
diff --git a/Engine/Source/Editor/GraphEditor/Private/SGraphNodeComment.cpp b/Engine/Source/Editor/GraphEditor/Private/SGraphNodeComment.cpp
--- a/Engine/Source/Editor/GraphEditor/Private/SGraphNodeComment.cpp
+++ b/Engine/Source/Editor/GraphEditor/Private/SGraphNodeComment.cpp
@@ -47,265 +47,10 @@ void SGraphNodeComment::Construct(const FArguments& InArgs, UEdGraphNode* InNode
 	UserSize.X = InNode->NodeWidth;
 	UserSize.Y = InNode->NodeHeight;
 
-	MouseZone = CWZ_NotInWindow;
+	MouseZone = CRWZ_NotInWindow;
 	bUserIsDragging = false;
 }
 
-bool SGraphNodeComment::InSelectionArea(ECommentWindowZone InMouseZone) const
-{
-	return (	(InMouseZone == CWZ_RightBorder) || (InMouseZone == CWZ_BottomBorder) || (InMouseZone == CWZ_BottomRightBorder) || 
-				(InMouseZone == CWZ_LeftBorder) || (InMouseZone == CWZ_TopBorder) || (InMouseZone == CWZ_TopLeftBorder) ||
-				(InMouseZone == CWZ_TopRightBorder) || (InMouseZone == CWZ_BottomLeftBorder) );
-}
-
-FCursorReply SGraphNodeComment::OnCursorQuery( const FGeometry& MyGeometry, const FPointerEvent& CursorEvent ) const
-{
-	if (MouseZone == CWZ_RightBorder || MouseZone == CWZ_LeftBorder)
-	{
-		// right/left of comment
-		return FCursorReply::Cursor(EMouseCursor::ResizeLeftRight);
-	}
-	else if (MouseZone == CWZ_BottomRightBorder || MouseZone == CWZ_TopLeftBorder)
-	{
-		// bottom right / top left hand corner
-		return FCursorReply::Cursor( EMouseCursor::ResizeSouthEast );
-	}
-	else if (MouseZone == CWZ_BottomBorder || MouseZone == CWZ_TopBorder)
-	{
-		// bottom / top of comment
-		return FCursorReply::Cursor(EMouseCursor::ResizeUpDown);
-	}
-	else if (MouseZone == CWZ_BottomLeftBorder || MouseZone == CWZ_TopRightBorder)
-	{
-		// bottom left / top right hand corner
-		return FCursorReply::Cursor( EMouseCursor::ResizeSouthWest );
-	}
-	else if (MouseZone == CWZ_TitleBar) 
-	{
-		return FCursorReply::Cursor(EMouseCursor::CardinalCross);
-	}
-
-	return FCursorReply::Unhandled();
-}
-
-FReply SGraphNodeComment::OnMouseMove(const FGeometry& MyGeometry, const FPointerEvent& MouseEvent)
-{
-	FVector2D LocalMouseCoordinates = MyGeometry.AbsoluteToLocal( MouseEvent.GetScreenSpacePosition() );
-	
-	if (bUserIsDragging)
-	{
-		FVector2D GraphSpaceCoordinates = NodeCoordToGraphCoord( MouseEvent.GetScreenSpacePosition() );
-		FVector2D OldGraphSpaceCoordinates = NodeCoordToGraphCoord( MouseEvent.GetLastScreenSpacePosition() );
-		FVector2D Delta = GraphSpaceCoordinates - OldGraphSpaceCoordinates;
-		
-
-		//Clamp delta value based on resizing direction
-		if( MouseZone == CWZ_LeftBorder || MouseZone == CWZ_RightBorder )
-		{
-			Delta.Y = 0.0f;
-		}
-		else if( MouseZone == CWZ_TopBorder || MouseZone == CWZ_BottomBorder )
-		{
-			Delta.X = 0.0f;
-		}
-
-		//Resize node delta value
-		FVector2D DeltaNodeSize = Delta;
-
-		//Modify node size delta value based on resizing direction
-		if(	(MouseZone == CWZ_LeftBorder) || (MouseZone == CWZ_TopBorder) || (MouseZone == CWZ_TopLeftBorder) )
-		{
-			DeltaNodeSize = -DeltaNodeSize;
-		}
-		else if( MouseZone == CWZ_TopRightBorder )
-		{
-			DeltaNodeSize.Y = -DeltaNodeSize.Y;
-		}
-		else if( MouseZone == CWZ_BottomLeftBorder )
-		{
-			DeltaNodeSize.X = -DeltaNodeSize.X;
-		}
-
-		StoredUserSize.X += DeltaNodeSize.X;
-		StoredUserSize.Y += DeltaNodeSize.Y;
-
-		const FVector2D PreviousUserSize = UserSize;
-		const float SnapSize = SNodePanel::GetSnapGridSize();
-		UserSize.X = SnapSize * FMath::RoundToFloat(StoredUserSize.X/SnapSize);
-		UserSize.Y = SnapSize * FMath::RoundToFloat(StoredUserSize.Y/SnapSize);
-
-		FVector2D DeltaNodePos(0,0);
-		//Modify node position (resizing top and left sides)
-		if( MouseZone != CWZ_BottomBorder && MouseZone != CWZ_RightBorder && MouseZone != CWZ_BottomRightBorder )
-		{
-			//Delta value to move graph node position
-			DeltaNodePos = PreviousUserSize - UserSize;
-
-			//Clamp position delta based on resizing direction
-			if( MouseZone == CWZ_BottomLeftBorder )
-			{
-				DeltaNodePos.Y = 0.0f;
-			}
-			else if( MouseZone == CWZ_TopRightBorder )
-			{
-				DeltaNodePos.X = 0.0f;
-			}
-		}
-
-		if(UserSize.Y < SCommentNodeDefs::MinHeight)
-		{
-			UserSize.Y = SCommentNodeDefs::MinHeight;
-			DeltaNodePos.Y = GetCorrectedNodePosition().Y - GetPosition().Y;
-		}
-
-		if(UserSize.X < SCommentNodeDefs::MinWidth) 
-		{
-			UserSize.X = SCommentNodeDefs::MinWidth;
-			DeltaNodePos.X = GetCorrectedNodePosition().X - GetPosition().X;
-		}
-		SGraphNode::FNodeSet NodeFilter;
-		SGraphNode::MoveTo( GetPosition() + DeltaNodePos, NodeFilter );
-	}
-	else
-	{
-		MouseZone = FindMouseZone(LocalMouseCoordinates);
-	}
-
-	return SGraphNode::OnMouseMove( MyGeometry, MouseEvent );
-}
-
-void SGraphNodeComment::InitNodeAnchorPoint()
-{
-	NodeAnchorPoint = GetPosition();
-	
-	if( (MouseZone == CWZ_LeftBorder) || (MouseZone == CWZ_TopBorder) || (MouseZone == CWZ_TopLeftBorder) )
-	{
-		NodeAnchorPoint += UserSize;
-	}
-	else if( MouseZone == CWZ_BottomLeftBorder )
-	{
-		NodeAnchorPoint.X += UserSize.X;
-	}
-	else if( MouseZone == CWZ_TopRightBorder )
-	{
-		NodeAnchorPoint.Y += UserSize.Y;
-	}
-}
-
-
-FVector2D SGraphNodeComment::GetCorrectedNodePosition() const
-{
-	FVector2D CorrectedPos = NodeAnchorPoint;
-	
-	if( (MouseZone == CWZ_LeftBorder) || (MouseZone == CWZ_TopBorder) || (MouseZone == CWZ_TopLeftBorder) )
-	{
-		CorrectedPos -= UserSize;
-	}
-	else if( MouseZone == CWZ_BottomLeftBorder )
-	{
-		CorrectedPos.X -= UserSize.X;
-	}
-	else if( MouseZone == CWZ_TopRightBorder )
-	{
-		CorrectedPos.Y -= UserSize.Y;
-	}
-
-	return CorrectedPos;
-}
-
-SGraphNodeComment::ECommentWindowZone SGraphNodeComment::FindMouseZone(const FVector2D& LocalMouseCoordinates) const
-{
-	ECommentWindowZone InMouseZone = CWZ_NotInWindow;
-
-	const float InverseZoomFactor = 1.0f / FMath::Min(GetOwnerPanel()->GetZoomAmount(), 1.0f);
-	const FSlateRect HitResultBorderSize(
-		SCommentNodeDefs::HitResultBorderSize.Left * InverseZoomFactor,
-		SCommentNodeDefs::HitResultBorderSize.Right * InverseZoomFactor,
-		SCommentNodeDefs::HitResultBorderSize.Top * InverseZoomFactor, 
-		SCommentNodeDefs::HitResultBorderSize.Bottom * InverseZoomFactor);
-
-	const float TitleBarHeight = TitleBar.IsValid() ? TitleBar->GetDesiredSize().Y : 0.0f;
-
-	// Test for hit in location of 'grab' zone
-	if (LocalMouseCoordinates.Y > ( UserSize.Y - HitResultBorderSize.Bottom))
-	{
-		InMouseZone = CWZ_BottomBorder;
-	}
-	else if (LocalMouseCoordinates.Y <= (HitResultBorderSize.Top))
-	{
-		InMouseZone = CWZ_TopBorder;
-	}
-	else if (LocalMouseCoordinates.Y <= (HitResultBorderSize.Top + TitleBarHeight))
-	{
-		InMouseZone = CWZ_TitleBar;
-	}
-
-	if (LocalMouseCoordinates.X > (UserSize.X - HitResultBorderSize.Right))
-	{
-		if (InMouseZone == CWZ_BottomBorder)
-		{
-			InMouseZone = CWZ_BottomRightBorder;
-		}
-		else if (InMouseZone == CWZ_TopBorder)
-		{
-			InMouseZone = CWZ_TopRightBorder;
-		}
-		else
-		{
-			InMouseZone = CWZ_RightBorder;
-		}
-	}
-	else if (LocalMouseCoordinates.X <= HitResultBorderSize.Left)
-	{
-		if (InMouseZone == CWZ_TopBorder)
-		{
-			InMouseZone = CWZ_TopLeftBorder;
-		}
-		else if (InMouseZone == CWZ_BottomBorder)
-		{
-			InMouseZone = CWZ_BottomLeftBorder;
-		}
-		else
-		{
-			InMouseZone = CWZ_LeftBorder;
-		}
-	}
-	
-	// Test for hit on rest of frame
-	if (InMouseZone == CWZ_NotInWindow)
-	{
-		if (LocalMouseCoordinates.Y > HitResultBorderSize.Top) 
-		{
-			InMouseZone = CWZ_InWindow;
-		}
-		else if (LocalMouseCoordinates.X > HitResultBorderSize.Left) 
-		{
-			InMouseZone = CWZ_InWindow;
-		}
-	}
-
-	return InMouseZone;
-}
-
-void SGraphNodeComment::OnMouseEnter( const FGeometry& MyGeometry, const FPointerEvent& MouseEvent )
-{	
-	// Determine the zone the mouse is in
-	FVector2D LocalMouseCoordinates = MyGeometry.AbsoluteToLocal( MouseEvent.GetScreenSpacePosition() );
-	MouseZone = FindMouseZone(LocalMouseCoordinates);
-
-	SCompoundWidget::OnMouseEnter( MyGeometry, MouseEvent );
-}
-
-void SGraphNodeComment::OnMouseLeave( const FPointerEvent& MouseEvent ) 
-{
-	if( !bUserIsDragging )
-	{
-		// Reset our mouse zone
-		MouseZone = CWZ_NotInWindow;
-		SCompoundWidget::OnMouseLeave( MouseEvent );
-	}
-}
-
 void SGraphNodeComment::Tick(const FGeometry& AllottedGeometry, const double InCurrentTime, const float InDeltaTime)
 {
 	SGraphNode::Tick(AllottedGeometry, InCurrentTime, InDeltaTime);
@@ -422,7 +167,7 @@ FVector2D SGraphNodeComment::ComputeDesiredSize() const
 FReply SGraphNodeComment::OnMouseButtonDoubleClick( const FGeometry& InMyGeometry, const FPointerEvent& InMouseEvent )
 {
 	// If user double-clicked in the title bar area
-	if(FindMouseZone(InMyGeometry.AbsoluteToLocal(InMouseEvent.GetScreenSpacePosition())) == CWZ_TitleBar && IsEditable.Get())
+	if(FindMouseZone(InMyGeometry.AbsoluteToLocal(InMouseEvent.GetScreenSpacePosition())) == CRWZ_TitleBar && IsEditable.Get())
 	{
 		// Request a rename
 		RequestRename();
@@ -442,41 +187,6 @@ FReply SGraphNodeComment::OnMouseButtonDoubleClick( const FGeometry& InMyGeometr
 	}
 }
 
-FReply SGraphNodeComment::OnMouseButtonDown( const FGeometry& MyGeometry, const FPointerEvent& MouseEvent )
-{
-	if ( InSelectionArea() && (MouseEvent.GetEffectingButton() == EKeys::LeftMouseButton) && IsEditable.Get() ) 
-	{
-		bUserIsDragging = true;
-		StoredUserSize = UserSize;
-
-		//Find node anchor point
-		InitNodeAnchorPoint();
-
-		return FReply::Handled().CaptureMouse( SharedThis(this) );
-	}
-	else
-	{
-		HandleSelection(true, true);
-		return FReply::Unhandled();
-	}
-}
-
-FReply SGraphNodeComment::OnMouseButtonUp( const FGeometry& MyGeometry, const FPointerEvent& MouseEvent )
-{
-	if ( (MouseEvent.GetEffectingButton() == EKeys::LeftMouseButton) && bUserIsDragging )
-	{
-		bUserIsDragging = false;
-
-		// Resize the node	
-		UserSize.X = FPlatformMath::RoundToFloat(UserSize.X);
-		UserSize.Y =FPlatformMath::RoundToFloat(UserSize.Y);
-
-		GetNodeObj()->ResizeNode(UserSize);
-		return FReply::Handled().ReleaseMouseCapture();
-	}
-	return FReply::Unhandled();
-}
-
 void SGraphNodeComment::HandleSelection(bool bSelected, bool bUpdateNodesUnderComment) const
 {
 	if ((!this->bIsSelected && bSelected) || bUpdateNodesUnderComment)
@@ -539,16 +249,6 @@ void SGraphNodeComment::GetOverlayBrushes(bool bSelected, const FVector2D Widget
 	return SGraphNode::GetOverlayBrushes(bSelected, WidgetSize, Brushes);
 }
 
-bool SGraphNodeComment::OnHitTest( const FGeometry& MyGeometry, FVector2D InAbsoluteCursorPosition )
-{
-	FVector2D LocalMouseCoordinates = MyGeometry.AbsoluteToLocal( InAbsoluteCursorPosition );
-	ECommentWindowZone InMouseZone = FindMouseZone(LocalMouseCoordinates);
-	
-	//@TODO: really want to return true for CWZ_InWindow as well, *as long as* there are no other comment bubbles underneath that return one of the edges
-	// This is a brute force way to allow nested comment boxes to work usefully
-	return (InMouseZone != CWZ_NotInWindow) && (InMouseZone != CWZ_InWindow);
-}
-
 void SGraphNodeComment::MoveTo( const FVector2D& NewPosition, FNodeSet& NodeFilter ) 
 {
 	FVector2D PositionDelta = NewPosition - GetPosition();
@@ -581,6 +281,24 @@ void SGraphNodeComment::MoveTo( const FVector2D& NewPosition, FNodeSet& NodeFilt
 	}
 }
 
+float SGraphNodeComment::GetTitleBarHeight() const
+{
+	return TitleBar.IsValid() ? TitleBar->GetDesiredSize().Y : 0.0f;
+}
+
+FSlateRect SGraphNodeComment::GetHitTestingBorder( float InverseZoomFactor ) const
+{
+	return FSlateRect(	SCommentNodeDefs::HitResultBorderSize.Left * InverseZoomFactor,
+						SCommentNodeDefs::HitResultBorderSize.Top * InverseZoomFactor,
+						SCommentNodeDefs::HitResultBorderSize.Right * InverseZoomFactor,
+						SCommentNodeDefs::HitResultBorderSize.Bottom * InverseZoomFactor );
+}
+
+FVector2D SGraphNodeComment::GetNodeMaximumSize() const
+{
+	return FVector2D( UserSize.X + 100, UserSize.Y + 100 );
+}
+
 bool SGraphNodeComment::ShouldScaleNodeComment() const 
 {
 	return false;
@@ -617,8 +335,8 @@ FSlateColor SGraphNodeComment::GetCommentTitleBarColor() const
 
 bool SGraphNodeComment::CanBeSelected(const FVector2D& MousePositionInNode) const
 {
-	const ECommentWindowZone InMouseZone = FindMouseZone(MousePositionInNode);
-	return CWZ_TitleBar == InMouseZone;
+	const EResizableWindowZone InMouseZone = FindMouseZone(MousePositionInNode);
+	return CRWZ_TitleBar == MouseZone;
 }
 
 FVector2D SGraphNodeComment::GetDesiredSizeForMarquee() const
diff --git a/Engine/Source/Editor/GraphEditor/Private/SGraphNodeDocumentation.cpp b/Engine/Source/Editor/GraphEditor/Private/SGraphNodeDocumentation.cpp
new file mode 100644
--- /dev/null
+++ b/Engine/Source/Editor/GraphEditor/Private/SGraphNodeDocumentation.cpp
@@ -0,0 +1,333 @@
+// Copyright 1998-2013 Epic Games, Inc. All Rights Reserved.
+
+
+#include "GraphEditorCommon.h"
+#include "SGraphNodeDocumentation.h"
+
+#define LOCTEXT_NAMESPACE "SGraphNodeDocumentation"
+
+namespace GraphNodeDocumentationDefs
+{
+	/** Size of the hit result border for the window borders */
+	static const FSlateRect HitTestBorderSize( 10, 10, 8, 14 );
+
+	/** Minimum size for node */
+	static const FVector2D MinNodeSize( 200.0f, 10.0f );
+
+	/** Maximum size for node */
+	static const FVector2D MaximumNodeSize( 4000.0f, 10.0f );
+
+	/** Default documentation content size */
+	static const FVector2D DefaultContentSize( 600.0f, 400.0f );
+
+	/** Placeholder documentation content size */
+	static const FVector2D PlaceholderContentSize( 380.0f, 45.0f );
+
+	/** Default content border */
+	static const FMargin DefaultContentBorder( 4.0f, 2.0f, 4.0f, 10.0f );
+
+	/** Line wrap adjustment from node width, to account for scroll bar */
+	static const float LineWrapAdjustment = 20.f;
+}
+
+void SGraphNodeDocumentation::Construct( const FArguments& InArgs, UEdGraphNode* InNode )
+{
+	GraphNode = InNode;
+
+	// Set up animation
+	{
+		ZoomCurve = SpawnAnim.AddCurve( 0, 0.1f );
+		FadeCurve = SpawnAnim.AddCurve( 0.15f, 0.15f );
+	}
+	UserSize.X = InNode->NodeWidth;
+	UserSize.Y = InNode->NodeHeight;
+	bUserIsDragging = false;
+	UpdateGraphNode();
+}
+
+void SGraphNodeDocumentation::UpdateGraphNode()
+{
+	// No pins in a document box
+	InputPins.Empty();
+	OutputPins.Empty();
+
+	// Avoid standard box model too
+	RightNodeBox.Reset();
+	LeftNodeBox.Reset();
+
+	TSharedPtr<SWidget> ErrorText = SetupErrorReporting();
+
+	// Create Node Title
+	TSharedPtr<SNodeTitle> NodeTitle = SNew( SNodeTitle, GraphNode );
+
+	TSharedRef<SOverlay> DefaultTitleAreaWidget =
+		SNew( SOverlay )
+		+SOverlay::Slot()
+		[
+			SNew( SImage )
+			.Image( FEditorStyle::GetBrush( "Graph.Node.TitleGloss" ))
+		]
+		+SOverlay::Slot()
+		.HAlign( HAlign_Left )
+		.VAlign( VAlign_Center )
+		[
+			SNew( SBorder )
+			.BorderImage( FEditorStyle::GetBrush( "Graph.Node.ColorSpill" ))
+			.Padding( FMargin( 10, 5, 30, 3 ))
+			.BorderBackgroundColor( this, &SGraphNodeDocumentation::GetNodeTitleColor )
+			[
+				SNew( SVerticalBox )
+				+SVerticalBox::Slot()
+				.AutoHeight()
+				[
+					SAssignNew( InlineEditableText,SInlineEditableTextBlock )
+					.Style( FEditorStyle::Get(), "Graph.Node.NodeTitleInlineEditableText" )
+					.Text( GraphNode->GetNodeTitle( ENodeTitleType::EditableTitle ))
+				]
+				+SVerticalBox::Slot()
+				.AutoHeight()
+				[
+					NodeTitle.ToSharedRef()
+				]
+			]
+		]
+		+SOverlay::Slot()
+		.VAlign( VAlign_Top )
+		[
+			SNew( SBorder )
+			.Visibility( EVisibility::HitTestInvisible )			
+			.BorderImage( FEditorStyle::GetBrush( "Graph.Node.TitleHighlight" ))
+			[
+				SNew( SSpacer )
+				.Size( FVector2D( 20, 20 ))
+			]
+		];
+
+	// Create Node content
+	SAssignNew( TitleBar, SLevelOfDetailBranchNode )
+	.UseLowDetailSlot( this, &SGraphNodeDocumentation::UseLowDetailNodeTitles )
+	.LowDetail()
+	[
+		SNew(SBorder)
+		.BorderImage( FEditorStyle::GetBrush( "Graph.Node.ColorSpill" ))
+		.BorderBackgroundColor( this, &SGraphNodeDocumentation::GetNodeTitleColor )
+	]
+	.HighDetail()
+	[
+		DefaultTitleAreaWidget
+	];
+
+	// Create Documentation Page
+	TSharedPtr<SWidget> DocumentationPage = CreateDocumentationPage();
+
+	TSharedPtr<SVerticalBox> InnerVerticalBox;
+	ChildSlot
+	.HAlign( HAlign_Center )
+	.VAlign( VAlign_Center )
+	[
+		SNew( SVerticalBox )
+		+SVerticalBox::Slot()
+		.AutoHeight()
+		.HAlign( HAlign_Fill )
+		.VAlign( VAlign_Fill )
+		[
+			SNew( SBorder )
+			.BorderImage( FEditorStyle::GetBrush( "Graph.Node.Body" ))
+			.Visibility( this, &SGraphNodeDocumentation::GetWidgetVisibility )
+			.Padding( 0.f )
+			[
+				SAssignNew( InnerVerticalBox, SVerticalBox )
+				+SVerticalBox::Slot()
+				.AutoHeight()
+				.HAlign( HAlign_Fill )
+				.VAlign( VAlign_Top )
+				[
+					TitleBar.ToSharedRef()
+				]
+				+SVerticalBox::Slot()
+				.HAlign( HAlign_Left )
+				.VAlign( VAlign_Top )
+				[
+					DocumentationPage.ToSharedRef()
+				]
+			]
+		]
+	];
+}
+
+TSharedPtr<SWidget> SGraphNodeDocumentation::CreateDocumentationPage()
+{
+	TSharedPtr<SWidget> DocumentationWidget;
+
+	if( IDocumentation::Get()->PageExists( GraphNode->GetDocumentationLink() ))
+	{
+		TSharedRef< IDocumentationPage > DocumentationPage = IDocumentation::Get()->GetPage( GraphNode->GetDocumentationLink(), NULL );
+		const TMap< FString, FString > InVariables;
+		FExcerpt DesiredExcerpt( GraphNode->GetDocumentationExcerptName(), SNullWidget::NullWidget, InVariables, 0 );
+
+		// Set attributes to control documentation WrapAt and optional width
+		DocumentationPage->SetTextWrapAt( TAttribute<float>( this, &SGraphNodeDocumentation::GetDocumentationWrapWidth ) );
+
+		if( DocumentationPage->GetExcerptContent( DesiredExcerpt ))
+		{
+			// Create Content
+			SAssignNew( DocumentationWidget, SBox )
+			.WidthOverride( this, &SGraphNodeDocumentation::GetContentWidth )
+			.HeightOverride( this, &SGraphNodeDocumentation::GetContentHeight )
+			[
+				SAssignNew( ContentWidget, SVerticalBox )
+				+SVerticalBox::Slot()
+				.Padding( GraphNodeDocumentationDefs::DefaultContentBorder )
+				[
+					SNew( SBorder )
+					.HAlign( HAlign_Left )
+					.Content()
+					[
+						SNew( SScrollBox )
+						+SScrollBox::Slot()
+						[
+							DesiredExcerpt.Content.ToSharedRef()
+						]
+					]
+				]
+			];
+			// Find Maximum Content area for resizing
+			UserSize = GraphNodeDocumentationDefs::MaximumNodeSize;
+			ContentWidget->SlatePrepass();
+			DocumentationSize = ContentWidget->GetDesiredSize();
+			if( GraphNode->NodeWidth != 0.f && GraphNode->NodeHeight != 0.f )
+			{
+				// restore node size
+				UserSize.X = GraphNode->NodeWidth;
+				UserSize.Y = GraphNode->NodeHeight;
+			}
+			else
+			{
+				// set initial size
+				UserSize = GraphNodeDocumentationDefs::DefaultContentSize;
+				ContentWidget->SlatePrepass();
+				UserSize = ContentWidget->GetDesiredSize();
+			}
+		}
+	}
+	if( !DocumentationWidget.IsValid() )
+	{
+		// Create Placeholder
+		SAssignNew( DocumentationWidget, SBox )
+		.WidthOverride( this, &SGraphNodeDocumentation::GetContentWidth )
+		.HeightOverride( this, &SGraphNodeDocumentation::GetContentHeight )
+		[
+			SAssignNew( ContentWidget, SVerticalBox )
+			+SVerticalBox::Slot()
+			.Padding( GraphNodeDocumentationDefs::DefaultContentBorder )
+			[
+				SNew( SBorder )
+				.HAlign( HAlign_Left )
+				.Content()
+				[
+					SNew( SScrollBox )
+					+SScrollBox::Slot()
+					[
+						SNew( STextBlock )
+						.WrapTextAt( this, &SGraphNodeDocumentation::GetDocumentationWrapWidth )
+						.Text( LOCTEXT( "InvalidContentNotification", "No valid content to display.Please choose a valid link and excerpt in the details panel" ))
+					]
+				]
+			]
+		];
+		// Find Maximum Content area for resizing
+		UserSize = GraphNodeDocumentationDefs::PlaceholderContentSize;
+		DocumentationSize = GraphNodeDocumentationDefs::PlaceholderContentSize;
+	}
+	// Cache link/excerpt this widget was based on
+	CachedDocumentationLink = GraphNode->GetDocumentationLink();
+	CachedDocumentationExcerpt = GraphNode->GetDocumentationExcerptName();
+	
+	return DocumentationWidget;
+}
+
+FOptionalSize SGraphNodeDocumentation::GetContentWidth() const
+{
+	return UserSize.X;
+}
+
+FOptionalSize SGraphNodeDocumentation::GetContentHeight() const
+{
+	return UserSize.Y;
+}
+
+float SGraphNodeDocumentation::GetDocumentationWrapWidth() const
+{
+	return UserSize.X - GraphNodeDocumentationDefs::LineWrapAdjustment;
+}
+
+FVector2D SGraphNodeDocumentation::ComputeDesiredSize() const
+{
+	return FVector2D( UserSize.X, UserSize.Y + GetTitleBarHeight() );
+}
+
+FVector2D SGraphNodeDocumentation::GetNodeMinimumSize() const
+{
+	return GraphNodeDocumentationDefs::MinNodeSize;
+}
+
+FVector2D SGraphNodeDocumentation::GetNodeMaximumSize() const
+{
+	return FVector2D( DocumentationSize.X, ContentWidget->GetDesiredSize().Y );
+}
+
+FReply SGraphNodeDocumentation::OnMouseWheel( const FGeometry& MyGeometry, const FPointerEvent& MouseEvent )
+{
+	return FReply::Unhandled();
+}
+
+bool SGraphNodeDocumentation::OnHitTest( const FGeometry& MyGeometry, FVector2D InAbsoluteCursorPosition )
+{
+	if( !bUserIsDragging )
+	{
+		const FVector2D LocalMouseCoordinates = MyGeometry.AbsoluteToLocal( InAbsoluteCursorPosition );
+		return FindMouseZone( LocalMouseCoordinates ) != CRWZ_NotInWindow;
+	}
+	return true;
+}
+
+EVisibility SGraphNodeDocumentation::GetWidgetVisibility() const
+{
+	return ChildWidgetVisibility;
+}
+
+float SGraphNodeDocumentation::GetTitleBarHeight() const
+{
+	return TitleBar.IsValid() ? TitleBar->GetDesiredSize().Y : 0.f;
+}
+
+FSlateRect SGraphNodeDocumentation::GetHitTestingBorder( float InverseZoomFactor ) const
+{
+	return GraphNodeDocumentationDefs::HitTestBorderSize;
+}
+
+void SGraphNodeDocumentation::Tick( const FGeometry& AllottedGeometry, const double InCurrentTime, const float InDeltaTime )
+{
+	if( !bUserIsDragging )
+	{
+		ChildWidgetVisibility = EVisibility::HitTestInvisible;
+		FVector2D LocalMouseCoordinates = AllottedGeometry.AbsoluteToLocal( FSlateApplication::Get().GetCursorPos() );
+
+		EResizableWindowZone CurrMouseZone = FindMouseZone( LocalMouseCoordinates );
+		
+		if( CurrMouseZone == CRWZ_InWindow )
+		{
+			ChildWidgetVisibility = EVisibility::Visible;
+		}
+	}
+	// Check Cached Links to determine if we need to update the documention link/excerpt
+	if( CachedDocumentationLink != GraphNode->GetDocumentationLink() ||
+		CachedDocumentationExcerpt != GraphNode->GetDocumentationExcerptName())
+	{
+		GraphNode->NodeWidth = 0.f;
+		GraphNode->NodeHeight = 0.f;
+		UpdateGraphNode();
+	}
+
+	SGraphNodeResizable::Tick( AllottedGeometry, InCurrentTime, InDeltaTime );
+}
\ No newline at end of file
diff --git a/Engine/Source/Editor/GraphEditor/Private/SGraphNodeResizable.cpp b/Engine/Source/Editor/GraphEditor/Private/SGraphNodeResizable.cpp
new file mode 100644
--- /dev/null
+++ b/Engine/Source/Editor/GraphEditor/Private/SGraphNodeResizable.cpp
@@ -0,0 +1,339 @@
+// Copyright 1998-2014 Epic Games, Inc. All Rights Reserved.
+
+
+#include "GraphEditorCommon.h"
+#include "SGraphNodeResizable.h"
+
+namespace GraphNodeResizableDefs
+{
+	/** Size of the hit result border for the window borders */
+	static const FSlateRect HitResultBorderSize( 10, 10, 10, 10 );
+
+	/** Default Title Bar Size */
+	static const float DefaultTitleBarHeight = 12.f;
+
+	/** Minimum size for node */
+	static const FVector2D MinNodeSize( 30.0f, 30.0f );
+
+	/** Maximum size for node */
+	static const FVector2D MaxNodeSize( 400.0f, 400.0f );
+}
+
+bool SGraphNodeResizable::InSelectionArea(EResizableWindowZone InMouseZone) const
+{
+	return (	(InMouseZone == CRWZ_RightBorder) || (InMouseZone == CRWZ_BottomBorder) || (InMouseZone == CRWZ_BottomRightBorder) || 
+				(InMouseZone == CRWZ_LeftBorder) || (InMouseZone == CRWZ_TopBorder) || (InMouseZone == CRWZ_TopLeftBorder) ||
+				(InMouseZone == CRWZ_TopRightBorder) || (InMouseZone == CRWZ_BottomLeftBorder) );
+}
+
+void SGraphNodeResizable::OnMouseEnter( const FGeometry& MyGeometry, const FPointerEvent& MouseEvent )
+{	
+	// Determine the zone the mouse is in
+	if( !bUserIsDragging )
+	{
+		FVector2D LocalMouseCoordinates = MyGeometry.AbsoluteToLocal( MouseEvent.GetScreenSpacePosition() );
+		MouseZone = FindMouseZone(LocalMouseCoordinates);
+		SCompoundWidget::OnMouseEnter( MyGeometry, MouseEvent );
+	}
+}
+
+void SGraphNodeResizable::OnMouseLeave( const FPointerEvent& MouseEvent ) 
+{
+	if( !bUserIsDragging )
+	{
+		// Reset our mouse zone
+		MouseZone = CRWZ_NotInWindow;
+		SCompoundWidget::OnMouseLeave( MouseEvent );
+	}
+}
+
+FCursorReply SGraphNodeResizable::OnCursorQuery( const FGeometry& MyGeometry, const FPointerEvent& CursorEvent ) const
+{
+	if (MouseZone == CRWZ_RightBorder || MouseZone == CRWZ_LeftBorder)
+	{
+		// right/left of node
+		return FCursorReply::Cursor(EMouseCursor::ResizeLeftRight);
+	}
+	else if (MouseZone == CRWZ_BottomRightBorder || MouseZone == CRWZ_TopLeftBorder)
+	{
+		// bottom right / top left hand corner
+		return FCursorReply::Cursor( EMouseCursor::ResizeSouthEast );
+	}
+	else if (MouseZone == CRWZ_BottomBorder || MouseZone == CRWZ_TopBorder)
+	{
+		// bottom / top of node
+		return FCursorReply::Cursor(EMouseCursor::ResizeUpDown);
+	}
+	else if (MouseZone == CRWZ_BottomLeftBorder || MouseZone == CRWZ_TopRightBorder)
+	{
+		// bottom left / top right hand corner
+		return FCursorReply::Cursor( EMouseCursor::ResizeSouthWest );
+	}
+	else if (MouseZone == CRWZ_TitleBar) 
+	{
+		return FCursorReply::Cursor(EMouseCursor::CardinalCross);
+	}
+
+	return FCursorReply::Unhandled();
+}
+
+FReply SGraphNodeResizable::OnMouseButtonDown( const FGeometry& MyGeometry, const FPointerEvent& MouseEvent )
+{
+	if ( InSelectionArea() && (MouseEvent.GetEffectingButton() == EKeys::LeftMouseButton) && IsEditable.Get() ) 
+	{
+		bUserIsDragging = true;
+		StoredUserSize = UserSize;
+		DragSize = UserSize;
+		//Find node anchor point
+		InitNodeAnchorPoint();
+
+		return FReply::Handled().CaptureMouse( SharedThis(this) );
+	}
+	else
+	{
+		return FReply::Unhandled();
+	}
+}
+
+FReply SGraphNodeResizable::OnMouseButtonUp( const FGeometry& MyGeometry, const FPointerEvent& MouseEvent )
+{
+	if ( (MouseEvent.GetEffectingButton() == EKeys::LeftMouseButton) && bUserIsDragging )
+	{
+		bUserIsDragging = false;
+
+		// Resize the node	
+		UserSize.X = FMath::RoundToFloat(UserSize.X);
+		UserSize.Y = FMath::RoundToFloat(UserSize.Y);
+
+		GetNodeObj()->ResizeNode(UserSize);
+		return FReply::Handled().ReleaseMouseCapture();
+	}
+	return FReply::Unhandled();
+}
+
+bool SGraphNodeResizable::OnHitTest( const FGeometry& MyGeometry, FVector2D InAbsoluteCursorPosition )
+{
+	FVector2D LocalMouseCoordinates = MyGeometry.AbsoluteToLocal( InAbsoluteCursorPosition );
+	EResizableWindowZone InMouseZone = FindMouseZone(LocalMouseCoordinates);
+	return (InMouseZone != CRWZ_NotInWindow) && (InMouseZone != CRWZ_InWindow);
+}
+
+FReply SGraphNodeResizable::OnMouseMove(const FGeometry& MyGeometry, const FPointerEvent& MouseEvent)
+{
+	FVector2D LocalMouseCoordinates = MyGeometry.AbsoluteToLocal( MouseEvent.GetScreenSpacePosition() );
+	
+	if (bUserIsDragging)
+	{
+		FVector2D GraphSpaceCoordinates = NodeCoordToGraphCoord( MouseEvent.GetScreenSpacePosition() );
+		FVector2D OldGraphSpaceCoordinates = NodeCoordToGraphCoord( MouseEvent.GetLastScreenSpacePosition() );
+		FVector2D Delta = GraphSpaceCoordinates - OldGraphSpaceCoordinates;
+		
+
+		//Clamp delta value based on resizing direction
+		if( MouseZone == CRWZ_LeftBorder || MouseZone == CRWZ_RightBorder )
+		{
+			Delta.Y = 0.0f;
+		}
+		else if( MouseZone == CRWZ_TopBorder || MouseZone == CRWZ_BottomBorder )
+		{
+			Delta.X = 0.0f;
+		}
+
+		//Resize node delta value
+		FVector2D DeltaNodeSize = Delta;
+
+		//Modify node size delta value based on resizing direction
+		if(	(MouseZone == CRWZ_LeftBorder) || (MouseZone == CRWZ_TopBorder) || (MouseZone == CRWZ_TopLeftBorder) )
+		{
+			DeltaNodeSize = -DeltaNodeSize;
+		}
+		else if( MouseZone == CRWZ_TopRightBorder )
+		{
+			DeltaNodeSize.Y = -DeltaNodeSize.Y;
+		}
+		else if( MouseZone == CRWZ_BottomLeftBorder )
+		{
+			DeltaNodeSize.X = -DeltaNodeSize.X;
+		}
+		// Apply delta unfiltered to DragSize
+		DragSize.X += DeltaNodeSize.X;
+		DragSize.Y += DeltaNodeSize.Y;
+		// apply snap
+		const float SnapSize = SNodePanel::GetSnapGridSize();
+		FVector2D SnappedSize;
+		SnappedSize.X = SnapSize * FMath::RoundToFloat( DragSize.X/SnapSize );
+		SnappedSize.Y = SnapSize * FMath::RoundToFloat( DragSize.Y/SnapSize );
+
+		// Enforce min/max sizing
+		const FVector2D MinSize = GetNodeMinimumSize();
+		SnappedSize.X = FMath::Max( SnappedSize.X, MinSize.X );
+		SnappedSize.Y = FMath::Max( SnappedSize.Y, MinSize.Y );
+
+		const FVector2D MaxSize = GetNodeMaximumSize();
+		SnappedSize.X = FMath::Min( SnappedSize.X, MaxSize.X );
+		SnappedSize.Y = FMath::Min( SnappedSize.Y, MaxSize.Y );
+
+		FVector2D DeltaNodePos(0,0);
+
+		if( UserSize != SnappedSize )
+		{
+			//Modify node position (resizing top and left sides)
+			if( MouseZone != CRWZ_BottomBorder && MouseZone != CRWZ_RightBorder && MouseZone != CRWZ_BottomRightBorder )
+			{
+				//Delta value to move graph node position
+				DeltaNodePos = UserSize - SnappedSize;
+
+				//Clamp position delta based on resizing direction
+				if( MouseZone == CRWZ_BottomLeftBorder )
+				{
+					DeltaNodePos.Y = 0.0f;
+				}
+				else if( MouseZone == CRWZ_TopRightBorder )
+				{
+					DeltaNodePos.X = 0.0f;
+				}
+			}
+			UserSize = SnappedSize;
+			GraphNode->ResizeNode( UserSize );
+			DeltaNodePos = GetCorrectedNodePosition() - GetPosition();
+		}
+		SGraphNode::FNodeSet NodeFilter;
+		SGraphNode::MoveTo( GetPosition() + DeltaNodePos, NodeFilter );
+	}
+	else
+	{
+		MouseZone = FindMouseZone(LocalMouseCoordinates);
+	}
+	return SGraphNode::OnMouseMove( MyGeometry, MouseEvent );
+}
+
+void SGraphNodeResizable::InitNodeAnchorPoint()
+{
+	NodeAnchorPoint = GetPosition();
+	
+	if( (MouseZone == CRWZ_LeftBorder) || (MouseZone == CRWZ_TopBorder) || (MouseZone == CRWZ_TopLeftBorder) )
+	{
+		NodeAnchorPoint += UserSize;
+	}
+	else if( MouseZone == CRWZ_BottomLeftBorder )
+	{
+		NodeAnchorPoint.X += UserSize.X;
+	}
+	else if( MouseZone == CRWZ_TopRightBorder )
+	{
+		NodeAnchorPoint.Y += UserSize.Y;
+	}
+}
+
+
+FVector2D SGraphNodeResizable::GetCorrectedNodePosition() const
+{
+	FVector2D CorrectedPos = NodeAnchorPoint;
+	
+	if( (MouseZone == CRWZ_LeftBorder) || (MouseZone == CRWZ_TopBorder) || (MouseZone == CRWZ_TopLeftBorder) )
+	{
+		CorrectedPos -= UserSize;
+	}
+	else if( MouseZone == CRWZ_BottomLeftBorder )
+	{
+		CorrectedPos.X -= UserSize.X;
+	}
+	else if( MouseZone == CRWZ_TopRightBorder )
+	{
+		CorrectedPos.Y -= UserSize.Y;
+	}
+
+	return CorrectedPos;
+}
+
+SGraphNodeResizable::EResizableWindowZone SGraphNodeResizable::FindMouseZone(const FVector2D& LocalMouseCoordinates) const
+{
+	EResizableWindowZone OutMouseZone = CRWZ_NotInWindow;
+
+	const float InverseZoomFactor = 1.0f / FMath::Min(GetOwnerPanel()->GetZoomAmount(), 1.0f);
+	const FSlateRect HitResultBorderSize = GetHitTestingBorder( InverseZoomFactor );
+
+	const FVector2D NodeSize = GetDesiredSize();
+	// Test for hit in location of 'grab' zone
+	if (LocalMouseCoordinates.Y > ( NodeSize.Y - HitResultBorderSize.Bottom))
+	{
+		OutMouseZone = CRWZ_BottomBorder;
+	}
+	else if (LocalMouseCoordinates.Y <= (HitResultBorderSize.Top))
+	{
+		OutMouseZone = CRWZ_TopBorder;
+	}
+	else if (LocalMouseCoordinates.Y <= (HitResultBorderSize.Top + GetTitleBarHeight()))
+	{
+		OutMouseZone = CRWZ_TitleBar;
+	}
+
+	if (LocalMouseCoordinates.X > (NodeSize.X - HitResultBorderSize.Right))
+	{
+		if (OutMouseZone == CRWZ_BottomBorder)
+		{
+			OutMouseZone = CRWZ_BottomRightBorder;
+		}
+		else if (OutMouseZone == CRWZ_TopBorder)
+		{
+			OutMouseZone = CRWZ_TopRightBorder;
+		}
+		else
+		{
+			OutMouseZone = CRWZ_RightBorder;
+		}
+	}
+	else if (LocalMouseCoordinates.X <= HitResultBorderSize.Left)
+	{
+		if (OutMouseZone == CRWZ_TopBorder)
+		{
+			OutMouseZone = CRWZ_TopLeftBorder;
+		}
+		else if (OutMouseZone == CRWZ_BottomBorder)
+		{
+			OutMouseZone = CRWZ_BottomLeftBorder;
+		}
+		else
+		{
+			OutMouseZone = CRWZ_LeftBorder;
+		}
+	}
+	
+	// Test for hit on rest of frame
+	if (OutMouseZone == CRWZ_NotInWindow)
+	{
+		if (LocalMouseCoordinates.Y > HitResultBorderSize.Top) 
+		{
+			OutMouseZone = CRWZ_InWindow;
+		}
+		else if (LocalMouseCoordinates.X > HitResultBorderSize.Left) 
+		{
+			OutMouseZone = CRWZ_InWindow;
+		}
+	}
+	return OutMouseZone;
+}
+
+float SGraphNodeResizable::GetTitleBarHeight() const
+{
+	// this can probably just be SGraphNode::GetTitleRect().Height()
+	return GraphNodeResizableDefs::DefaultTitleBarHeight;
+}
+
+FVector2D SGraphNodeResizable::GetNodeMinimumSize() const
+{
+	return GraphNodeResizableDefs::MinNodeSize;
+}
+
+FVector2D SGraphNodeResizable::GetNodeMaximumSize() const
+{
+	return GraphNodeResizableDefs::MaxNodeSize;
+}
+
+FSlateRect SGraphNodeResizable::GetHitTestingBorder(float InverseZoomFactor) const
+{
+	return FSlateRect(	GraphNodeResizableDefs::HitResultBorderSize.Left * InverseZoomFactor,
+						GraphNodeResizableDefs::HitResultBorderSize.Top * InverseZoomFactor,
+						GraphNodeResizableDefs::HitResultBorderSize.Right * InverseZoomFactor,
+						GraphNodeResizableDefs::HitResultBorderSize.Bottom * InverseZoomFactor );
+}
\ No newline at end of file
diff --git a/Engine/Source/Editor/GraphEditor/Public/SGraphNodeComment.h b/Engine/Source/Editor/GraphEditor/Public/SGraphNodeComment.h
--- a/Engine/Source/Editor/GraphEditor/Public/SGraphNodeComment.h
+++ b/Engine/Source/Editor/GraphEditor/Public/SGraphNodeComment.h
@@ -2,39 +2,14 @@
 
 #pragma once
 
-class SGraphNodeComment : public SGraphNode
+class SGraphNodeComment : public SGraphNodeResizable
 {
-	/**
-	 * The Comment window zone the user is interacting with
-	 */
-	enum ECommentWindowZone
-	{
-		CWZ_NotInWindow			= 0,
-		CWZ_InWindow			= 1,
-		CWZ_RightBorder			= 2,
-		CWZ_BottomBorder		= 3,
-		CWZ_BottomRightBorder	= 4,
-		CWZ_LeftBorder			= 5,
-		CWZ_TopBorder			= 6,
-		CWZ_TopLeftBorder		= 7,
-		CWZ_TopRightBorder		= 8,
-		CWZ_BottomLeftBorder	= 9,
-		CWZ_TitleBar			= 10,
-	};
-
 public:
 	SLATE_BEGIN_ARGS(SGraphNodeComment){}
 	SLATE_END_ARGS()
 
 	// Begin SWidget Interface
-	virtual FReply OnMouseMove( const FGeometry& MyGeometry, const FPointerEvent& MouseEvent ) override;
-	virtual void OnMouseEnter( const FGeometry& MyGeometry, const FPointerEvent& MouseEvent ) override;
-	virtual void OnMouseLeave( const FPointerEvent& MouseEvent ) override;
-	virtual FCursorReply OnCursorQuery( const FGeometry& MyGeometry, const FPointerEvent& CursorEvent ) const override;
 	virtual FReply OnMouseButtonDoubleClick( const FGeometry& InMyGeometry, const FPointerEvent& InMouseEvent ) override;
-	virtual FReply OnMouseButtonDown( const FGeometry& MyGeometry, const FPointerEvent& MouseEvent ) override;
-	virtual FReply OnMouseButtonUp( const FGeometry& MyGeometry, const FPointerEvent& MouseEvent ) override;
-	virtual bool OnHitTest( const FGeometry& MyGeometry, FVector2D InAbsoluteCursorPosition ) override;
 	virtual void Tick(const FGeometry& AllottedGeometry, const double InCurrentTime, const float InDeltaTime) override;
 	virtual FReply OnDrop( const FGeometry& MyGeometry, const FDragDropEvent& DragDropEvent ) override;
 	virtual void OnDragEnter( const FGeometry& MyGeometry, const FDragDropEvent& DragDropEvent ) override;
@@ -80,16 +55,13 @@ protected:
 	/** called when user is moving the comment node */
 	virtual void MoveTo(const FVector2D& NewPosition, FNodeSet& NodeFilter ) override;
 
-private:
+	// SGraphNodeResizable Interface
+	virtual float GetTitleBarHeight() const override;
+	virtual FSlateRect GetHitTestingBorder( float InverseZoomFactor ) const override;
+	virtual FVector2D GetNodeMaximumSize() const override;
+	// SGraphNodeResizable Interface
 
-	/** Find the current window zone the mouse is in for the comment */
-	SGraphNodeComment::ECommentWindowZone FindMouseZone(const FVector2D& LocalGeo) const;
-
-	/** @return true if the current window zone is considered a selection area */
-	bool InSelectionArea() const { return InSelectionArea(MouseZone); }
-
-	/** @return true if the passed zone is a selection area */
-	bool InSelectionArea(ECommentWindowZone InZone) const;
+private:
 
 	/** @return the color to tint the comment body */
 	FSlateColor GetCommentBodyColor() const;
@@ -97,30 +69,10 @@ private:
 	/** @return the color to tint the title bar */
 	FSlateColor GetCommentTitleBarColor() const;
 
-	/** Function to store anchor point before resizing the node. The node will be anchored to this point when resizing happens*/
-	void InitNodeAnchorPoint();
-
-	/** Function to fetch the corrected node position based on anchor point*/
-	FVector2D GetCorrectedNodePosition() const;
-
 	/** Returns the width to wrap the text of the comment at */
 	float GetWrapAt() const;
 
 private:
-	/** The desired size of the comment box, set during a drag */
-	FVector2D UserSize;
-
-	/** The original size of the comment box while resizing */
-	FVector2D StoredUserSize;
-
-	/** Anchor point used to correct node position on resizing the node*/
-	FVector2D NodeAnchorPoint;
-
-	/** The current window zone the mouse is in */
-	ECommentWindowZone MouseZone;
-
-	/** If true the user is actively dragging the comment */
-	bool bUserIsDragging;
 
 	/** The current selection state of the comment */
 	mutable bool bIsSelected;
diff --git a/Engine/Source/Editor/GraphEditor/Public/SGraphNodeDocumentation.h b/Engine/Source/Editor/GraphEditor/Public/SGraphNodeDocumentation.h
new file mode 100644
--- /dev/null
+++ b/Engine/Source/Editor/GraphEditor/Public/SGraphNodeDocumentation.h
@@ -0,0 +1,64 @@
+// Copyright 1998-2013 Epic Games, Inc. All Rights Reserved.
+
+#pragma once
+
+class IDocumentationPage;
+
+class SGraphNodeDocumentation : public SGraphNodeResizable
+{
+public:
+	SLATE_BEGIN_ARGS(SGraphNodeDocumentation){}
+	SLATE_END_ARGS()
+
+	// Begin SPanel Interface
+	virtual FVector2D ComputeDesiredSize() const override;
+	// End SPanel interface
+
+	// Begin SWidget Interface
+	void Construct( const FArguments& InArgs, UEdGraphNode* InNode );
+	virtual FReply OnMouseWheel( const FGeometry& MyGeometry, const FPointerEvent& MouseEvent );
+	virtual bool OnHitTest( const FGeometry& MyGeometry, FVector2D InAbsoluteCursorPosition ) override;
+	virtual void Tick( const FGeometry& AllottedGeometry, const double InCurrentTime, const float InDeltaTime ) override;
+	// End SWidget Interface
+
+	// Begin SGraphNodeResizable Interface
+	virtual FVector2D GetNodeMinimumSize() const override;
+	virtual FVector2D GetNodeMaximumSize() const override;
+	virtual float GetTitleBarHeight() const override;
+	virtual FSlateRect GetHitTestingBorder( float InverseZoomFactor ) const override;
+	// End SGraphNodeResizable Interface
+
+protected:
+
+	// Begin SGraphNode Interface
+	virtual void UpdateGraphNode() override;
+	virtual bool IsNameReadOnly () const override { return false; }
+	// End SGraphNode Interface
+
+	/** Create documentation page from link and excerpt */
+	TSharedPtr<SWidget> CreateDocumentationPage();
+	/** Returns the width the documentation content must adhere to, used as a delegate in child widgets */
+	FOptionalSize GetContentWidth() const;
+	/** Returns the height the documentation content must adhere to, used as a delegate in child widgets */
+	FOptionalSize GetContentHeight() const;
+	/** Returns the child widget text wrapat size, used as a delegate during creation of documentation page */
+	float GetDocumentationWrapWidth() const;
+	/** Returns the current child widgets visibility for hit testing */
+	EVisibility GetWidgetVisibility() const;
+
+private:
+
+	/** Node Title Bar */
+	TSharedPtr<SWidget> TitleBar;
+	/** Documentation excerpt maximum page size */
+	FVector2D DocumentationSize;
+	/** Content Widget with the desired size being managed */
+	TSharedPtr<SWidget> ContentWidget;
+	/** Tracks if child widgets are availble to hit test against */
+	EVisibility ChildWidgetVisibility;
+	/** Cached Documentation Link */
+	FString CachedDocumentationLink;
+	/** Cached Documentation Excerpt */
+	FString CachedDocumentationExcerpt;
+
+};
diff --git a/Engine/Source/Editor/GraphEditor/Public/SGraphNodeResizable.h b/Engine/Source/Editor/GraphEditor/Public/SGraphNodeResizable.h
new file mode 100644
--- /dev/null
+++ b/Engine/Source/Editor/GraphEditor/Public/SGraphNodeResizable.h
@@ -0,0 +1,86 @@
+// Copyright 1998-2014 Epic Games, Inc. All Rights Reserved.
+
+#pragma once
+
+class SGraphNodeResizable : public SGraphNode
+{
+public:
+
+	/**
+	 * The resizable window zone the user is interacting with
+	 */
+	enum EResizableWindowZone
+	{
+		CRWZ_NotInWindow		= 0,
+		CRWZ_InWindow			= 1,
+		CRWZ_RightBorder		= 2,
+		CRWZ_BottomBorder		= 3,
+		CRWZ_BottomRightBorder	= 4,
+		CRWZ_LeftBorder			= 5,
+		CRWZ_TopBorder			= 6,
+		CRWZ_TopLeftBorder		= 7,
+		CRWZ_TopRightBorder		= 8,
+		CRWZ_BottomLeftBorder	= 9,
+		CRWZ_TitleBar			= 10,
+	};
+
+	// Begin SWidget Interface
+	virtual FReply OnMouseMove( const FGeometry& MyGeometry, const FPointerEvent& MouseEvent ) override;
+	virtual void OnMouseEnter( const FGeometry& MyGeometry, const FPointerEvent& MouseEvent ) override;
+	virtual void OnMouseLeave( const FPointerEvent& MouseEvent ) override;
+	virtual FReply OnMouseButtonDown( const FGeometry& MyGeometry, const FPointerEvent& MouseEvent ) override;
+	virtual FReply OnMouseButtonUp( const FGeometry& MyGeometry, const FPointerEvent& MouseEvent ) override;
+	virtual FCursorReply OnCursorQuery( const FGeometry& MyGeometry, const FPointerEvent& CursorEvent ) const override;
+	virtual bool OnHitTest( const FGeometry& MyGeometry, FVector2D InAbsoluteCursorPosition ) override;
+	// End SWidget Interface
+	
+protected:
+
+	/** Find the current window zone the mouse is in */
+	virtual SGraphNodeResizable::EResizableWindowZone FindMouseZone(const FVector2D& LocalMouseCoordinates) const;
+
+	/** @return true if the current window zone is considered a selection area */
+	bool InSelectionArea() const { return InSelectionArea(MouseZone); }
+
+	/** @return true if the passed zone is a selection area */
+	bool InSelectionArea(EResizableWindowZone InZone) const;
+
+	/** Function to store anchor point before resizing the node. The node will be anchored to this point when resizing happens*/
+	void InitNodeAnchorPoint();
+
+	/** Function to fetch the corrected node position based on anchor point*/
+	FVector2D GetCorrectedNodePosition() const;
+
+	/** Get the current titlebar size */
+	virtual float GetTitleBarHeight() const;
+
+	/** Return smallest desired node size */
+	virtual FVector2D GetNodeMinimumSize() const;
+
+	/** Return largest desired node size */
+	virtual FVector2D GetNodeMaximumSize() const;
+
+	//** Return slate rect border for hit testing */
+	virtual FSlateRect GetHitTestingBorder( float InverseZoomFactor ) const;
+
+protected:
+
+	/** The non snapped size of the node for fluid resizing */
+	FVector2D DragSize;
+
+	/** The desired size of the node set during a drag */
+	FVector2D UserSize;
+
+	/** The original size of the node while resizing */
+	FVector2D StoredUserSize;
+
+	/** Anchor point used to correct node position on resizing the node*/
+	FVector2D NodeAnchorPoint;
+
+	/** The current window zone the mouse is in */
+	EResizableWindowZone MouseZone;
+
+	/** If true the user is actively dragging the node */
+	bool bUserIsDragging;
+
+};
diff --git a/Engine/Source/Editor/Kismet/Private/BlueprintDetailsCustomization.cpp b/Engine/Source/Editor/Kismet/Private/BlueprintDetailsCustomization.cpp
--- a/Engine/Source/Editor/Kismet/Private/BlueprintDetailsCustomization.cpp
+++ b/Engine/Source/Editor/Kismet/Private/BlueprintDetailsCustomization.cpp
@@ -4873,4 +4873,181 @@ void FChildActorComponentDetails::CustomizeDetails(IDetailLayoutBuilder& DetailB
 	}
 }
 
+namespace BlueprintDocumentationDetailDefs
+{
+	/** Minimum size of the details title panel */
+	static const float DetailsTitleMinWidth = 125.f;
+	/** Maximum size of the details title panel */
+	static const float DetailsTitleMaxWidth = 300.f;
+};
+
+void FBlueprintDocumentationDetails::CustomizeDetails(IDetailLayoutBuilder& DetailLayout)
+{
+	check( BlueprintEditorPtr.IsValid() );
+	// find currently selected edgraph documentation node
+	DocumentationNodePtr = EdGraphSelectionAsDocumentNode();
+
+	if( DocumentationNodePtr.IsValid() )
+	{
+		IDetailCategoryBuilder& DocumentationCategory = DetailLayout.EditCategory("Documentation", LOCTEXT("DocumentationDetailsCategory", "Documentation").ToString(), ECategoryPriority::Default);
+
+		DocumentationCategory.AddCustomRow( TEXT( "Documentation Link" ))
+		.NameContent()
+		.HAlign( HAlign_Fill )
+		[
+			SNew( STextBlock )
+			.Text( LOCTEXT( "FBlueprintDocumentationDetails_Link", "Link" ).ToString() )
+			.ToolTipText( LOCTEXT( "FBlueprintDocumentationDetails_LinkPathTooltip", "The documentation content path" ))
+			.Font( IDetailLayoutBuilder::GetDetailFont() )
+		]
+		.ValueContent()
+		.HAlign( HAlign_Left )
+		.MinDesiredWidth( BlueprintDocumentationDetailDefs::DetailsTitleMinWidth )
+		.MaxDesiredWidth( BlueprintDocumentationDetailDefs::DetailsTitleMaxWidth )
+		[
+			SNew( SEditableTextBox )
+			.Padding( FMargin( 4.f, 2.f ))
+			.Text( this, &FBlueprintDocumentationDetails::OnGetDocumentationLink )
+			.ToolTipText( LOCTEXT( "FBlueprintDocumentationDetails_LinkTooltip", "The path of the documentation content relative to /Engine/Documentation/Source" ))
+			.OnTextCommitted( this, &FBlueprintDocumentationDetails::OnDocumentationLinkCommitted )
+			.Font( IDetailLayoutBuilder::GetDetailFont() )
+		];
+
+		DocumentationCategory.AddCustomRow( TEXT( "Documentation Excerpts" ))
+		.NameContent()
+		.HAlign( HAlign_Left )
+		[
+			SNew( STextBlock )
+			.Text( LOCTEXT( "FBlueprintDocumentationDetails_Excerpt", "Excerpt" ).ToString() )
+			.ToolTipText( LOCTEXT( "FBlueprintDocumentationDetails_ExcerptTooltip", "The current documentation excerpt" ))
+			.Font( IDetailLayoutBuilder::GetDetailFont() )
+		]
+		.ValueContent()
+		.HAlign( HAlign_Left )
+		.MinDesiredWidth( BlueprintDocumentationDetailDefs::DetailsTitleMinWidth )
+		.MaxDesiredWidth( BlueprintDocumentationDetailDefs::DetailsTitleMaxWidth )
+		[
+			SAssignNew( ExcerptComboButton, SComboButton )
+			.ContentPadding( 2.f )
+			.IsEnabled( this, &FBlueprintDocumentationDetails::OnExcerptChangeEnabled )
+			.ButtonContent()
+			[
+				SNew(SBorder)
+				.BorderImage( FEditorStyle::GetBrush( "NoBorder" ))
+				.Padding( FMargin( 0, 0, 5, 0 ))
+				[
+					SNew( STextBlock )
+					.Text( this, &FBlueprintDocumentationDetails::OnGetDocumentationExcerpt )
+					.ToolTipText( LOCTEXT( "FBlueprintDocumentationDetails_ExcerptComboTooltip", "Select Excerpt" ))
+					.Font( IDetailLayoutBuilder::GetDetailFont() )
+				]
+			]
+			.OnGetMenuContent( this, &FBlueprintDocumentationDetails::GenerateExcerptList )
+		];
+	}
+}
+
+TWeakObjectPtr<UEdGraphNode_Documentation> FBlueprintDocumentationDetails::EdGraphSelectionAsDocumentNode()
+{
+	DocumentationNodePtr.Reset();
+
+	if( BlueprintEditorPtr.IsValid() )
+	{
+		/** Get the currently selected set of nodes */
+		if( BlueprintEditorPtr.Pin()->GetNumberOfSelectedNodes() == 1 )
+		{
+			TSet<UObject*> Objects = BlueprintEditorPtr.Pin()->GetSelectedNodes();
+			TSet<UObject*>::TIterator Iter( Objects );
+			UObject* Object = *Iter;
+
+			if( Object && Object->IsA<UEdGraphNode_Documentation>() )
+			{
+				DocumentationNodePtr = Cast<UEdGraphNode_Documentation>( Object );
+			}
+		}
+	}
+	return DocumentationNodePtr;
+}
+
+FText FBlueprintDocumentationDetails::OnGetDocumentationLink() const
+{
+	FText Link = FText::GetEmpty();
+	if( DocumentationNodePtr.IsValid() )
+	{
+		Link = FText::FromString( DocumentationNodePtr->GetDocumentationLink() );
+	}
+	return Link;
+}
+
+FText FBlueprintDocumentationDetails::OnGetDocumentationExcerpt() const
+{
+	FText Excerpt = NSLOCTEXT( "FBlueprintDocumentationDetails", "ExcerptCombo_DefaultText", "Select Excerpt" );
+	if( DocumentationNodePtr.IsValid() )
+	{
+		Excerpt = FText::FromString( DocumentationNodePtr->GetDocumentationExcerptName() );
+	}
+	return Excerpt;
+}
+
+bool FBlueprintDocumentationDetails::OnExcerptChangeEnabled() const
+{
+	return DocumentationNodePtr.IsValid() && IDocumentation::Get()->PageExists( DocumentationNodePtr->Link );
+}
+
+void FBlueprintDocumentationDetails::OnDocumentationLinkCommitted( const FText& InNewName, ETextCommit::Type InTextCommit ) const
+{
+	if( DocumentationNodePtr.IsValid() )
+	{
+		DocumentationNodePtr->Link = InNewName.ToString();
+		DocumentationNodePtr->Excerpt.Empty();
+	}
+}
+
+TSharedRef< ITableRow > FBlueprintDocumentationDetails::MakeExcerptViewWidget( TSharedPtr<FString> Item, const TSharedRef< STableViewBase >& OwnerTable )
+{
+	return 
+		SNew( STableRow<TSharedPtr<FString>>, OwnerTable )
+		[
+			SNew( STextBlock )
+			.Text( *Item.Get() )
+		];
+}
+
+void FBlueprintDocumentationDetails::OnExcerptSelectionChanged( TSharedPtr<FString> ProposedSelection, ESelectInfo::Type /*SelectInfo*/ )
+{
+	if( ProposedSelection.IsValid() && DocumentationNodePtr.IsValid() )
+	{
+		DocumentationNodePtr->Excerpt = *ProposedSelection.Get();
+		ExcerptComboButton->SetIsOpen( false );
+	}
+}
+
+TSharedRef<SWidget> FBlueprintDocumentationDetails::GenerateExcerptList()
+{
+	ExcerptList.Empty();
+
+	if( DocumentationNodePtr.IsValid() && IDocumentation::Get()->PageExists( DocumentationNodePtr->Link ))
+	{
+		TSharedPtr<IDocumentationPage> DocumentationPage = IDocumentation::Get()->GetPage( DocumentationNodePtr->Link, NULL );
+		TArray<FExcerpt> Excerpts;
+		DocumentationPage->GetExcerpts( Excerpts );
+
+		for( auto ExcerptIter = Excerpts.CreateConstIterator(); ExcerptIter; ++ExcerptIter )
+		{
+			ExcerptList.Add( MakeShareable( new FString( ExcerptIter->Name )));
+		}
+	}
+
+	return
+		SNew( SHorizontalBox )
+		+SHorizontalBox::Slot()
+		.Padding( 2.f )
+		[
+			SNew( SListView< TSharedPtr<FString>> )
+			.ListItemsSource( &ExcerptList )
+			.OnGenerateRow( this, &FBlueprintDocumentationDetails::MakeExcerptViewWidget )
+			.OnSelectionChanged( this, &FBlueprintDocumentationDetails::OnExcerptSelectionChanged )
+		];
+}
+
 #undef LOCTEXT_NAMESPACE
diff --git a/Engine/Source/Editor/Kismet/Private/BlueprintDetailsCustomization.h b/Engine/Source/Editor/Kismet/Private/BlueprintDetailsCustomization.h
--- a/Engine/Source/Editor/Kismet/Private/BlueprintDetailsCustomization.h
+++ b/Engine/Source/Editor/Kismet/Private/BlueprintDetailsCustomization.h
@@ -2,6 +2,8 @@
 
 #pragma once
 
+#include "EdGraph/EdGraphNode_Documentation.h"
+
 /** Variable replication states */
 namespace EVariableReplication
 {
@@ -724,3 +726,65 @@ private:
 	/** Weak reference to the Blueprint editor */
 	TWeakPtr<FBlueprintEditor> BlueprintEditorPtr;
 };
+
+/** Details customization for Blueprint Documentation */
+class FBlueprintDocumentationDetails : public IDetailCustomization
+{
+public:
+	/** Constructor */
+	FBlueprintDocumentationDetails(TWeakPtr<FBlueprintEditor> InBlueprintEditorPtr)
+		: BlueprintEditorPtr(InBlueprintEditorPtr)
+	{
+	}
+
+	/** Makes a new instance of this detail layout class for a specific detail view requesting it */
+	static TSharedRef<class IDetailCustomization> MakeInstance(TWeakPtr<FBlueprintEditor> InBlueprintEditorPtr)
+	{
+		return MakeShareable(new FBlueprintDocumentationDetails(InBlueprintEditorPtr));
+	}
+
+	/** IDetailCustomization interface */
+	virtual void CustomizeDetails( IDetailLayoutBuilder& DetailLayout ) override;
+
+protected:
+
+	/** Accessors passed to parent */
+	UBlueprint* GetBlueprintObj() const { return BlueprintEditorPtr.Pin()->GetBlueprintObj(); }
+
+	/** Get the currently selected node from the edgraph */
+	TWeakObjectPtr<UEdGraphNode_Documentation> EdGraphSelectionAsDocumentNode();
+
+	/** Accessor for the current nodes documentation link */
+	FText OnGetDocumentationLink() const;
+
+	/** Accessor for the nodes current documentation excerpt  */
+	FText OnGetDocumentationExcerpt() const;
+	
+	/** Accessor to evaluate if the current excerpt can be modified */
+	bool OnExcerptChangeEnabled() const;
+
+	/** Handler for the documentation link being committed */
+	void OnDocumentationLinkCommitted( const FText& InNewName, ETextCommit::Type InTextCommit ) const;
+	
+	/** Generate table row for excerpt combo */
+	TSharedRef< ITableRow > MakeExcerptViewWidget( TSharedPtr<FString> Item, const TSharedRef< STableViewBase >& OwnerTable );
+	
+	/** Apply selection changes from the excerpt combo */
+	void OnExcerptSelectionChanged( TSharedPtr<FString> ProposedSelection, ESelectInfo::Type SelectInfo );
+
+	/** Generate excerpt list widget from documentation page */
+	TSharedRef<SWidget> GenerateExcerptList();
+
+
+private:
+
+	/** Weak reference to the Blueprint editor */
+	TWeakPtr<FBlueprintEditor> BlueprintEditorPtr;
+	/** The editor node we're editing */
+	TWeakObjectPtr<UEdGraphNode_Documentation> DocumentationNodePtr;
+	/** Excerpt combo widget */
+	TSharedPtr<SComboButton> ExcerptComboButton;
+	/** Excerpt List */
+	TArray<TSharedPtr<FString>> ExcerptList;
+
+};
\ No newline at end of file
diff --git a/Engine/Source/Editor/Kismet/Private/SBlueprintPalette.cpp b/Engine/Source/Editor/Kismet/Private/SBlueprintPalette.cpp
--- a/Engine/Source/Editor/Kismet/Private/SBlueprintPalette.cpp
+++ b/Engine/Source/Editor/Kismet/Private/SBlueprintPalette.cpp
@@ -361,6 +361,11 @@ static void GetPaletteItemIcon(TSharedPtr<FEdGraphSchemaAction> ActionIn, UBluep
 		FBlueprintActionMenuItem* NodeSpawnerAction = (FBlueprintActionMenuItem*)ActionIn.Get();
 		BrushOut = NodeSpawnerAction->GetMenuIcon(ColorOut);
 	}
+	else if (ActionIn->GetTypeId() == FEdGraphSchemaAction_K2AddDocumentation::StaticGetTypeId())
+	{
+		ToolTipOut = LOCTEXT("Documentation_Tooltip", "Create a documentation node.").ToString();
+		BrushOut = FEditorStyle::GetBrush(TEXT("GraphEditor.Documentation_16x"));
+	}
 }
 
 /**
diff --git a/Engine/Source/Editor/Kismet/Private/SKismetInspector.cpp b/Engine/Source/Editor/Kismet/Private/SKismetInspector.cpp
--- a/Engine/Source/Editor/Kismet/Private/SKismetInspector.cpp
+++ b/Engine/Source/Editor/Kismet/Private/SKismetInspector.cpp
@@ -230,6 +230,9 @@ void SKismetInspector::Construct(const FArguments& InArgs)
 		FOnGetDetailCustomizationInstance LayoutFormatTextDetails = FOnGetDetailCustomizationInstance::CreateStatic(&FFormatTextDetails::MakeInstance);
 		PropertyView->RegisterInstancedCustomPropertyLayout(UK2Node_FormatText::StaticClass(), LayoutFormatTextDetails);
 
+		FOnGetDetailCustomizationInstance LayoutDocumentationDetails = FOnGetDetailCustomizationInstance::CreateStatic(&FBlueprintDocumentationDetails::MakeInstance, Kismet2Ptr);
+		PropertyView->RegisterInstancedCustomPropertyLayout(UEdGraphNode_Documentation::StaticClass(), LayoutDocumentationDetails);
+
 		FOnGetDetailCustomizationInstance GraphNodeDetails = FOnGetDetailCustomizationInstance::CreateStatic(&FBlueprintGraphNodeDetails::MakeInstance, Kismet2Ptr);
 		PropertyView->RegisterInstancedCustomPropertyLayout(UEdGraphNode::StaticClass(), GraphNodeDetails);
 
diff --git a/Engine/Source/Editor/KismetCompiler/Private/KismetCompiler.cpp b/Engine/Source/Editor/KismetCompiler/Private/KismetCompiler.cpp
--- a/Engine/Source/Editor/KismetCompiler/Private/KismetCompiler.cpp
+++ b/Engine/Source/Editor/KismetCompiler/Private/KismetCompiler.cpp
@@ -18,6 +18,7 @@
 #include "UserDefinedStructureCompilerUtils.h"
 #include "EditorCategoryUtils.h"
 #include "K2Node_EnumLiteral.h"
+#include "EdGraph/EdGraphNode_Documentation.h"
 
 static bool bDebugPropertyPropagation = false;
 
@@ -384,8 +385,8 @@ bool FKismetCompilerContext::IsNodePure(const UEdGraphNode* Node) const
 	{
 		return K2Node->IsNodePure();
 	}
-	// Only non K2Nodes are comments, which are pure
-	ensure(Node->IsA(UEdGraphNode_Comment::StaticClass()));
+	// Only non K2Nodes are comments and documentation nodes, which are pure
+	ensure(Node->IsA(UEdGraphNode_Comment::StaticClass())||Node->IsA(UEdGraphNode_Documentation::StaticClass()));
 	return true;
 }
 
diff --git a/Engine/Source/Runtime/Engine/Classes/EdGraph/EdGraphNode_Documentation.h b/Engine/Source/Runtime/Engine/Classes/EdGraph/EdGraphNode_Documentation.h
new file mode 100644
--- /dev/null
+++ b/Engine/Source/Runtime/Engine/Classes/EdGraph/EdGraphNode_Documentation.h
@@ -0,0 +1,42 @@
+ // Copyright 1998-2013 Epic Games, Inc. All Rights Reserved.
+
+#pragma once
+#include "EdGraphNode_Documentation.generated.h"
+
+UCLASS(MinimalAPI)
+class UEdGraphNode_Documentation : public UEdGraphNode
+{
+	GENERATED_UCLASS_BODY()
+
+	/** Documentation Link */
+	UPROPERTY(meta=(FriendlyName="Documentation Link"))
+	FString Link;
+
+	/** Documentation Excerpt */
+	UPROPERTY(meta=(FriendlyName="Documentation Excerpt"))
+	FString Excerpt;
+
+public:
+
+#if WITH_EDITOR
+	// Begin UEdGraphNode interface
+	virtual void AllocateDefaultPins() override {}
+	virtual bool ShouldOverridePinNames() const override { return true; }
+	virtual FLinearColor GetNodeTitleColor() const override { return FLinearColor( 0.2f, 1.f, 0.2f ); }
+	ENGINE_API virtual FString GetTooltip() const override;
+	ENGINE_API virtual FText GetNodeTitle(ENodeTitleType::Type TitleType) const override;
+	ENGINE_API virtual FString GetPinNameOverride(const UEdGraphPin& Pin) const override;
+	ENGINE_API virtual void ResizeNode(const FVector2D& NewSize) override;
+	ENGINE_API virtual void PostPlacedNewNode() override;
+	ENGINE_API virtual void OnRenameNode(const FString& NewName) override;
+	ENGINE_API virtual TSharedPtr<class INameValidatorInterface> MakeNameValidator() const override;
+	ENGINE_API virtual FString GetDocumentationLink() const override { return Link; }
+	ENGINE_API virtual FString GetDocumentationExcerptName() const override { return Excerpt; }
+	// End UEdGraphNode interface
+
+	/** Set the Bounds for the comment node */
+	ENGINE_API void SetBounds(const class FSlateRect& Rect);
+
+#endif
+};
+
diff --git a/Engine/Source/Runtime/Engine/Classes/EdGraph/EdGraphSchema.h b/Engine/Source/Runtime/Engine/Classes/EdGraph/EdGraphSchema.h
--- a/Engine/Source/Runtime/Engine/Classes/EdGraph/EdGraphSchema.h
+++ b/Engine/Source/Runtime/Engine/Classes/EdGraph/EdGraphSchema.h
@@ -668,6 +668,9 @@ class ENGINE_API UEdGraphSchema : public UObject
 	/** Returns schema action to create comment from implemention */
 	virtual TSharedPtr<FEdGraphSchemaAction> GetCreateCommentAction() const { return NULL; }
 
+	/** Returns schema action to create documention node from implemention */
+	virtual TSharedPtr<FEdGraphSchemaAction> GetCreateDocumentNodeAction() const { return NULL; }
+
 	/**
 	 * Handle a graph being removed by the user (potentially removing associated bound nodes, etc...)
 	 */
diff --git a/Engine/Source/Runtime/Engine/Private/EdGraph/EdGraphNode_Documentation.cpp b/Engine/Source/Runtime/Engine/Private/EdGraph/EdGraphNode_Documentation.cpp
new file mode 100644
--- /dev/null
+++ b/Engine/Source/Runtime/Engine/Private/EdGraph/EdGraphNode_Documentation.cpp
@@ -0,0 +1,77 @@
+// Copyright 1998-2014 Epic Games, Inc. All Rights Reserved.
+
+#include "EnginePrivate.h"
+#include "EdGraph/EdGraphNode_Documentation.h"
+
+#define LOCTEXT_NAMESPACE "EdGraph"
+
+/////////////////////////////////////////////////////
+// UEdGraphNode_Documentation
+
+UEdGraphNode_Documentation::UEdGraphNode_Documentation( const class FPostConstructInitializeProperties& PCIP )
+	: Super( PCIP )
+{
+#if WITH_EDITORONLY_DATA
+	bCanResizeNode = true;
+	bCanRenameNode = false;
+#endif // WITH_EDITORONLY_DATA
+	Link = TEXT( "Shared/GraphNodes/Blueprint" );
+	Excerpt = TEXT( "UEdGraphNode_Documentation" );
+}
+
+#if WITH_EDITOR
+void UEdGraphNode_Documentation::PostPlacedNewNode()
+{
+	NodeComment = TEXT( "" );
+}
+
+FString UEdGraphNode_Documentation::GetTooltip() const
+{
+	return FString::Printf( *NSLOCTEXT( "K2Node", "DocumentationBlock_Tooltip", "Documentation:\n%s" ).ToString(), *NodeComment );
+}
+
+FText UEdGraphNode_Documentation::GetNodeTitle( ENodeTitleType::Type TitleType ) const
+{
+	return NSLOCTEXT( "K2Node", "DocumentationBlock_Title", "UDN Documentation Excerpt" );
+}
+
+FString UEdGraphNode_Documentation::GetPinNameOverride(const UEdGraphPin& Pin) const
+{
+	return GetNodeTitle( ENodeTitleType::ListView ).ToString();
+}
+
+void UEdGraphNode_Documentation::ResizeNode(const FVector2D& NewSize)
+{
+	if( bCanResizeNode ) 
+	{
+		NodeHeight = NewSize.Y;
+		NodeWidth = NewSize.X;
+	}
+}
+
+void UEdGraphNode_Documentation::SetBounds( const class FSlateRect& Rect )
+{
+	NodePosX = Rect.Left;
+	NodePosY = Rect.Top;
+
+	FVector2D Size = Rect.GetSize();
+	NodeWidth = Size.X;
+	NodeHeight = Size.Y;
+}
+
+void UEdGraphNode_Documentation::OnRenameNode( const FString& NewName )
+{
+	NodeComment = NewName;
+}
+
+TSharedPtr<class INameValidatorInterface> UEdGraphNode_Documentation::MakeNameValidator() const
+{
+	// Documentation can be duplicated, etc...
+	return MakeShareable( new FDummyNameValidator( EValidatorResult::Ok ));
+}
+
+#endif // WITH_EDITOR
+
+/////////////////////////////////////////////////////
+
+#undef LOCTEXT_NAMESPACE
