Fix for my changes that broke unity builds with cl 2247425

diff --git a/Engine/Source/Editor/GraphEditor/Private/SGraphNodeDocumentation.cpp b/Engine/Source/Editor/GraphEditor/Private/SGraphNodeDocumentation.cpp
--- a/Engine/Source/Editor/GraphEditor/Private/SGraphNodeDocumentation.cpp
+++ b/Engine/Source/Editor/GraphEditor/Private/SGraphNodeDocumentation.cpp
@@ -1,8 +1,10 @@
-// Copyright 1998-2013 Epic Games, Inc. All Rights Reserved.
+// Copyright 1998-2014 Epic Games, Inc. All Rights Reserved.
 
 
 #include "GraphEditorCommon.h"
 #include "SGraphNodeDocumentation.h"
+#include "IDocumentation.h"
+#include "SLevelOfDetailBranchNode.h"
 
 #define LOCTEXT_NAMESPACE "SGraphNodeDocumentation"
 
diff --git a/Engine/Source/Runtime/Engine/Private/EdGraph/EdGraphNode_Documentation.cpp b/Engine/Source/Runtime/Engine/Private/EdGraph/EdGraphNode_Documentation.cpp
--- a/Engine/Source/Runtime/Engine/Private/EdGraph/EdGraphNode_Documentation.cpp
+++ b/Engine/Source/Runtime/Engine/Private/EdGraph/EdGraphNode_Documentation.cpp
@@ -3,6 +3,11 @@
 #include "EnginePrivate.h"
 #include "EdGraph/EdGraphNode_Documentation.h"
 
+#if WITH_EDITOR
+#include "Slate.h"
+#include "Editor/UnrealEd/Public/Kismet2/Kismet2NameValidators.h"
+#endif
+
 #define LOCTEXT_NAMESPACE "EdGraph"
 
 /////////////////////////////////////////////////////
