Fix Mac build.

diff --git a/Engine/Source/Editor/GraphEditor/Private/SGraphNodeDocumentation.cpp b/Engine/Source/Editor/GraphEditor/Private/SGraphNodeDocumentation.cpp
--- a/Engine/Source/Editor/GraphEditor/Private/SGraphNodeDocumentation.cpp
+++ b/Engine/Source/Editor/GraphEditor/Private/SGraphNodeDocumentation.cpp
@@ -164,7 +164,7 @@ TSharedPtr<SWidget> SGraphNodeDocumentation::CreateDocumentationPage()
 	if( IDocumentation::Get()->PageExists( GraphNode->GetDocumentationLink() ))
 	{
 		TSharedRef< IDocumentationPage > DocumentationPage = IDocumentation::Get()->GetPage( GraphNode->GetDocumentationLink(), NULL );
-		const TMap< FString, FString > InVariables;
+		TMap< FString, FString > InVariables;
 		FExcerpt DesiredExcerpt( GraphNode->GetDocumentationExcerptName(), SNullWidget::NullWidget, InVariables, 0 );
 
 		// Set attributes to control documentation WrapAt and optional width
@@ -332,4 +332,8 @@ void SGraphNodeDocumentation::Tick( const FGeometry& AllottedGeometry, const dou
 	}
 
 	SGraphNodeResizable::Tick( AllottedGeometry, InCurrentTime, InDeltaTime );
-}
\ No newline at end of file
+}
+
+/////////////////////////////////////////////////////
+
+#undef LOCTEXT_NAMESPACE
\ No newline at end of file
