Holding Alt+RMB in a graph panel will zoom on the initial mouse position instead of the corner of the panel.

#ttp 343146 - BP: NUX:  Alt+RMB zooming is wonky

diff --git a/Engine/Source/Editor/GraphEditor/Private/SNodePanel.cpp b/Engine/Source/Editor/GraphEditor/Private/SNodePanel.cpp
--- a/Engine/Source/Editor/GraphEditor/Private/SNodePanel.cpp
+++ b/Engine/Source/Editor/GraphEditor/Private/SNodePanel.cpp
@@ -551,6 +551,10 @@ FReply SNodePanel::OnMouseButtonDown( const FGeometry& MyGeometry, const FPointe
 	}
 	else if ( bIsRightMouseButtonEffecting )
 	{
+		// Cache current cursor position as zoom origin and software cursor position
+		ZoomStartOffset = MyGeometry.AbsoluteToLocal( MouseEvent.GetLastScreenSpacePosition() );
+		SoftwareCursorPosition = PanelCoordToGraphCoord( ZoomStartOffset );
+
 		FReply ReplyState = FReply::Handled();
 		ReplyState.CaptureMouse( SharedThis(this) );
 		ReplyState.UseHighPrecisionMouseMovement( SharedThis(this) );
