Tweaked tags

diff --git a/Engine/Source/Editor/GraphEditor/Private/KismetNodes/SGraphNodeK2Base.cpp b/Engine/Source/Editor/GraphEditor/Private/KismetNodes/SGraphNodeK2Base.cpp
--- a/Engine/Source/Editor/GraphEditor/Private/KismetNodes/SGraphNodeK2Base.cpp
+++ b/Engine/Source/Editor/GraphEditor/Private/KismetNodes/SGraphNodeK2Base.cpp
@@ -49,7 +49,7 @@ void SGraphNodeK2Base::UpdateCompactNode()
 	if (GraphNode != nullptr)
 	{
 		UBlueprint* Blueprint = FBlueprintEditorUtils::FindBlueprintForGraphChecked(CastChecked<UEdGraph>(GraphNode->GetOuter()));
-		TagName = FString::Printf(TEXT("%s,%s"), *GraphNode->GetName(), *GraphNode->NodeGuid.ToString());		
+		TagName = FString::Printf(TEXT("GraphNode,%s,%s"), *GraphNode->GetName(), *GraphNode->NodeGuid.ToString());		
 	}
 	
 	//
diff --git a/Engine/Source/Editor/GraphEditor/Private/KismetNodes/SGraphNodeK2Var.cpp b/Engine/Source/Editor/GraphEditor/Private/KismetNodes/SGraphNodeK2Var.cpp
--- a/Engine/Source/Editor/GraphEditor/Private/KismetNodes/SGraphNodeK2Var.cpp
+++ b/Engine/Source/Editor/GraphEditor/Private/KismetNodes/SGraphNodeK2Var.cpp
@@ -162,23 +162,24 @@ void SGraphNodeK2Var::UpdateGraphNode()
 	FString TagName;
 	if (GraphNode != nullptr)
 	{
+		FString TagLabel;		
 		// Try to get a suitable name for the tag.
 		if (UK2Node_Variable const* K2Node = Cast<UK2Node_Variable>(GraphNode))
 		{
 			UBlueprint* Blueprint = FBlueprintEditorUtils::FindBlueprintForGraphChecked(CastChecked<UEdGraph>(GraphNode->GetOuter()));
-			TagName = FString::Printf(TEXT("%s,%s"), *K2Node->GetVarNameString(), *GraphNode->NodeGuid.ToString());
+			TagLabel = K2Node->GetVarNameString();
 		}
 		else if (UK2Node_Self const* K2Node = Cast<UK2Node_Self>(GraphNode))
 		{
 			UBlueprint* Blueprint = FBlueprintEditorUtils::FindBlueprintForGraphChecked(CastChecked<UEdGraph>(GraphNode->GetOuter()));
-			TagName = FString::Printf(TEXT("Self,%s"), *GraphNode->NodeGuid.ToString());
+			TagLabel = GraphNode->NodeGuid.ToString();
 		}
 		else
 		{
 			UBlueprint* Blueprint = FBlueprintEditorUtils::FindBlueprintForGraphChecked(CastChecked<UEdGraph>(GraphNode->GetOuter()));
-			TagName = FString::Printf(TEXT("%s,%s"), *GraphNode->GetName(), *GraphNode->NodeGuid.ToString());
+			TagLabel = GraphNode->GetName();
 		}
-
+		TagName = FString::Printf(TEXT("GraphNode,%s,%s"), *TagLabel, *GraphNode->NodeGuid.ToString());
 	}
 
 	//             ________________
diff --git a/Engine/Source/Editor/GraphEditor/Private/SGraphNode.cpp b/Engine/Source/Editor/GraphEditor/Private/SGraphNode.cpp
--- a/Engine/Source/Editor/GraphEditor/Private/SGraphNode.cpp
+++ b/Engine/Source/Editor/GraphEditor/Private/SGraphNode.cpp
@@ -712,7 +712,7 @@ void SGraphNode::UpdateGraphNode()
 			NodeName = NodeTitle->GetHeadTitle().ToString();
 		}
 
-		TagName = FString::Printf(TEXT("%s,%s"), *NodeName, *GraphNode->NodeGuid.ToString());
+		TagName = FString::Printf(TEXT("GraphNode,%s,%s"), *NodeName, *GraphNode->NodeGuid.ToString());
  	}
 	
 	TSharedPtr<SVerticalBox> InnerVerticalBox;
diff --git a/Engine/Source/Editor/GraphEditor/Private/SGraphNodeComment.cpp b/Engine/Source/Editor/GraphEditor/Private/SGraphNodeComment.cpp
--- a/Engine/Source/Editor/GraphEditor/Private/SGraphNodeComment.cpp
+++ b/Engine/Source/Editor/GraphEditor/Private/SGraphNodeComment.cpp
@@ -102,9 +102,7 @@ void SGraphNodeComment::UpdateGraphNode()
 	if (GraphNode != nullptr)
 	{
 		UBlueprint* Blueprint = FBlueprintEditorUtils::FindBlueprintForGraphChecked(CastChecked<UEdGraph>(GraphNode->GetOuter()));
-		FString NodeName;
-		NodeName = GetEditableNodeTitleAsText().ToString();
-		TagName = FString::Printf(TEXT("%s,%s"), *NodeName, *GraphNode->NodeGuid.ToString());
+		TagName = FString::Printf(TEXT("GraphNode,%s,%s"), *GetEditableNodeTitleAsText().ToString(), *GraphNode->NodeGuid.ToString());
 	}
 
 	TSharedPtr<SWidget> ErrorText = SetupErrorReporting();
diff --git a/Engine/Source/Editor/IntroTutorials/Private/TutorialStructCustomization.cpp b/Engine/Source/Editor/IntroTutorials/Private/TutorialStructCustomization.cpp
--- a/Engine/Source/Editor/IntroTutorials/Private/TutorialStructCustomization.cpp
+++ b/Engine/Source/Editor/IntroTutorials/Private/TutorialStructCustomization.cpp
@@ -310,13 +310,18 @@ private:
 
 	FText MakeFriendlyStringFromName(const FString& WidgetName) const
 	{
+		// We will likely have meta data for this eventually. For now just parse a comma delimited string which currently will either be a name or ident,name,UID
 		if (WidgetName == "None")
 		{
 			return FText::FromName(*WidgetName);
 		}
 		TArray<FString> Tokens;
 		WidgetName.ParseIntoArray(&Tokens, TEXT(","),true);
-		if (Tokens.Num() == 2)
+		if (Tokens.Num() == 3)
+		{
+			return FText::FromName(*Tokens[1]);
+		}
+		else if (Tokens.Num() == 3)
 		{
 			return FText::FromName(*Tokens[0]);
 		}
diff --git a/Engine/Source/Editor/Kismet/Private/SSCSEditor.cpp b/Engine/Source/Editor/Kismet/Private/SSCSEditor.cpp
--- a/Engine/Source/Editor/Kismet/Private/SSCSEditor.cpp
+++ b/Engine/Source/Editor/Kismet/Private/SSCSEditor.cpp
@@ -2044,7 +2044,22 @@ TSharedRef<ITableRow> SSCSEditor::MakeTableRowWidget( FSCSEditorTreeNodePtrType
 		OnRenameComponent(false);
 	}
 
-	return SNew(SSCS_RowWidget, SharedThis(this), InNodePtr, OwnerTable);
+	
+	// Setup a tag for this node
+	FString TagName;
+	if (InNodePtr.IsValid() && InNodePtr->GetComponentTemplate() != NULL )
+	{
+		if (InNodePtr->GetSCSNode())
+		{
+			TagName = FString::Printf(TEXT("TableRow,%s,%d"), *InNodePtr->GetComponentTemplate()->GetReadableName(), InNodePtr->GetSCSNode()->GetUniqueID());
+		}
+		else
+		{
+			TagName = FString::Printf(TEXT("TableRow,%s,0"), *InNodePtr->GetComponentTemplate()->GetReadableName());
+		}
+	}
+	return SNew(SSCS_RowWidget, SharedThis(this), InNodePtr, OwnerTable)
+		.Tag(*TagName);
 }
 
 void SSCSEditor::GetSelectedItemsForContextMenu(TArray<FComponentEventConstructionData>& OutSelectedItems) const
diff --git a/Engine/Source/Editor/PropertyEditor/Private/DetailItemNode.cpp b/Engine/Source/Editor/PropertyEditor/Private/DetailItemNode.cpp
--- a/Engine/Source/Editor/PropertyEditor/Private/DetailItemNode.cpp
+++ b/Engine/Source/Editor/PropertyEditor/Private/DetailItemNode.cpp
@@ -138,7 +138,7 @@ TSharedRef< ITableRow > FDetailItemNode::GenerateNodeWidget( const TSharedRef<ST
 	FString TagName;
 	if (Customization.IsValidCustomization() && Customization.GetPropertyNode().IsValid() )
 	{
-		TagName = Customization.GetPropertyNode()->GetDisplayName();
+		TagName = FString::Printf(TEXT("DetailItemNode,%s,%d"), *Customization.GetPropertyNode()->GetDisplayName(), Customization.GetPropertyNode()->GetArrayIndex());
 	}
 	if( Customization.HasPropertyNode() && Customization.GetPropertyNode()->AsCategoryNode() )
 	{
diff --git a/Engine/Source/Editor/PropertyEditor/Private/SDetailsView.cpp b/Engine/Source/Editor/PropertyEditor/Private/SDetailsView.cpp
--- a/Engine/Source/Editor/PropertyEditor/Private/SDetailsView.cpp
+++ b/Engine/Source/Editor/PropertyEditor/Private/SDetailsView.cpp
@@ -148,7 +148,7 @@ void SDetailsView::Construct(const FArguments& InArgs)
 				.ContentPadding(0)
 				.ForegroundColor( FSlateColor::UseForeground() )
 				.ButtonStyle( FEditorStyle::Get(), "ToggleButton" )
-				.Tag(TEXT("ViewOtions"))
+				.Tag(TEXT("ViewOptions"))
 				.MenuContent()
 				[
 					DetailViewOptions.MakeWidget()
diff --git a/Engine/Source/Runtime/Slate/Private/Framework/MultiBox/SToolBarComboButtonBlock.cpp b/Engine/Source/Runtime/Slate/Private/Framework/MultiBox/SToolBarComboButtonBlock.cpp
--- a/Engine/Source/Runtime/Slate/Private/Framework/MultiBox/SToolBarComboButtonBlock.cpp
+++ b/Engine/Source/Runtime/Slate/Private/Framework/MultiBox/SToolBarComboButtonBlock.cpp
@@ -104,10 +104,14 @@ void SToolBarComboButtonBlock::BuildMultiBlockWidget(const ISlateStyle* StyleSet
 		Label = ToolBarComboButtonBlock->Label;
 	}
 
+	// Create a tag
+	FString TagName;
+	TagName = FString::Printf(TEXT("ToolbarComboButton,%s,0"), *Label.Get().ToString());
+	
 	// Create the content for our button
 	TSharedRef< SWidget > ButtonContent =
 		SNew( SVerticalBox )
-		.Tag(*Label.Get().ToString())
+		.Tag(*TagName)
 		// Icon image
 		+ SVerticalBox::Slot()
 		.AutoHeight()
