CIS non-unity fixes

diff --git a/Engine/Source/Editor/GraphEditor/Private/KismetNodes/SGraphNodeK2Var.cpp b/Engine/Source/Editor/GraphEditor/Private/KismetNodes/SGraphNodeK2Var.cpp
--- a/Engine/Source/Editor/GraphEditor/Private/KismetNodes/SGraphNodeK2Var.cpp
+++ b/Engine/Source/Editor/GraphEditor/Private/KismetNodes/SGraphNodeK2Var.cpp
@@ -7,6 +7,7 @@
 #include "SGraphNodeK2Var.h"
 #include "ClassIconFinder.h"
 #include "IDocumentation.h"
+#include "BlueprintEditorUtils.h"
 
 void SGraphNodeK2Var::Construct( const FArguments& InArgs, UK2Node* InNode )
 {
diff --git a/Engine/Source/Editor/GraphEditor/Private/SGraphNodeComment.cpp b/Engine/Source/Editor/GraphEditor/Private/SGraphNodeComment.cpp
--- a/Engine/Source/Editor/GraphEditor/Private/SGraphNodeComment.cpp
+++ b/Engine/Source/Editor/GraphEditor/Private/SGraphNodeComment.cpp
@@ -4,6 +4,7 @@
 #include "GraphEditorCommon.h"
 #include "SGraphNodeComment.h"
 //#include "TextWrapperHelpers.h"
+#include "BlueprintEditorUtils.h"
 
 namespace SCommentNodeDefs
 {
diff --git a/Engine/Source/Editor/Persona/Private/Shared/AnimationRecorder.cpp b/Engine/Source/Editor/Persona/Private/Shared/AnimationRecorder.cpp
--- a/Engine/Source/Editor/Persona/Private/Shared/AnimationRecorder.cpp
+++ b/Engine/Source/Editor/Persona/Private/Shared/AnimationRecorder.cpp
@@ -2,6 +2,7 @@
 
 
 #include "PersonaPrivatePCH.h"
+#include "AnimationRecorder.h"
 
 #define LOCTEXT_NAMESPACE "FAnimationRecorder"
 
