Fixed Assert in SGraphNode

diff --git a/Engine/Source/Editor/GraphEditor/Private/SGraphNode.cpp b/Engine/Source/Editor/GraphEditor/Private/SGraphNode.cpp
--- a/Engine/Source/Editor/GraphEditor/Private/SGraphNode.cpp
+++ b/Engine/Source/Editor/GraphEditor/Private/SGraphNode.cpp
@@ -701,7 +701,7 @@ void SGraphNode::UpdateGraphNode()
  	FString TagName;
  	if (GraphNode != nullptr)
  	{
-		UBlueprint* Blueprint = FBlueprintEditorUtils::FindBlueprintForGraphChecked(CastChecked<UEdGraph>(GraphNode->GetOuter()));
+		UBlueprint* Blueprint = FBlueprintEditorUtils::FindBlueprintForGraph(Cast<UEdGraph>(GraphNode->GetOuter()));
 		FString NodeName;
 		if (Blueprint != nullptr)
 		{
