Changed FindBlueprintForGraphChecked to FindBlueprintForGraph to prevent asserts from firing

diff --git a/Engine/Source/Editor/GraphEditor/Private/SGraphNode.cpp b/Engine/Source/Editor/GraphEditor/Private/SGraphNode.cpp
--- a/Engine/Source/Editor/GraphEditor/Private/SGraphNode.cpp
+++ b/Engine/Source/Editor/GraphEditor/Private/SGraphNode.cpp
@@ -290,7 +290,7 @@ FReply SGraphNode::OnDrop( const FGeometry& MyGeometry, const FDragDropEvent& Dr
 				//Replace literal node's object reference
 				LiteralNode->SetObjectRef( DragConnectionOp->Actors[ 0 ].Get() );
 
-				UBlueprint* Blueprint = FBlueprintEditorUtils::FindBlueprintForGraphChecked(CastChecked<UEdGraph>(GraphNode->GetOuter()));
+				UBlueprint* Blueprint = FBlueprintEditorUtils::FindBlueprintForGraph(CastChecked<UEdGraph>(GraphNode->GetOuter()));
 				FBlueprintEditorUtils::MarkBlueprintAsModified(Blueprint);
 			}
 		}
diff --git a/Engine/Source/Editor/GraphEditor/Private/SGraphNodeComment.cpp b/Engine/Source/Editor/GraphEditor/Private/SGraphNodeComment.cpp
--- a/Engine/Source/Editor/GraphEditor/Private/SGraphNodeComment.cpp
+++ b/Engine/Source/Editor/GraphEditor/Private/SGraphNodeComment.cpp
@@ -102,7 +102,7 @@ void SGraphNodeComment::UpdateGraphNode()
 	FString TagName;
 	if (GraphNode != nullptr)
 	{
-		UBlueprint* Blueprint = FBlueprintEditorUtils::FindBlueprintForGraphChecked(CastChecked<UEdGraph>(GraphNode->GetOuter()));
+		UBlueprint* Blueprint = FBlueprintEditorUtils::FindBlueprintForGraph(CastChecked<UEdGraph>(GraphNode->GetOuter()));
 		TagName = FString::Printf(TEXT("GraphNode,%s,%s"), *GetEditableNodeTitleAsText().ToString(), *GraphNode->NodeGuid.ToString());
 	}
 
