adding nullptr check

diff --git a/Engine/Source/Editor/GraphEditor/Private/SGraphNode.cpp b/Engine/Source/Editor/GraphEditor/Private/SGraphNode.cpp
--- a/Engine/Source/Editor/GraphEditor/Private/SGraphNode.cpp
+++ b/Engine/Source/Editor/GraphEditor/Private/SGraphNode.cpp
@@ -291,7 +291,10 @@ FReply SGraphNode::OnDrop( const FGeometry& MyGeometry, const FDragDropEvent& Dr
 				LiteralNode->SetObjectRef( DragConnectionOp->Actors[ 0 ].Get() );
 
 				UBlueprint* Blueprint = FBlueprintEditorUtils::FindBlueprintForGraph(CastChecked<UEdGraph>(GraphNode->GetOuter()));
-				FBlueprintEditorUtils::MarkBlueprintAsModified(Blueprint);
+				if (Blueprint != nullptr)
+				{
+					FBlueprintEditorUtils::MarkBlueprintAsModified(Blueprint);
+				}
 			}
 		}
 		return FReply::Handled();
