UMG - Hiding the key pin in blueprints when told to.

diff --git a/Engine/Source/Editor/GraphEditor/Private/KismetPins/SGraphPinKey.cpp b/Engine/Source/Editor/GraphEditor/Private/KismetPins/SGraphPinKey.cpp
--- a/Engine/Source/Editor/GraphEditor/Private/KismetPins/SGraphPinKey.cpp
+++ b/Engine/Source/Editor/GraphEditor/Private/KismetPins/SGraphPinKey.cpp
@@ -27,8 +27,9 @@ void SGraphPinKey::Construct(const FArguments& InArgs, UEdGraphPin* InGraphPinOb
 TSharedRef<SWidget>	SGraphPinKey::GetDefaultValueWidget()
 {
 	return SNew(SKeySelector)
-			.CurrentKey(this, &SGraphPinKey::GetCurrentKey)
-			.OnKeyChanged(this, &SGraphPinKey::OnKeyChanged);
+		.Visibility(this, &SGraphPin::GetDefaultValueVisibility)
+		.CurrentKey(this, &SGraphPinKey::GetCurrentKey)
+		.OnKeyChanged(this, &SGraphPinKey::OnKeyChanged);
 }
 
 FKey SGraphPinKey::GetCurrentKey() const
