Fix for crash in FKnotNetCollector

diff --git a/Engine/Source/Editor/GraphEditor/Private/SGraphPin.cpp b/Engine/Source/Editor/GraphEditor/Private/SGraphPin.cpp
--- a/Engine/Source/Editor/GraphEditor/Private/SGraphPin.cpp
+++ b/Engine/Source/Editor/GraphEditor/Private/SGraphPin.cpp
@@ -24,7 +24,7 @@ struct FKnotNetCollector
 
 	void TraversePin(UEdGraphPin* Pin)
 	{
-		if (UK2Node_Knot* Knot = Cast<UK2Node_Knot>(Pin->GetOwningNode()))
+		if (UK2Node_Knot* Knot = Cast<UK2Node_Knot>(Pin->GetOwningNodeUnchecked()))
 		{
 			TraverseNodes(Knot);
 		}
@@ -34,8 +34,8 @@ struct FKnotNetCollector
 
 			for (UEdGraphPin* OtherPin : Pin->LinkedTo)
 			{
-				UEdGraphNode* OtherNode = OtherPin->GetOwningNode();
-				if (OtherPin->GetOwningNode()->IsA(UK2Node_Knot::StaticClass()))
+				UEdGraphNode* OtherNode = OtherPin->GetOwningNodeUnchecked();
+				if (OtherNode && OtherNode->IsA(UK2Node_Knot::StaticClass()))
 				{
 					TraverseNodes(OtherNode);
 				}
@@ -57,8 +57,8 @@ struct FKnotNetCollector
 
 			for (UEdGraphPin* OtherPin : MyPin->LinkedTo)
 			{
-				UEdGraphNode* OtherNode = OtherPin->GetOwningNode();
-				if (OtherNode->IsA(UK2Node_Knot::StaticClass()))
+				UEdGraphNode* OtherNode = OtherPin->GetOwningNodeUnchecked();
+				if (OtherNode && OtherNode->IsA(UK2Node_Knot::StaticClass()))
 				{
 					TraverseNodes(OtherNode);
 				}
