Display labels on component pins when splitting a pin on all compact nodes

diff --git a/Engine/Source/Editor/GraphEditor/Private/KismetNodes/SGraphNodeK2ArrayFunction.cpp b/Engine/Source/Editor/GraphEditor/Private/KismetNodes/SGraphNodeK2ArrayFunction.cpp
--- a/Engine/Source/Editor/GraphEditor/Private/KismetNodes/SGraphNodeK2ArrayFunction.cpp
+++ b/Engine/Source/Editor/GraphEditor/Private/KismetNodes/SGraphNodeK2ArrayFunction.cpp
@@ -169,7 +169,10 @@ void SGraphNodeK2ArrayFunction::UpdateGraphNode()
 		// Hide pin labels
 		for (auto InputPin : this->InputPins)
 		{
-			InputPin->SetShowLabel(false);
+			if (InputPin->GetPinObj()->ParentPin == nullptr)
+			{
+				InputPin->SetShowLabel(false);
+			}
 		}
 
 		for (auto OutputPin : this->OutputPins)
diff --git a/Engine/Source/Editor/GraphEditor/Private/KismetNodes/SGraphNodeK2Base.cpp b/Engine/Source/Editor/GraphEditor/Private/KismetNodes/SGraphNodeK2Base.cpp
--- a/Engine/Source/Editor/GraphEditor/Private/KismetNodes/SGraphNodeK2Base.cpp
+++ b/Engine/Source/Editor/GraphEditor/Private/KismetNodes/SGraphNodeK2Base.cpp
@@ -119,14 +119,20 @@ void SGraphNodeK2Base::UpdateCompactNode()
 	CreatePinWidgets();
 
 	// Hide pin labels
-	for (int32 PinIndex=0; PinIndex < this->InputPins.Num(); ++PinIndex)
+	for (auto InputPin: this->InputPins)
 	{
-		InputPins[PinIndex]->SetShowLabel(false);
+		if (InputPin->GetPinObj()->ParentPin == nullptr)
+		{
+			InputPin->SetShowLabel(false);
+		}
 	}
 
-	for (int32 PinIndex = 0; PinIndex < this->OutputPins.Num(); ++PinIndex)
+	for (auto OutputPin : this->OutputPins)
 	{
-		OutputPins[PinIndex]->SetShowLabel(false);
+		if (OutputPin->GetPinObj()->ParentPin == nullptr)
+		{
+			OutputPin->SetShowLabel(false);
+		}
 	}
 
 	CreateInputSideAddButton(LeftNodeBox);
