Optimizing how graph menus are sorted/built.

diff --git a/Engine/Source/Editor/GraphEditor/Private/GraphActionNode.cpp b/Engine/Source/Editor/GraphEditor/Private/GraphActionNode.cpp
new file mode 100644
--- /dev/null
+++ b/Engine/Source/Editor/GraphEditor/Private/GraphActionNode.cpp
@@ -0,0 +1,554 @@
+// Copyright 1998-2014 Epic Games, Inc. All Rights Reserved.
+
+#include "GraphEditorCommon.h"
+#include "GraphActionNode.h"
+
+/*******************************************************************************
+ * Static FGraphActionNode Helpers
+ ******************************************************************************/
+
+struct FGraphActionNodeImpl
+{
+	static const int32 DEFAULT_GROUPING = 0;
+
+	/**
+	 * Utility sort function. Compares nodes based off of section, grouping, and
+	 * type.
+	 * 
+	 * @param  LhsMenuNodePtr	The node to determine if it should come first.
+	 * @param  RhsMenuNodePtr	The node to determine if it should come second.
+	 * @return True if LhsMenuNodePtr should come before RhsMenuNodePtr.
+	 */
+	static bool NodeCompare(TSharedPtr<FGraphActionNode> const& LhsMenuNodePtr, TSharedPtr<FGraphActionNode> const& RhsMenuNodePtr);
+
+	/**
+	 * Utility sort function. Compares nodes based off of section, grouping, 
+	 * type, and then alphabetically.
+	 * 
+	 * @param  LhsMenuNodePtr	The node to determine if it should come first.
+	 * @param  RhsMenuNodePtr	The node to determine if it should come second.
+	 * @return True if LhsMenuNodePtr should come before RhsMenuNodePtr.
+	 */
+	static bool AlphabeticalNodeCompare(TSharedPtr<FGraphActionNode> const& LhsMenuNodePtr, TSharedPtr<FGraphActionNode> const& RhsMenuNodePtr);
+};
+
+//------------------------------------------------------------------------------
+bool FGraphActionNodeImpl::NodeCompare(TSharedPtr<FGraphActionNode> const& LhsMenuNodePtr, TSharedPtr<FGraphActionNode> const& RhsMenuNodePtr)
+{
+	FGraphActionNode* LhsMenuNode = LhsMenuNodePtr.Get();
+	FGraphActionNode* RhsMenuNode = RhsMenuNodePtr.Get();
+
+	bool const bLhsIsCategory  = LhsMenuNode->IsCategoryNode();
+	bool const bRhsIsCategory  = RhsMenuNode->IsCategoryNode();
+	bool const bLhsIsSeparator = LhsMenuNode->IsSeparator();
+	bool const bRhsIsSeparator = RhsMenuNode->IsSeparator();
+	bool const bLhsIsSectionHeader = LhsMenuNode->IsSectionHeadingNode();
+	bool const bRhsIsSectionHeader = RhsMenuNode->IsSectionHeadingNode();
+
+	if (LhsMenuNode->SectionID != RhsMenuNode->SectionID)
+	{
+		// since we don't add section headers for children that have the same
+		// section as their parents (the header is above the parent), we need to
+		// organize them first (so they're seemingly under the same header)
+		if ((LhsMenuNode->SectionID == LhsMenuNode->ParentNode.Pin()->SectionID) && 
+			(LhsMenuNode->SectionID != FGraphActionNode::INVALID_SECTION_ID))
+		{
+			return true;
+		}
+		else // otherwise...
+		{
+			// sections are ordered in ascending order
+			return (LhsMenuNode->SectionID < RhsMenuNode->SectionID);
+		}		
+	}
+	else if (bLhsIsSectionHeader != bRhsIsSectionHeader)
+	{
+		// section headers go to the top of that section
+		return bLhsIsSectionHeader;
+	}
+	else if (LhsMenuNode->Grouping != RhsMenuNode->Grouping)
+	{
+		// groups are ordered in descending order
+		return (LhsMenuNode->Grouping >= RhsMenuNode->Grouping);
+	}
+	// next, make sure separators are preserved
+	else if (bLhsIsSeparator != bRhsIsSeparator)
+	{
+		// separators with the same grouping go to the bottom of that "group"
+		return bRhsIsSeparator;
+	}
+	// next, categories get listed before action nodes
+	else if (bLhsIsCategory != bRhsIsCategory)
+	{
+		return bLhsIsCategory;
+	}
+	else
+	{
+		// both lhs and rhs are seemingly the same, so to keep them menu from 
+		// jumping around everytime an entry is added, we sort by the order they
+		// were inserted
+		return (LhsMenuNode->InsertOrder < RhsMenuNode->InsertOrder);
+	}
+}
+
+//------------------------------------------------------------------------------
+bool FGraphActionNodeImpl::AlphabeticalNodeCompare(TSharedPtr<FGraphActionNode> const& LhsMenuNodePtr, TSharedPtr<FGraphActionNode> const& RhsMenuNodePtr)
+{
+	FGraphActionNode* LhsMenuNode = LhsMenuNodePtr.Get();
+	FGraphActionNode* RhsMenuNode = RhsMenuNodePtr.Get();
+
+	bool const bLhsIsCategory  = LhsMenuNode->IsCategoryNode();
+	bool const bRhsIsCategory  = RhsMenuNode->IsCategoryNode();
+	bool const bLhsIsSeparator = LhsMenuNode->IsSeparator();
+	bool const bRhsIsSeparator = RhsMenuNode->IsSeparator();
+	bool const bLhsIsSectionHeader = LhsMenuNode->IsSectionHeadingNode();
+	bool const bRhsIsSectionHeader = RhsMenuNode->IsSectionHeadingNode();
+
+	if (LhsMenuNode->SectionID != RhsMenuNode->SectionID)
+	{
+		// since we don't add section headers for children that have the same
+		// section as their parents (the header is above the parent), we need to
+		// organize them first (so they're seemingly under the same header)
+		if ((LhsMenuNode->SectionID == LhsMenuNode->ParentNode.Pin()->SectionID) &&
+			(LhsMenuNode->SectionID != FGraphActionNode::INVALID_SECTION_ID))
+		{
+			return true;
+		}
+		else // otherwise...
+		{
+			// sections are ordered in ascending order
+			return (LhsMenuNode->SectionID < RhsMenuNode->SectionID);
+		}
+	}
+	else if (bLhsIsSectionHeader != bRhsIsSectionHeader)
+	{
+		// section headers go to the top of that section
+		return bLhsIsSectionHeader;
+	}
+	else if (LhsMenuNode->Grouping != RhsMenuNode->Grouping)
+	{
+		// groups are ordered in descending order
+		return (LhsMenuNode->Grouping >= RhsMenuNode->Grouping);
+	}
+	// next, make sure separators are preserved
+	else if (bLhsIsSeparator != bRhsIsSeparator)
+	{
+		// separators with the same grouping go to the bottom of that "group"
+		return bRhsIsSeparator;
+	}
+	// next, categories get listed before action nodes
+	else if (bLhsIsCategory != bRhsIsCategory)
+	{
+		return bLhsIsCategory;
+	}
+	else if (bLhsIsCategory) // if both nodes are category nodes
+	{
+		// @TODO: Should we be doing localized compares for categories? Probably.
+		return (LhsMenuNode->GetDisplayName().ToString() <= RhsMenuNode->GetDisplayName().ToString());
+	}
+	else // both nodes are action nodes
+	{
+		return (LhsMenuNode->GetDisplayName().CompareTo(RhsMenuNode->GetDisplayName()) <= 0);
+	}
+}
+
+/*******************************************************************************
+ * FGraphActionNode
+ ******************************************************************************/
+
+//------------------------------------------------------------------------------
+TSharedPtr<FGraphActionNode> FGraphActionNode::NewRootNode()
+{
+	// same as a group-divider node, just with an invalid parent
+	return MakeShareable(new FGraphActionNode(FGraphActionNodeImpl::DEFAULT_GROUPING, INVALID_SECTION_ID));
+}
+
+//------------------------------------------------------------------------------
+FGraphActionNode::FGraphActionNode(int32 Grouping, int32 SectionID)
+	: SectionID(SectionID)
+	, Grouping(Grouping)
+	, bPendingRenameRequest(false)
+	, InsertOrder(0)
+{
+}
+
+//------------------------------------------------------------------------------
+FGraphActionNode::FGraphActionNode(TArray< TSharedPtr<FEdGraphSchemaAction> > const& ActionList, int32 Grouping, int32 SectionID)
+	: SectionID(SectionID)
+	, Grouping(Grouping)
+	, Actions(ActionList)
+	, bPendingRenameRequest(false)
+	, InsertOrder(0)
+{
+}
+
+//------------------------------------------------------------------------------
+TSharedPtr<FGraphActionNode> FGraphActionNode::AddChild(FGraphActionListBuilderBase::ActionGroup const& ActionSet)
+{
+	TArray<FString> CategoryStack;
+	ActionSet.GetCategoryChain(CategoryStack);
+
+	TSharedPtr<FGraphActionNode> ActionNode = FGraphActionNode::NewActionNode(ActionSet.Actions);
+	AddChildRecursively(CategoryStack, ActionNode);
+
+	return ActionNode;
+}
+
+//------------------------------------------------------------------------------
+void FGraphActionNode::SortChildren(bool bAlphabetically/* = true*/, bool bRecursive/* = true*/)
+{
+	if (bRecursive)
+	{
+		for (TSharedPtr<FGraphActionNode>& ChildNode : Children)
+		{
+			ChildNode->SortChildren(bAlphabetically, bRecursive);
+		}
+	}
+
+	if (bAlphabetically)
+	{
+		Children.Sort(FGraphActionNodeImpl::AlphabeticalNodeCompare);
+	}
+	else
+	{
+		Children.Sort(FGraphActionNodeImpl::NodeCompare);
+	}
+}
+
+//------------------------------------------------------------------------------
+void FGraphActionNode::GetAllNodes(TArray< TSharedPtr<FGraphActionNode> >& OutNodeArray) const
+{
+	for (TSharedPtr<FGraphActionNode> const& ChildNode : Children)
+	{
+		OutNodeArray.Add(ChildNode);
+		ChildNode->GetAllNodes(OutNodeArray);
+	}
+}
+
+//------------------------------------------------------------------------------
+void FGraphActionNode::GetLeafNodes(TArray< TSharedPtr<FGraphActionNode> >& OutLeafArray) const
+{
+	for (TSharedPtr<FGraphActionNode> const& ChildNode : Children)
+	{
+		if (ChildNode->IsCategoryNode())
+		{
+			ChildNode->GetLeafNodes(OutLeafArray);
+		}
+		else
+		{
+			// @TODO: sometimes, certain action nodes can have children as well
+			//        (for sub-graphs in the "MyBlueprint" tab)
+			OutLeafArray.Add(ChildNode);
+		}
+	}
+}
+
+//------------------------------------------------------------------------------
+void FGraphActionNode::ExpandAllChildren(TSharedPtr< STreeView< TSharedPtr<FGraphActionNode> > > TreeView, bool bRecursive/*= true*/)
+{
+	if (Children.Num() > 0)
+	{
+		TreeView->SetItemExpansion(this->AsShared(), /*ShouldExpandItem =*/true);
+		for (TSharedPtr<FGraphActionNode>& ChildNode : Children)
+		{
+			if (bRecursive)
+			{
+				ChildNode->ExpandAllChildren(TreeView);
+			}
+			else
+			{
+				TreeView->SetItemExpansion(ChildNode, /*ShouldExpandItem =*/true);
+			}
+		}
+	}
+}
+
+//------------------------------------------------------------------------------
+void FGraphActionNode::ClearChildren()
+{
+	Children.Empty();
+	ChildGroupings.Empty();
+	ChildSections.Empty();
+}
+
+//------------------------------------------------------------------------------
+bool FGraphActionNode::IsRootNode() const
+{
+	return (!IsActionNode() && !IsCategoryNode() && !ParentNode.IsValid());
+}
+
+//------------------------------------------------------------------------------
+bool FGraphActionNode::IsSectionHeadingNode() const
+{
+	return (!IsActionNode() && !IsCategoryNode() && !IsRootNode() && (SectionID != INVALID_SECTION_ID));
+}
+
+//------------------------------------------------------------------------------
+bool FGraphActionNode::IsCategoryNode() const
+{
+	return (!IsActionNode() && !DisplayText.IsEmpty());
+}
+
+//------------------------------------------------------------------------------
+bool FGraphActionNode::IsActionNode() const
+{
+	return HasValidAction();
+}
+
+//------------------------------------------------------------------------------
+bool FGraphActionNode::IsGroupDividerNode() const
+{
+	return (!IsActionNode() && !IsCategoryNode() && !IsRootNode() && (SectionID == INVALID_SECTION_ID));
+}
+
+//------------------------------------------------------------------------------
+bool FGraphActionNode::IsSeparator() const
+{
+	return IsSectionHeadingNode() || IsGroupDividerNode();
+}
+
+//------------------------------------------------------------------------------
+FText const& FGraphActionNode::GetDisplayName() const
+{
+	return DisplayText;
+}
+
+//------------------------------------------------------------------------------
+FString FGraphActionNode::GetCategoryPath() const
+{
+	FString CategoryPath;
+	if (IsCategoryNode())
+	{
+		CategoryPath = DisplayText.ToString();
+	}
+
+	TWeakPtr<FGraphActionNode> AncestorNode = ParentNode;
+	while (AncestorNode.IsValid())
+	{
+		CategoryPath = AncestorNode.Pin()->DisplayText.ToString() + TEXT("|") + CategoryPath;
+	}
+	return CategoryPath;
+}
+
+//------------------------------------------------------------------------------
+bool FGraphActionNode::HasValidAction() const
+{
+	return GetPrimaryAction().IsValid();
+}
+
+//------------------------------------------------------------------------------
+TSharedPtr<FEdGraphSchemaAction> FGraphActionNode::GetPrimaryAction() const
+{
+	TSharedPtr<FEdGraphSchemaAction> PrimaryAction;
+	for (TSharedPtr<FEdGraphSchemaAction> NodeAction : Actions)
+	{
+		if (NodeAction.IsValid())
+		{
+			PrimaryAction = NodeAction;
+		}
+	}
+	return PrimaryAction;
+}
+
+//------------------------------------------------------------------------------
+bool FGraphActionNode::BroadcastRenameRequest()
+{
+	if (RenameRequestEvent.IsBound())
+	{
+		RenameRequestEvent.Execute();
+		bPendingRenameRequest = false;
+	}
+	else
+	{
+		bPendingRenameRequest = true;
+	}
+	return bPendingRenameRequest;
+}
+
+//------------------------------------------------------------------------------
+bool FGraphActionNode::IsRenameRequestPending() const
+{
+	return bPendingRenameRequest;
+}
+
+//------------------------------------------------------------------------------
+TSharedPtr<FGraphActionNode> FGraphActionNode::NewSectionHeadingNode(TWeakPtr<FGraphActionNode> Parent, int32 Grouping, int32 SectionID)
+{
+	checkSlow(SectionID != INVALID_SECTION_ID);
+
+	FGraphActionNode* SectionNode = new FGraphActionNode(Grouping, SectionID);
+	SectionNode->ParentNode = Parent;
+	checkSlow(Parent.IsValid());
+
+	return MakeShareable(SectionNode);
+}
+
+//------------------------------------------------------------------------------
+TSharedPtr<FGraphActionNode> FGraphActionNode::NewCategoryNode(FString const& Category, int32 Grouping, int32 SectionID)
+{
+	FGraphActionNode* CategoryNode = new FGraphActionNode(Grouping, SectionID);
+	CategoryNode->DisplayText = FText::FromString(Category);
+
+	return MakeShareable(CategoryNode);
+}
+
+//------------------------------------------------------------------------------
+TSharedPtr<FGraphActionNode> FGraphActionNode::NewActionNode(TArray< TSharedPtr<FEdGraphSchemaAction> > const& ActionList)
+{
+	int32 Grouping  = FGraphActionNodeImpl::DEFAULT_GROUPING;
+	int32 SectionID = INVALID_SECTION_ID;
+
+	for (TSharedPtr<FEdGraphSchemaAction> const& Action : ActionList)
+	{
+		Grouping = FMath::Max(Grouping, Action->Grouping);
+		if (SectionID == INVALID_SECTION_ID)
+		{
+			// take the first non-zero section ID
+			SectionID = Action->SectionID;
+		}
+	}
+
+	FGraphActionNode* ActionNode = new FGraphActionNode(ActionList, Grouping, SectionID);
+	TSharedPtr<FEdGraphSchemaAction> PrimeAction = ActionNode->GetPrimaryAction();
+	checkSlow(PrimeAction.IsValid());
+	ActionNode->DisplayText = PrimeAction->MenuDescription;
+
+	return MakeShareable(ActionNode);
+}
+
+//------------------------------------------------------------------------------
+TSharedPtr<FGraphActionNode> FGraphActionNode::NewGroupDividerNode(TWeakPtr<FGraphActionNode> Parent, int32 Grouping)
+{
+	FGraphActionNode* DividerNode = new FGraphActionNode(Grouping, INVALID_SECTION_ID);
+	DividerNode->ParentNode = Parent;
+	checkSlow(Parent.IsValid());
+
+	return MakeShareable(DividerNode);
+}
+
+//------------------------------------------------------------------------------
+void FGraphActionNode::AddChildRecursively(TArray<FString>& CategoryStack, TSharedPtr<FGraphActionNode> NodeToAdd)
+{
+	if (CategoryStack.Num() > 0)
+	{
+		FString CategorySection = CategoryStack[0];
+		CategoryStack.RemoveAt(0, 1);
+
+		// make sure we don't already have a child that this can nest under
+		TSharedPtr<FGraphActionNode> ExistingNode = FindMatchingParent(CategorySection, NodeToAdd);
+		if (ExistingNode.IsValid())
+		{
+			ExistingNode->AddChildRecursively(CategoryStack, NodeToAdd);
+		}
+		else
+		{
+			TSharedPtr<FGraphActionNode> CategoryNode = NewCategoryNode(CategorySection, NodeToAdd->Grouping, NodeToAdd->SectionID);
+			InsertChild(CategoryNode);
+			CategoryNode->AddChildRecursively(CategoryStack, NodeToAdd);
+		}
+	}
+	else
+	{
+		InsertChild(NodeToAdd);
+	}
+}
+
+//------------------------------------------------------------------------------
+TSharedPtr<FGraphActionNode> FGraphActionNode::FindMatchingParent(FString const& ParentName, TSharedPtr<FGraphActionNode> NodeToAdd)
+{
+	TSharedPtr<FGraphActionNode> FoundCategoryNode;
+
+	// for the "MyBlueprint" tab, sub-graph actions can be nested under graph
+	// actions (meaning that action node can have children).
+	bool const bCanNestUnderActionNodes = NodeToAdd->IsActionNode() && NodeToAdd->GetPrimaryAction()->IsParentable();
+
+	for (TSharedPtr<FGraphActionNode> const& ChildNode : Children)
+	{
+		if (ChildNode->IsCategoryNode())
+		{
+			// @TODO: should we be matching grouping as well?
+			if ((NodeToAdd->SectionID == ChildNode->SectionID) &&
+				(ParentName == ChildNode->DisplayText.ToString()))
+			{
+				FoundCategoryNode = ChildNode;
+				break;
+			}
+		}
+		else if (bCanNestUnderActionNodes && ChildNode->IsActionNode())
+		{
+			// make the action's name into a display name, all categories are 
+			// set as such (to ensure that the action name best matches the 
+			// category ParentName)
+			FString ChildNodeName = FName::NameToDisplayString(ChildNode->DisplayText.ToString(), /*bIsBool =*/false);
+
+			// @TODO: should we be matching section/grouping as well?
+			if (ChildNodeName == ParentName)
+			{
+				FoundCategoryNode = ChildNode;
+				break;
+			}
+		}
+	}
+
+	return FoundCategoryNode;
+}
+
+//------------------------------------------------------------------------------
+void FGraphActionNode::InsertChild(TSharedPtr<FGraphActionNode> NodeToAdd)
+{
+	ensure(!NodeToAdd->IsRootNode());
+	ensure(!IsSeparator());
+
+	NodeToAdd->ParentNode = this->AsShared();
+
+	if (NodeToAdd->SectionID != INVALID_SECTION_ID)
+	{
+		// don't need a section heading if the parent is under the same section
+		bool const bAddSectionHeading = (NodeToAdd->SectionID != SectionID) && 
+			// make sure we already haven't already added a heading for this section
+			!ChildSections.Contains(NodeToAdd->SectionID) &&
+			// if this node also has a category, use that over a section heading
+			(!NodeToAdd->IsActionNode() || NodeToAdd->GetPrimaryAction()->Category.IsEmpty());
+
+		if (bAddSectionHeading)
+		{
+			ChildSections.Add(NodeToAdd->SectionID); // to avoid recursion, add before we insert
+			InsertChild(NewSectionHeadingNode(NodeToAdd->ParentNode, NodeToAdd->Grouping, NodeToAdd->SectionID));
+		}
+	}
+	// we don't use group-dividers inside of sections (we use groups to more to
+	// hardcode the order), but if this isn't in a section...
+	else if (!ChildGroupings.Contains(NodeToAdd->Grouping))
+	{
+		// don't need a divider if this is the first group
+		if (ChildGroupings.Num() > 0)
+		{
+			int32 LowestGrouping = MAX_int32;
+			for (int32 Group : ChildGroupings)
+			{
+				LowestGrouping = FMath::Min(LowestGrouping, Group);
+			}
+			// dividers come at the end of a menu group, so it would be 
+			// undesirable to add it for NodeToAdd->Grouping if that group is 
+			// lower than all the others (the lowest group should not have a 
+			// divider associated with it)
+			int32 DividerGrouping = FMath::Max(LowestGrouping, NodeToAdd->Grouping);
+
+			ChildGroupings.Add(NodeToAdd->Grouping); // to avoid recursion, add before we insert
+			InsertChild(NewGroupDividerNode(NodeToAdd->ParentNode, DividerGrouping));
+		}
+		else
+		{
+			ChildGroupings.Add(NodeToAdd->Grouping);
+		}
+	}
+
+	NodeToAdd->InsertOrder = Children.Num();
+	Children.Add(NodeToAdd);
+}
+
+
+
+
+
+
diff --git a/Engine/Source/Editor/GraphEditor/Private/GraphActionNode.h b/Engine/Source/Editor/GraphEditor/Private/GraphActionNode.h
--- a/Engine/Source/Editor/GraphEditor/Private/GraphActionNode.h
+++ b/Engine/Source/Editor/GraphEditor/Private/GraphActionNode.h
@@ -1,421 +1,263 @@
 // Copyright 1998-2014 Epic Games, Inc. All Rights Reserved.
 
-// Utility class for building menus of graph actions
+#pragma once
 
-#ifndef __GraphActionNode_h__
-#define __GraphActionNode_h__
+#include "SGraphActionMenu.h" // for FOnRenameRequestActionNode
 
-struct FGraphActionNode : TSharedFromThis<FGraphActionNode>
+// Utility class for building menus of graph actions
+struct GRAPHEDITOR_API FGraphActionNode : TSharedFromThis<FGraphActionNode>
 {
-	/** The section this node belongs to (if any) */
-	int32 SectionID;
+public:
+	/** */
+	static const int32 INVALID_SECTION_ID = 0;
 
-	/** The category this node belongs in. */
-	FString Category;
-	/** The category chain, separated using '|' */
-	FString CategoryChain;
+	/** Identifies the named section that this node belongs to, if any (defaults to INVALID_SECTION_ID) */
+	int32 const SectionID;
+	/** Identifies the menu group that this node belongs to (defaults to zero) */
+	int32 const Grouping;
+	/** A set of actions to execute when this node is picked from a menu */
+	TArray< TSharedPtr<FEdGraphSchemaAction> > const Actions;
 
-	int32 Grouping;
-	TArray< TSharedPtr<FEdGraphSchemaAction> > Actions;
+	/** */
+	TArray< TSharedPtr<FGraphActionNode> > Children;
 
-	/** When the item is first created, a rename request may occur before everything is setup for it. This toggles to true in those cases */
-	bool bIsRenameRequestBeforeReady;
+public:
+	/**
+	 * Static allocator for a new root node (so external users have a starting
+	 * point to build graph action trees from).
+	 *
+	 * @return A newly allocated root node (should not be displayed in the tree view).
+	 */
+	static TSharedPtr<FGraphActionNode> NewRootNode();
+
+	/**
+	 * Inserts a new action node (and any accompanying category nodes) based off
+	 * the provided ActionSet. 
+	 *
+	 * NOTE: This does NOT insert the node in a sorted manner. Call SortChildren() 
+	 *       separately to accomplish that (done for performance reasons).
+	 * 
+	 * @param  ActionSet	A list of actions that you want the node to execute when picked.
+	 * @return The new action node.
+	 */
+	TSharedPtr<FGraphActionNode> AddChild(FGraphActionListBuilderBase::ActionGroup const& ActionSet);
+
+	/**
+	 * Sorts all child nodes by section, group, and type (additionally, can
+	 * sort alphabetically if wanted).
+	 * 
+	 * @param  bAlphabetically	Determines if we sort alphabetically on top of section/group/type.
+	 * @param  bRecursive		Determines if we should sort all decedent nodes' children ass well.
+	 */
+	void SortChildren(bool bAlphabetically = true, bool bRecursive = true);
+
+	/**
+	 * Recursively collects all child/grandchild/decedent nodes.
+	 * 
+	 * @param  OutNodeArray	The array to fill out with decedent nodes.
+	 */
+	void GetAllNodes(TArray< TSharedPtr<FGraphActionNode> >& OutNodeArray) const;
+
+	/**
+	 * Recursively collects all decedent action/separator nodes (leaves out 
+	 * branching category-nodes).
+	 * 
+	 * @param  OutLeafArray	The array to fill out with decedent leaf nodes.
+	 */
+	void GetLeafNodes(TArray< TSharedPtr<FGraphActionNode> >& OutLeafArray) const;
+
+	/**
+	 * Takes the tree view and expands its elements for each child.
+	 * 
+	 * @param  TreeView		The tree responsible for visualizing this node hierarchy.
+	 * @param  bRecursive	Determines if you want children/decedents to expand their children as well. 
+	 */
+	void ExpandAllChildren(TSharedPtr< STreeView< TSharedPtr<FGraphActionNode> > > TreeView, bool bRecursive = true);
+
+	/**
+	 * Clears all children (not recursively... the TSharedPtrs should clean up 
+	 * appropriately).
+	 */
+	void ClearChildren();
+
+	/**
+	 * Query to determine this node's type (there are five distinguishable node
+	 * types: root, section heading, category, action, & group-divider).
+	 *
+	 * @return True if this is the type your queried about, otherwise false.
+	 */
+	bool IsRootNode() const;
+	bool IsSectionHeadingNode() const;
+	bool IsCategoryNode() const;
+	bool IsActionNode() const;
+	bool IsGroupDividerNode() const;
+
+	/**
+	 * Determines if this node is a menu separator of some kind (either a
+	 * "group-divider" or a "section heading").
+	 *
+	 * @return True if this is a menu divider, otherwise false.
+	 */
+	bool IsSeparator() const;
+	/**
+	 * Retrieves this node's display name (for category and action nodes). The
+	 * text string will be empty for separator and root nodes.
+	 *
+	 * @return The name to present this node with in the tree view (will be an empty text string if this is a separator node)
+	 */
+	FText const& GetDisplayName() const;
+
+	/**
+	 * Walks the node chain backwards, constructing a category path (delimited
+	 * by '|' characters). This includes this node's category (if it is a
+	 * category node).
+	 *
+	 * @return A category path string, denoting the category hierarchy up to this node.
+	 */
+	FString GetCategoryPath() const;
+
+	/**
+	 * Checks to see if this node contains at least one valid action.
+	 *
+	 * @return True is the Actions array contains a valid entry, otherwise false.
+	 */
+	bool HasValidAction() const;
 
-	/** Broadcasts whenever renaming a layer is requested */
+	/**
+	 * Looks through this node's Actions array, and returns the first valid
+	 * action it finds.
+	 *
+	 * @return This node's first valid action (will be an empty pointer if this is not an action node).
+	 */
+	TSharedPtr<FEdGraphSchemaAction> GetPrimaryAction() const;
+
+	/**
+	 * Accessor to the node's RenameRequestEvent (for binding purposes). Do not
+	 * Execute() the delegate from this function, instead call
+	 * BroadcastRenameRequest() on the node.
+	 *
+	 * @return The node's internal RenameRequestEvent.
+	 */
 	FOnRenameRequestActionNode& OnRenameRequest() { return RenameRequestEvent; }
-	void BroadcastRenameRequest() { RenameRequestEvent.ExecuteIfBound(); }
 
-	/** Broadcasts whenever a rename is requested */
-	FOnRenameRequestActionNode RenameRequestEvent;
+	/**
+	 * Executes the node's RenameRequestEvent if it is bound. Otherwise, it will
+	 * mark the node as having a pending rename request.
+	 *
+	 * @return True if the broadcast went through, false if the "pending rename request" flag was set.
+	 */
+	bool BroadcastRenameRequest();
 
-	TArray< TSharedPtr<FGraphActionNode> > Children;
-public:
-	static TSharedPtr<FGraphActionNode> NewAction( TSharedPtr<FEdGraphSchemaAction> InAction )
-	{
-		return MakeShareable( new FGraphActionNode( TEXT(""), (InAction.IsValid() ? InAction->Grouping : 0), InAction) );
-	}
-
-	static TSharedPtr<FGraphActionNode> NewAction( TArray< TSharedPtr<FEdGraphSchemaAction> >& InActionList )
-	{
-		int32 Grouping = 0;
-		int32 SectionID = 0;
-		for ( int32 ActionIndex = 0; ActionIndex < InActionList.Num(); ActionIndex++ )
-		{
-			TSharedPtr<FEdGraphSchemaAction> CurrentAction = InActionList[ActionIndex];
-			if ( CurrentAction.IsValid() )
-			{
-				Grouping = FMath::Max( CurrentAction->Grouping, Grouping );
-				// Copy the section ID if we don't already have one
-				if( SectionID == 0 )
-				{
-					SectionID = CurrentAction->SectionID;
-				}
-			}
-		}
-
-		return MakeShareable( new FGraphActionNode( TEXT(""), Grouping, InActionList, SectionID ) );
-	}
-
-	static TSharedPtr<FGraphActionNode> NewCategory( const FString& InCategory, const FString& InCategoryChain, int32 InGrouping = 0 )
-	{
-		return MakeShareable( new FGraphActionNode(InCategory, InCategoryChain, InGrouping, NULL) );
-	}
-
-	static TSharedPtr<FGraphActionNode> NewSeparator( int32 InGrouping = 0 )
-	{
-		return MakeShareable( new FGraphActionNode(TEXT(""), InGrouping, NULL) );
-	}
-
-	static TSharedPtr<FGraphActionNode> NewSeparatorWithTitle( int32 InGrouping, int32 InNodeSection )
-	{
-		return MakeShareable( new FGraphActionNode( InGrouping, InNodeSection) );
-	}
-
-	bool IsActionNode() const
-	{
-		return (Actions.Num() > 0 && Actions[0].IsValid());
-	}
-
-	bool IsCategoryNode() const
-	{
-		return (Actions.Num()==0 || !Actions[0].IsValid()) && !Category.IsEmpty();
-	}
-
-	bool IsSeparator() const
-	{
-		return (Actions.Num()==0 || !Actions[0].IsValid()) && Category.IsEmpty();
-	}
-
-	void AddChild(TSharedPtr<FGraphActionNode> NodeToAdd, TArray<FString>& CategoryStack, bool bAlphaSort, FString& CategoryChain)
-	{
-		// If we still have some categories in the chain to process...
-		if( CategoryStack.Num() )
-		{
-			// ...check to see if we have any 
-			const FString NextCategory = CategoryStack[0];
-			CategoryStack.RemoveAt(0,1);
-			bool bFoundMatch = false;
-			for( int32 i = 0; i < Children.Num(); i++ )
-			{
-				TSharedPtr<FGraphActionNode>& CurrentChild = Children[i];
-				if(CurrentChild->IsCategoryNode())
-				{
-					bFoundMatch = NextCategory == CurrentChild->Category && NodeToAdd->SectionID == CurrentChild->SectionID;
-				}
-				else if(CurrentChild->IsActionNode())
-				{
-					if ( NodeToAdd->IsActionNode() )
-					{
-						TSharedPtr<FEdGraphSchemaAction> CurrentAction = NodeToAdd->Actions[0];
-
-						// Make the action's name into a display name, all categories are set as such
-						bFoundMatch = CurrentAction->IsParentable()
-							&& NextCategory == FName::NameToDisplayString(CurrentChild->Actions[0]->MenuDescription.ToString(), false);
-					}
-				}
-				
-				if( bFoundMatch )
-				{
-					CategoryChain = CategoryChain.IsEmpty() ? NextCategory : FString::Printf( TEXT( "%s|%s" ), *CategoryChain, *NextCategory );
-					CurrentChild->AddChild(NodeToAdd, CategoryStack, bAlphaSort, CategoryChain);
-					break;
-				}
-			}
-
-			if( !bFoundMatch )
-			{
-				CategoryChain = CategoryChain.IsEmpty() ? NextCategory : FString::Printf( TEXT( "%s|%s" ), *CategoryChain, *NextCategory );
-				TSharedPtr<FGraphActionNode> NewCategoryNode = NewCategory(NextCategory, CategoryChain, NodeToAdd->Grouping);
-				// Copy the section ID also 
-				NewCategoryNode->SectionID = NodeToAdd->SectionID;
-				AddToChildrenSorted(NewCategoryNode, bAlphaSort);
-				NewCategoryNode->AddChild(NodeToAdd, CategoryStack, bAlphaSort, CategoryChain);
-			}
-		}
-		else
-		{
-			// End of the category chain, add it to this level!
-			AddToChildrenSorted(NodeToAdd, bAlphaSort);
-		}
-	}
+	/**
+	 * Sometimes a call to BroadcastRenameRequest() is made before the
+	 * RenameRequestEvent has been bound. When that happens, this node is
+	 * marked with a pending rename request. This method determines if that is
+	 * the case for this node.
+	 *
+	 * @return True if a call to BroadcastRenameRequest() was made without a valid RenameRequestEvent.
+	 */
+	bool IsRenameRequestPending() const;
+
+private:
+	/**
+	 *
+	 *
+	 * @param  Grouping
+	 * @param  SectionID
+	 */
+	FGraphActionNode(int32 Grouping, int32 SectionID);
 
 	/**
-	* Iterates through all this node's children, and tells the tree view to expand them
-	*/
-	void ExpandAllChildren( TSharedPtr< STreeView< TSharedPtr<FGraphActionNode> > > TreeView )
-	{
-		if( Children.Num() )
-		{
-			TreeView->SetItemExpansion(this->AsShared(), true);
-			for( int32 i = 0; i < Children.Num(); i++ )
-			{
-				Children[i]->ExpandAllChildren(TreeView);
-			}
-		}
-	}
-
-	void ClearChildren()
-	{
-		Children.Empty();
-	}
-
-	void AddToChildrenSorted(TSharedPtr<FGraphActionNode> NodeToAdd, bool bAlphaSort)
-	{
- 		if( NodeToAdd->SectionID != 0 )
- 		{
- 			return AddToChildrenSortedWithSection( NodeToAdd, bAlphaSort );
- 		}
-		// Run over the children list, and break when we figure out where we belong
-		int32 i;
-		for(i = 0; i < Children.Num(); i++)
-		{
-			if( NodeCompare(NodeToAdd, Children[i], bAlphaSort) )
-			{
-				break;
-			}
-		}
-
-		// Check to see if this is the first node in a boundary between groupings, and add a separator if so
-		const int32 CurrentChildren = Children.Num();
-		if( i < CurrentChildren - 1 )
-		{
-			TSharedPtr<FGraphActionNode> PreviousEntry = Children[i];
-			if( NodeToAdd->Grouping != PreviousEntry->Grouping )
-			{
-				Children.Insert( NewSeparator(NodeToAdd->Grouping), i );
-			}
-		}
-
-		// Finally, add the item
-		Children.Insert(NodeToAdd, i);
-	}
-
-	
-	void AddToChildrenSortedWithSection(TSharedPtr<FGraphActionNode> NodeToAdd, bool bAlphaSort)
-	{
-		// Run over the children list, and break when we figure out where we belong
-		int32 i;
-		for(i = 0; i < Children.Num(); i++)
-		{
-			if( NodeCompareWithSection(NodeToAdd, Children[i], bAlphaSort) )
-			{
-				break;
-			}
-		}
-
-		const int32 CurrentChildren = Children.Num();
-		bool AddGroup = false;		
-		if ( CurrentChildren != 0)
-		{
-			// If we have any children check if the previous section matches, if not, insert a new one
-			int32 PreviousIndex = FMath::Clamp( i-1, 0, CurrentChildren - 1 );
-			TSharedPtr<FGraphActionNode> PreviousEntry = Children[PreviousIndex];
-			if( NodeToAdd->SectionID != PreviousEntry->SectionID )
-			{
-				AddGroup = true;
-			}
-		}
-		else
-		{
-			// There are no children so we need to create a section
-			AddGroup = true;
-		}
-		// Override the group requirement if this node is in a category (only variables are shown in categories anyway - no point in adding separators)
-		if( (NodeToAdd->Actions.Num() != 0 ) && ( NodeToAdd->Actions[0].IsValid() == true ) && (!NodeToAdd->Actions[0]->Category.IsEmpty()) )
-		{
-			AddGroup = false;
-		}
-		// If we want a group - add it now
-		if( AddGroup )
-		{
-			Children.Insert( NewSeparatorWithTitle(NodeToAdd->Grouping, NodeToAdd->SectionID), i );
-			i++;
-		}		
-
-		// Finally, add the item
-		Children.Insert(NodeToAdd, i);
-	}
-
-	/** Comparison function for sorted insertion:  returns true when NodeA belongs before NodeB, false otherwise */
-	bool NodeCompareWithSection(TSharedPtr<FGraphActionNode> NodeA, TSharedPtr<FGraphActionNode> NodeB, bool bAlphaSort)
-	{
-		const bool bNodeAIsSeparator = NodeA->IsSeparator();
-		const bool bNodeBIsSeparator = NodeB->IsSeparator();
-
-		const bool bNodeAIsCategory = NodeA->IsCategoryNode();
-		const bool bNodeBIsCategory = NodeB->IsCategoryNode();
-
-		if( ( NodeA->SectionID == NodeB->SectionID) && ( bNodeAIsSeparator == true ) )
-		{
-			return false;
-		}
-		if( NodeA->SectionID != NodeB->SectionID )
-		{
-			return NodeA->SectionID < NodeB->SectionID;
-		}
-		else if ( bNodeBIsCategory == true ) 
-		{
-			// If next node is a category and we are not, we go first
-			if( bNodeAIsCategory == false )
-			{
-				return true;
-			}
-
-			// Both nodes must categories - insert me here if this is my category
-			return (NodeA->Category == NodeB->Category);		
-		}		
-		else
-		{
-			if(bAlphaSort)
-			{
-				// Finally, alphabetize
-				if( bNodeAIsCategory )
-				{
-					return (NodeA->Category <= NodeB->Category);
-				}
-				else
-				{
-					return (NodeA->Actions[0]->MenuDescription.CompareTo(NodeB->Actions[0]->MenuDescription) <= 0);
-				}
-			}
-			else
-			{
-				return false;
-			}
-		}
-	}
-	/** Comparison function for sorted insertion:  returns true when NodeA belongs before NodeB, false otherwise */
-	bool NodeCompare(TSharedPtr<FGraphActionNode> NodeA, TSharedPtr<FGraphActionNode> NodeB, bool bAlphaSort)
-	{
-		const bool bNodeAIsSeparator = NodeA->IsSeparator();
-		const bool bNodeBIsSeparator = NodeB->IsSeparator();
-
-		const bool bNodeAIsCategory = NodeA->IsCategoryNode();
-		const bool bNodeBIsCategory = NodeB->IsCategoryNode();
-
-		// Grouping trumps all
-		if( NodeA->Grouping != NodeB->Grouping )
-		{
-			return NodeA->Grouping >= NodeB->Grouping;
-		}
-		// Next, make sure separators are preserved
-		else if( bNodeAIsSeparator != bNodeBIsSeparator )
-		{
-			return bNodeBIsSeparator;
-		}
-		// Next, categories come before terminals
-		else if( bNodeAIsCategory != bNodeBIsCategory )
-		{
-			return bNodeAIsCategory;
-		}
-		else
-		{
-			if(bAlphaSort)
-			{
-				// Finally, alphabetize
-				if( bNodeAIsCategory )
-				{
-					return (NodeA->Category <= NodeB->Category);
-				}
-				else
-				{
-					return (NodeA->Actions[0]->MenuDescription.CompareTo(NodeB->Actions[0]->MenuDescription) <= 0);
-				}
-			}
-			else
-			{
-				return false;
-			}
-		}
-	}
-
-	void GetLeafNodes(TArray< TSharedPtr<FGraphActionNode> >& LeafArray)
-	{
-		for (int32 ChildIndex = 0; ChildIndex < Children.Num(); ++ChildIndex)
-		{
-			TSharedPtr<FGraphActionNode>& Child = Children[ChildIndex];
-
-			if (Child->IsCategoryNode())
-			{
-				Child->GetLeafNodes(LeafArray);
-			}
-			else
-			{
-				LeafArray.Add(Child);
-			}
-		}
-	}
-
-	void GetAllNodes(TArray< TSharedPtr<FGraphActionNode> >& OutNodes, bool bLeavesOnly)
-	{
-		for (int32 ChildIndex = 0; ChildIndex < Children.Num(); ++ChildIndex)
-		{
-			TSharedPtr<FGraphActionNode>& Child = Children[ChildIndex];
-
-			// Add us if a leaf, or we want all nodes (not just leaves)
-			if(!Child->IsCategoryNode() || !bLeavesOnly)
-			{
-				OutNodes.Add(Child);
-			}
-
-			// If a category, recurse into childen
-			if (Child->IsCategoryNode())
-			{
-				Child->GetAllNodes(OutNodes, bLeavesOnly);
-			}
-		}
-	}
-
-	/** Returns the node's string used for display purposes */
-	FString GetNodeString() const
-	{
-		if( IsCategoryNode() )
-		{
-			return Category;
-		}
-		else if(Actions.Num() > 0 && Actions[0].IsValid())
-		{
-			return Actions[0]->MenuDescription.ToString();
-		}
-		return FString();
-	}
-protected:
-
-	FGraphActionNode( FString InCategory, int32 InGrouping, TSharedPtr<FEdGraphSchemaAction> InAction )
-		: SectionID ( 0 )
-		, Category( InCategory )
-		, Grouping( InGrouping )
-		, bIsRenameRequestBeforeReady( false )
-	{
-		Actions.Add( InAction );
-	}
-
-	FGraphActionNode( FString InCategory, FString InCategoryChain, int32 InGrouping, TSharedPtr<FEdGraphSchemaAction> InAction )
-		: SectionID ( 0 )
-		, Category( InCategory )
-		, CategoryChain( InCategoryChain )
-		, Grouping( InGrouping )
-		, bIsRenameRequestBeforeReady( false )
-	{
-		Actions.Add( InAction );
-	}
-
-	FGraphActionNode( FString InCategory, int32 InGrouping, TArray< TSharedPtr<FEdGraphSchemaAction> >& InActionList, int32 InSectionID = 0 )
-		: SectionID ( InSectionID )
-		, Category( InCategory )
-		, Grouping( InGrouping )
-		, bIsRenameRequestBeforeReady( false )
-	{
-		Actions.Append( InActionList );
-	}
-
-	FGraphActionNode( int32 InGrouping, int32 InNodeSection )		
-		: SectionID ( InNodeSection )
-		, Grouping( InGrouping )
-		, bIsRenameRequestBeforeReady( false )
-	{
-		Actions.Add( NULL );		
-	}
-};
+	 * Constructor for action nodes. Private so that users go through AddChild().
+	 *
+	 * @param  ActionList
+	 * @param  Grouping
+	 * @param  SectionID
+	 */
+	FGraphActionNode(TArray< TSharedPtr<FEdGraphSchemaAction> > const& ActionList, int32 Grouping, int32 SectionID);
+
+	/**
+	 *
+	 *
+	 * @param  Parent
+	 * @param  Grouping
+	 * @param  SectionID
+	 * @return
+	 */
+	static TSharedPtr<FGraphActionNode> NewSectionHeadingNode(TWeakPtr<FGraphActionNode> Parent, int32 Grouping, int32 SectionID);
+
+	/**
+	 *
+	 *
+	 * @param  Category
+	 * @param  Grouping
+	 * @param  SectionID
+	 * @return
+	 */
+	static TSharedPtr<FGraphActionNode> NewCategoryNode(FString const& Category, int32 Grouping, int32 SectionID);
 
-#endif // __GraphActionNode_h__
+	/**
+	 *
+	 *
+	 * @param  ActionList
+	 * @return
+	 */
+	static TSharedPtr<FGraphActionNode> NewActionNode(TArray< TSharedPtr<FEdGraphSchemaAction> > const& ActionList);
+
+	/**
+	 *
+	 *
+	 * @param  Parent
+	 * @param  Grouping
+	 * @return
+	 */
+	static TSharedPtr<FGraphActionNode> NewGroupDividerNode(TWeakPtr<FGraphActionNode> Parent, int32 Grouping);
+
+	/**
+	 * Whittles down the CategoryStack, adding category-nodes as needed. The
+	 * last category is what the node will be inserted under.
+	 *
+	 * @param  CategoryStack	A list of categories denoting where to nest the new node (the first element is the highest category)
+	 * @param  NodeToAdd		The node you want inserted.
+	 */
+	void AddChildRecursively(TArray<FString>& CategoryStack, TSharedPtr<FGraphActionNode> NodeToAdd);
+
+	/**
+	 * Looks through this node's children to see if a there already exists a 
+	 * node matching one we'd have to spawn (to parent the supplied NodeToAdd).
+	 * 
+	 * @param   ParentName	The name of the category NodeToAdd wants to nest under.
+	 * @param   NodeToAdd	The node that we'll be adding to this child.
+	 * @return  A child node matching the supplied parameters (will be empty if no match was found).
+	 */
+	TSharedPtr<FGraphActionNode> FindMatchingParent(FString const& ParentName, TSharedPtr<FGraphActionNode> NodeToAdd);
+
+	/**
+	 * Adds the specified node directly to this node's Children array. Will
+	 * create and insert separators if needed (if the node has a new group or
+	 * section).
+	 *
+	 * @param  NodeToAdd The node you want inserted.
+	 */
+	void InsertChild(TSharedPtr<FGraphActionNode> NodeToAdd);
+
+private:
+	/** The category or action name (depends on what type of node this is) */
+	FText DisplayText;
+	/** The node that this is a direct child of (empty if this is a root node) */
+	TWeakPtr<FGraphActionNode> ParentNode;
+
+	/** Tracks what groups have already been added (so we can easily determine what group-dividers we need) */
+	TSet<int32> ChildGroupings;
+	/** Tracks what sections have already been added (so we can easily determine what heading we need) */
+	TSet<int32> ChildSections;
+
+	/** When the item is first created, a rename request may occur before everything is setup for it. This toggles to true in those cases */
+	bool bPendingRenameRequest;
+	/** Delegate to trigger when a rename was requested on this node */
+	FOnRenameRequestActionNode RenameRequestEvent;
+
+	friend struct FGraphActionNodeImpl;
+	/** For sorting, when we don't alphabetically sort (so menu items don't jump around). */
+	int32 InsertOrder;	
+};
diff --git a/Engine/Source/Editor/GraphEditor/Private/SGraphActionMenu.cpp b/Engine/Source/Editor/GraphEditor/Private/SGraphActionMenu.cpp
--- a/Engine/Source/Editor/GraphEditor/Private/SGraphActionMenu.cpp
+++ b/Engine/Source/Editor/GraphEditor/Private/SGraphActionMenu.cpp
@@ -90,7 +90,7 @@ public:
 
 		FText CategoryTooltip;
 		FString CategoryLink, CategoryExcerpt;
-		FEditorCategoryUtils::GetCategoryTooltipInfo(*InActionNode->Category, CategoryTooltip, CategoryLink, CategoryExcerpt);
+		FEditorCategoryUtils::GetCategoryTooltipInfo(*InActionNode->GetDisplayName().ToString(), CategoryTooltip, CategoryLink, CategoryExcerpt);
 
 		TSharedRef<SToolTip> ToolTipWidget = IDocumentation::Get()->CreateToolTip(CategoryTooltip, NULL, CategoryLink, CategoryExcerpt);
 
@@ -102,7 +102,7 @@ public:
 			[
 				SAssignNew(InlineWidget, SInlineEditableTextBlock)
 				.Font( FSlateFontInfo( FPaths::EngineContentDir() / TEXT("Slate/Fonts/Roboto-Bold.ttf"), 9 )  )
-				.Text( FEditorCategoryUtils::GetCategoryDisplayString(FText::FromString(InActionNode->Category)) )
+				.Text( FEditorCategoryUtils::GetCategoryDisplayString(InActionNode->GetDisplayName()) )
 				.ToolTip( ToolTipWidget )
 				.HighlightText( InArgs._HighlightText )
 				.OnVerifyTextChanged( this, &SGraphActionCategoryWidget::OnVerifyTextChanged )
@@ -119,7 +119,7 @@ public:
 		TSharedPtr<FGraphEditorDragDropAction> GraphDropOp = DragDropEvent.GetOperationAs<FGraphEditorDragDropAction>();
 		if (GraphDropOp.IsValid())
 		{
-			GraphDropOp->DroppedOnCategory( ActionNode.Pin()->CategoryChain );
+			GraphDropOp->DroppedOnCategory( ActionNode.Pin()->GetCategoryPath() );
 			return FReply::Handled();
 		}
 		return FReply::Unhandled();
@@ -130,7 +130,7 @@ public:
 		TSharedPtr<FGraphEditorDragDropAction> GraphDropOp = DragDropEvent.GetOperationAs<FGraphEditorDragDropAction>();
 		if (GraphDropOp.IsValid())
 		{
-			GraphDropOp->SetHoveredCategoryName( ActionNode.Pin()->Category );
+			GraphDropOp->SetHoveredCategoryName( ActionNode.Pin()->GetDisplayName().ToString() );
 		}
 	}
 
@@ -178,7 +178,7 @@ void SGraphActionMenu::Construct( const FArguments& InArgs, bool bIsReadOnly/* =
 	this->OnCategoryTextCommitted = InArgs._OnCategoryTextCommitted;
 	this->OnCanRenameSelectedAction = InArgs._OnCanRenameSelectedAction;
 	this->OnGetSectionTitle = InArgs._OnGetSectionTitle;
-	this->FilteredRootAction = FGraphActionNode::NewCategory(TEXT("FILTEREDROOT"), TEXT( "" ));
+	this->FilteredRootAction = FGraphActionNode::NewRootNode();
 	this->OnActionMatchesName = InArgs._OnActionMatchesName;
 	
 	// If a delegate for filtering text is passed in, assign it so that it will be used instead of the built-in filter box
@@ -253,12 +253,12 @@ void SGraphActionMenu::RefreshAllActions(bool bPreserveExpansion, bool bHandleOn
 
 		if(bHandleOnSelectionEvent)
 		{
-			SelectItemByName(*SelectedAction->GetNodeString(), ESelectInfo::OnMouseClick, SelectedAction->SectionID, SelectedNodes[0]->IsCategoryNode());
+			SelectItemByName(*SelectedAction->GetDisplayName().ToString(), ESelectInfo::OnMouseClick, SelectedAction->SectionID, SelectedNodes[0]->IsCategoryNode());
 		}
 		else
 		{
 			// If we do not want to handle the selection, set it directly so it will reselect the item but not handle the event.
-			SelectItemByName(*SelectedAction->GetNodeString(), ESelectInfo::Direct, SelectedAction->SectionID, SelectedNodes[0]->IsCategoryNode());
+			SelectItemByName(*SelectedAction->GetDisplayName().ToString(), ESelectInfo::Direct, SelectedAction->SectionID, SelectedNodes[0]->IsCategoryNode());
 		}
 	}
 }
@@ -287,14 +287,9 @@ void SGraphActionMenu::OnRequestRenameOnActionNode()
 	TArray< TSharedPtr<FGraphActionNode> > SelectedNodes = TreeView->GetSelectedItems();
 	if(SelectedNodes.Num() > 0)
 	{
-		if(SelectedNodes[0]->RenameRequestEvent.IsBound())
-		{
-			SelectedNodes[0]->BroadcastRenameRequest();
-		}
-		else
+		if (!SelectedNodes[0]->BroadcastRenameRequest())
 		{
 			TreeView->RequestScrollIntoView(SelectedNodes[0]);
-			SelectedNodes[0]->bIsRenameRequestBeforeReady = true;
 		}
 	}
 }
@@ -313,7 +308,7 @@ bool SGraphActionMenu::CanRequestRenameOnActionNode() const
 FString SGraphActionMenu::GetSelectedCategoryName() const
 {
 	TArray< TSharedPtr<FGraphActionNode> > SelectedNodes = TreeView->GetSelectedItems();
-	return (SelectedNodes.Num() > 0) ? SelectedNodes[0]->Category : FString();
+	return (SelectedNodes.Num() > 0) ? SelectedNodes[0]->GetDisplayName().ToString() : FString();
 }
 
 void SGraphActionMenu::GetSelectedCategorySubActions(TArray<TSharedPtr<FEdGraphSchemaAction>>& OutActions) const
@@ -358,11 +353,11 @@ bool SGraphActionMenu::SelectItemByName(const FName& ItemName, ESelectInfo::Type
 		TSharedPtr<FGraphActionNode> SelectionNode;
 
 		TArray<TSharedPtr<FGraphActionNode>> GraphNodes;
-		FilteredRootAction->GetAllNodes(GraphNodes, false);
+		FilteredRootAction->GetAllNodes(GraphNodes);
 		for (int32 i = 0; i < GraphNodes.Num() && !SelectionNode.IsValid(); ++i)
 		{
 			TSharedPtr<FGraphActionNode> CurrentGraphNode = GraphNodes[i];
-			FEdGraphSchemaAction* GraphAction = CurrentGraphNode->Actions[0].Get();
+			FEdGraphSchemaAction* GraphAction = CurrentGraphNode->GetPrimaryAction().Get();
 
 			// If the user is attempting to select a category, make sure it's a category
 			if( CurrentGraphNode->IsCategoryNode() == bIsCategory )
@@ -378,7 +373,7 @@ bool SGraphActionMenu::SelectItemByName(const FName& ItemName, ESelectInfo::Type
 							break;
 						}
 					}
-					else if( CurrentGraphNode->Category == FName::NameToDisplayString(ItemName.ToString(), false) )
+					else if (CurrentGraphNode->GetDisplayName().ToString() == FName::NameToDisplayString(ItemName.ToString(), false))
 					{
 						SelectionNode = CurrentGraphNode;
 
@@ -410,7 +405,7 @@ bool SGraphActionMenu::SelectItemByName(const FName& ItemName, ESelectInfo::Type
 									break;
 								}
 							}
-							else if( CurrentChildNode->Category == FName::NameToDisplayString(ItemName.ToString(), false) )
+							else if (CurrentChildNode->GetDisplayName().ToString() == FName::NameToDisplayString(ItemName.ToString(), false))
 							{
 								SelectionNode = CurrentChildNode;
 
@@ -442,10 +437,10 @@ void SGraphActionMenu::ExpandCategory(const FString& CategoryName)
 	if (CategoryName.Len())
 	{
 		TArray<TSharedPtr<FGraphActionNode>> GraphNodes;
-		FilteredRootAction->GetAllNodes(GraphNodes, false);
+		FilteredRootAction->GetAllNodes(GraphNodes);
 		for (int32 i = 0; i < GraphNodes.Num(); ++i)
 		{
-			if (GraphNodes[i]->Category == CategoryName)
+			if (GraphNodes[i]->GetDisplayName().ToString() == CategoryName)
 			{
 				GraphNodes[i]->ExpandAllChildren(TreeView);
 			}
@@ -459,16 +454,16 @@ static bool CompareGraphActionNode(TSharedPtr<FGraphActionNode> A, TSharedPtr<FG
 	check(B.IsValid());
 
 	// First check grouping is the same
-	if(A->Category != B->Category)
+	if (A->GetDisplayName().ToString() != B->GetDisplayName().ToString())
 	{
 		return false;
 	}
 
-	if(A->Actions[0].IsValid() && B->Actions[0].IsValid())
+	if (A->HasValidAction() && B->HasValidAction())
 	{
-		return A->Actions[0]->MenuDescription.CompareTo(B->Actions[0]->MenuDescription) == 0;
+		return A->GetPrimaryAction()->MenuDescription.CompareTo(B->GetPrimaryAction()->MenuDescription) == 0;
 	}
-	else if(!A->Actions[0].IsValid() && !B->Actions[0].IsValid())
+	else if(!A->HasValidAction() && !B->HasValidAction())
 	{
 		return true;
 	}
@@ -502,7 +497,6 @@ void RestoreExpansionState(TSharedPtr< STreeView<ItemType> > InTree, const TArra
 	}
 }
 
-
 void SGraphActionMenu::GenerateFilteredItems(bool bPreserveExpansion)
 {
 	// First, save off current expansion state
@@ -567,25 +561,20 @@ void SGraphActionMenu::GenerateFilteredItems(bool bPreserveExpansion)
 				BestMatchCount = EachWeight;
 				BestMatchIndex = CurTypeIndex;
 			}								
-			// Add the action to the filtered list.  This will automatically place it in the right subcategory
-			TArray<FString> CategoryChain;
-			CurrentAction.GetCategoryChain(CategoryChain);
-			
-			TSharedPtr<FGraphActionNode> NewNode = FGraphActionNode::NewAction(CurrentAction.Actions);
-			FString EmptyCategoryChain;
-			FilteredRootAction->AddChild(NewNode, CategoryChain, bAlphaSortItems, EmptyCategoryChain);
+			FilteredRootAction->AddChild(CurrentAction);
 		}
 	}
+	FilteredRootAction->SortChildren(bAlphaSortItems, /*bRecursive =*/true);
 
 	TreeView->RequestTreeRefresh();
 
 	// Update the filtered list (needs to be done in a separate pass because the list is sorted as items are inserted)
 	FilteredActionNodes.Empty();
-	FilteredRootAction->GetAllNodes(FilteredActionNodes, true);
+	FilteredRootAction->GetLeafNodes(FilteredActionNodes);
 
 	// Get _all_ new nodes (flattened tree basically)
 	TArray< TSharedPtr<FGraphActionNode> > AllNodes;
-	FilteredRootAction->GetAllNodes(AllNodes, false);
+	FilteredRootAction->GetAllNodes(AllNodes);
 
 	// If theres a BestMatchIndex find it in the actions nodes and select it (maybe this should check the current selected suggestion first ?)
 	if( BestMatchIndex != INDEX_NONE ) 
@@ -595,7 +584,7 @@ void SGraphActionMenu::GenerateFilteredItems(bool bPreserveExpansion)
 		{
 			for (int32 iNode = 0; iNode < FilteredActionNodes.Num() ; iNode++)
 			{
-				if( FilteredActionNodes[ iNode ].Get()->Actions[ 0 ] == FilterSelectAction.Actions[ 0 ] )
+				if( FilteredActionNodes[ iNode ].Get()->GetPrimaryAction() == FilterSelectAction.Actions[ 0 ] )
 				{
 					SelectedSuggestion = iNode;
 				}
@@ -806,9 +795,8 @@ void SGraphActionMenu::OnNameTextCommitted(const FText& NewText, ETextCommit::Ty
 
 void SGraphActionMenu::OnItemScrolledIntoView( TSharedPtr<FGraphActionNode> InActionNode, const TSharedPtr<ITableRow>& InWidget )
 {
-	if(InActionNode->bIsRenameRequestBeforeReady)
+	if (InActionNode->IsRenameRequestPending())
 	{
-		InActionNode->bIsRenameRequestBeforeReady = false;
 		InActionNode->BroadcastRenameRequest();
 	}
 }
@@ -830,10 +818,10 @@ TSharedRef<ITableRow> SGraphActionMenu::MakeWidget( TSharedPtr<FGraphActionNode>
 
 	if( InItem->IsActionNode() )
 	{
-		check(InItem->Actions.Num() > 0 && InItem->Actions[0].IsValid() );
+		check(InItem->HasValidAction());
 
 		FCreateWidgetForActionData CreateData(&InItem->OnRenameRequest());
-		CreateData.Action = InItem->Actions[0];
+		CreateData.Action = InItem->GetPrimaryAction();
 		CreateData.HighlightText = TAttribute<FText>(this, &SGraphActionMenu::GetFilterText);
 		CreateData.MouseButtonDownDelegate = FCreateWidgetMouseButtonDown::CreateSP( this, &SGraphActionMenu::OnMouseButtonDownEvent );
 
@@ -947,8 +935,8 @@ TSharedRef<ITableRow> SGraphActionMenu::MakeWidget( TSharedPtr<FGraphActionNode>
 
 		if (InItem->IsActionNode())
 		{
-			check(InItem->Actions.Num() > 0);
-			CreateData.RowAction = InItem->Actions[0];
+			check(InItem->HasValidAction());
+			CreateData.RowAction = InItem->GetPrimaryAction();
 		}
 
 		ExpanderWidget = OnCreateCustomRowExpander.Execute(CreateData);
@@ -1027,7 +1015,7 @@ FReply SGraphActionMenu::OnItemDragDetected( const FGeometry& MyGeometry, const
 			{
 				if(OnCategoryDragged.IsBound())
 				{
-					return OnCategoryDragged.Execute(Node->CategoryChain, MouseEvent);
+					return OnCategoryDragged.Execute(Node->GetCategoryPath(), MouseEvent);
 				}
 			}
 			// Dragging an action
@@ -1061,9 +1049,9 @@ bool SGraphActionMenu::OnMouseButtonDownEvent( TWeakPtr<FEdGraphSchemaAction> In
 		{
 			SelectedNode = FilteredActionNodes[0];			
 		}
-		if( ( SelectedNode.IsValid() ) && ( SelectedNode->Actions.Num() > 0 ) && ( SelectedNode->Actions[0].IsValid() ) )
+		if (SelectedNode.IsValid() && SelectedNode->HasValidAction())
 		{
-			if( SelectedNode->Actions[0].Get() == InAction.Pin().Get() )
+			if( SelectedNode->GetPrimaryAction().Get() == InAction.Pin().Get() )
 			{				
 				bResult = HandleSelection( SelectedNode );
 			}
