Data Table Blueprint nodes: TTP: 346291 any BP with GetDataTableRow node is always marked as dirty

SGraphPinNameList::GetDefaultValueWidget() was calling  GraphPinObj->GetSchema()->TrySetDefaultValue(), but it turns out that was not actually necessary, so just removed that call.

diff --git a/Engine/Source/Editor/GraphEditor/Private/KismetPins/SGraphPinNameList.cpp b/Engine/Source/Editor/GraphEditor/Private/KismetPins/SGraphPinNameList.cpp
--- a/Engine/Source/Editor/GraphEditor/Private/KismetPins/SGraphPinNameList.cpp
+++ b/Engine/Source/Editor/GraphEditor/Private/KismetPins/SGraphPinNameList.cpp
@@ -32,7 +32,6 @@ TSharedRef<SWidget>	SGraphPinNameList::GetDefaultValueWidget()
 			}
 		}
 		check(CurrentlySelectedName.IsValid());
-		GraphPinObj->GetSchema()->TrySetDefaultValue(*GraphPinObj, *CurrentlySelectedName->ToString());
 	}
 
 	// Create widget
