fix for drag dropping onto categories in My Blueprint causing an infinite loop.
#Branch UE4
#Proj GraphEditor

note: also patched up what appeared to be a repeated typo, decedent for descendent :)

diff --git a/Engine/Source/Editor/GraphEditor/Private/GraphActionNode.cpp b/Engine/Source/Editor/GraphEditor/Private/GraphActionNode.cpp
--- a/Engine/Source/Editor/GraphEditor/Private/GraphActionNode.cpp
+++ b/Engine/Source/Editor/GraphEditor/Private/GraphActionNode.cpp
@@ -325,7 +325,9 @@ FString FGraphActionNode::GetCategoryPath() const
 	TWeakPtr<FGraphActionNode> AncestorNode = ParentNode;
 	while (AncestorNode.IsValid())
 	{
-		CategoryPath = AncestorNode.Pin()->DisplayText.ToString() + TEXT("|") + CategoryPath;
+		TSharedPtr<FGraphActionNode> ActionNode = AncestorNode.Pin();
+		CategoryPath = ActionNode->DisplayText.ToString() + TEXT("|") + CategoryPath;
+		AncestorNode = ActionNode->GetParentNode();
 	}
 	return CategoryPath;
 }
diff --git a/Engine/Source/Editor/GraphEditor/Private/GraphActionNode.h b/Engine/Source/Editor/GraphEditor/Private/GraphActionNode.h
--- a/Engine/Source/Editor/GraphEditor/Private/GraphActionNode.h
+++ b/Engine/Source/Editor/GraphEditor/Private/GraphActionNode.h
@@ -47,22 +47,27 @@ public:
 	 * sort alphabetically if wanted).
 	 * 
 	 * @param  bAlphabetically	Determines if we sort alphabetically on top of section/group/type.
-	 * @param  bRecursive		Determines if we should sort all decedent nodes' children ass well.
+	 * @param  bRecursive		Determines if we should sort all decendent nodes' children ass well.
 	 */
 	void SortChildren(bool bAlphabetically = true, bool bRecursive = true);
 
 	/**
-	 * Recursively collects all child/grandchild/decedent nodes.
+	 * Returns a WeakPtr to the Parent Node
+	 */
+	TWeakPtr<FGraphActionNode> GetParentNode() const{ return ParentNode; }
+
+	/**
+	 * Recursively collects all child/grandchild/decendent nodes.
 	 * 
-	 * @param  OutNodeArray	The array to fill out with decedent nodes.
+	 * @param  OutNodeArray	The array to fill out with decendent nodes.
 	 */
 	void GetAllNodes(TArray< TSharedPtr<FGraphActionNode> >& OutNodeArray) const;
 
 	/**
-	 * Recursively collects all decedent action/separator nodes (leaves out 
+	 * Recursively collects all decendent action/separator nodes (leaves out 
 	 * branching category-nodes).
 	 * 
-	 * @param  OutLeafArray	The array to fill out with decedent leaf nodes.
+	 * @param  OutLeafArray	The array to fill out with decendent leaf nodes.
 	 */
 	void GetLeafNodes(TArray< TSharedPtr<FGraphActionNode> >& OutLeafArray) const;
 
@@ -70,7 +75,7 @@ public:
 	 * Takes the tree view and expands its elements for each child.
 	 * 
 	 * @param  TreeView		The tree responsible for visualizing this node hierarchy.
-	 * @param  bRecursive	Determines if you want children/decedents to expand their children as well. 
+	 * @param  bRecursive	Determines if you want children/decendents to expand their children as well. 
 	 */
 	void ExpandAllChildren(TSharedPtr< STreeView< TSharedPtr<FGraphActionNode> > > TreeView, bool bRecursive = true);
 
