fix for drag dropping onto categories in My Blueprint causing an infinite loop.
#Branch UE4
#Proj GraphEditor

note previous fix could generate bad category paths if it wasn't chained.

diff --git a/Engine/Source/Editor/GraphEditor/Private/GraphActionNode.cpp b/Engine/Source/Editor/GraphEditor/Private/GraphActionNode.cpp
--- a/Engine/Source/Editor/GraphEditor/Private/GraphActionNode.cpp
+++ b/Engine/Source/Editor/GraphEditor/Private/GraphActionNode.cpp
@@ -325,9 +325,13 @@ FString FGraphActionNode::GetCategoryPath() const
 	TWeakPtr<FGraphActionNode> AncestorNode = ParentNode;
 	while (AncestorNode.IsValid())
 	{
-		TSharedPtr<FGraphActionNode> ActionNode = AncestorNode.Pin();
-		CategoryPath = ActionNode->DisplayText.ToString() + TEXT("|") + CategoryPath;
-		AncestorNode = ActionNode->GetParentNode();
+		const FText& DisplayText = AncestorNode.Pin()->DisplayText;
+
+		if( !DisplayText.IsEmpty() )
+		{
+			CategoryPath = DisplayText.ToString() + TEXT("|") + CategoryPath;
+		}
+		AncestorNode = AncestorNode.Pin()->GetParentNode();
 	}
 	return CategoryPath;
 }
