GraphEditor: Adjust spline tangent calculations to improve look and feel of wires in all of the graph editors
- Splines now take vertical travel into account in tangent generation, and use a stronger setting when going horziontally backwards (loopbacks)
- Various parameters related to tangent calculations can now be adjusted in the Graph Editor section of Editor Settings

diff --git a/Engine/Source/Editor/GraphEditor/Private/ConnectionDrawingPolicy.cpp b/Engine/Source/Editor/GraphEditor/Private/ConnectionDrawingPolicy.cpp
--- a/Engine/Source/Editor/GraphEditor/Private/ConnectionDrawingPolicy.cpp
+++ b/Engine/Source/Editor/GraphEditor/Private/ConnectionDrawingPolicy.cpp
@@ -190,6 +190,7 @@ FKismetNodeInfoContext::FKismetNodeInfoContext(UEdGraph* SourceGraph)
 FConnectionDrawingPolicy::FConnectionDrawingPolicy(int32 InBackLayerID, int32 InFrontLayerID, float InZoomFactor, const FSlateRect& InClippingRect, FSlateWindowElementList& InDrawElements)
 	: WireLayerID(InBackLayerID)
 	, ArrowLayerID(InFrontLayerID)
+	, Settings(GetDefault<UGraphEditorSettings>())
 	, ZoomFactor(InZoomFactor)
 	, ClippingRect(InClippingRect)
 	, DrawElementsList(InDrawElements)
@@ -326,8 +327,7 @@ void FConnectionDrawingPolicy::DrawConnection( int32 LayerId, const FVector2D& S
 	const FVector2D& P0 = Start;
 	const FVector2D& P1 = End;
 
-	const int32 Tension  = FMath::Abs<int32>(Start.X - End.X);
-	const FVector2D P0Tangent = Tension * FVector2D(1.0f, 0);
+	const FVector2D P0Tangent = Settings->ComputeSplineTangent(P0, P1);
 	const FVector2D P1Tangent = P0Tangent;
 
 	// Draw the spline itself
@@ -531,8 +531,6 @@ FKismetConnectionDrawingPolicy::FKismetConnectionDrawingPolicy(int32 InBackLayer
 	: FConnectionDrawingPolicy(InBackLayerID, InFrontLayerID, ZoomFactor, InClippingRect, InDrawElements)
 	, GraphObj(InGraphObj)
 {
-	const UGraphEditorSettings* Settings = GetDefault<UGraphEditorSettings>();
-
 	// Don't want to draw ending arrowheads
 	ArrowImage = nullptr;
 	ArrowRadius = FVector2D::ZeroVector;
@@ -1245,8 +1243,6 @@ FSoundCueGraphConnectionDrawingPolicy::FSoundCueGraphConnectionDrawingPolicy(int
 	, GraphObj(InGraphObj)
 {
 	// Cache off the editor options
-	const UGraphEditorSettings* Settings = GetDefault<UGraphEditorSettings>();
-
 	ActiveColor = Settings->TraceAttackColor;
 	InactiveColor = Settings->TraceReleaseColor;
 
diff --git a/Engine/Source/Editor/GraphEditor/Private/ConnectionDrawingPolicy.h b/Engine/Source/Editor/GraphEditor/Private/ConnectionDrawingPolicy.h
--- a/Engine/Source/Editor/GraphEditor/Private/ConnectionDrawingPolicy.h
+++ b/Engine/Source/Editor/GraphEditor/Private/ConnectionDrawingPolicy.h
@@ -32,6 +32,9 @@ protected:
 	const FSlateBrush* ArrowImage;
 	const FSlateBrush* MidpointImage;
 	const FSlateBrush* BubbleImage;
+
+	const UGraphEditorSettings* Settings;
+
 public:
 	FVector2D ArrowRadius;
 	FVector2D MidpointRadius;
diff --git a/Engine/Source/Editor/GraphEditor/Private/GraphEditorSettings.cpp b/Engine/Source/Editor/GraphEditor/Private/GraphEditorSettings.cpp
--- a/Engine/Source/Editor/GraphEditor/Private/GraphEditorSettings.cpp
+++ b/Engine/Source/Editor/GraphEditor/Private/GraphEditorSettings.cpp
@@ -12,6 +12,14 @@
 
 UGraphEditorSettings::UGraphEditorSettings( const class FPostConstructInitializeProperties& PCIP )
 	: Super(PCIP)
+	, ForwardSplineHorizontalDeltaRange(1000.0f)
+	, ForwardSplineVerticalDeltaRange(1000.0f)
+	, ForwardSplineTangentFromHorizontalDelta(1.0f, 0.0f)
+	, ForwardSplineTangentFromVerticalDelta(1.0f, 0.0f)
+	, BackwardSplineHorizontalDeltaRange(200.0f)
+	, BackwardSplineVerticalDeltaRange(200.0f)
+	, BackwardSplineTangentFromHorizontalDelta(3.0f, 0.0f)
+	, BackwardSplineTangentFromVerticalDelta(1.5f, 0.0f)
 {
 	DataPinStyle = BPST_VariantA;
 
@@ -64,3 +72,21 @@ UGraphEditorSettings::UGraphEditorSettings( const class FPostConstructInitialize
 	ExecSequenceNodeTitleColor= FLinearColor(0.8f, 0.4f, 0.4f, 1.0f);
 	ResultNodeTitleColor = FLinearColor(1.0f, 0.65f, 0.4f, 1.0f);
 }
+
+FVector2D UGraphEditorSettings::ComputeSplineTangent(const FVector2D& Start, const FVector2D& End) const
+{
+	const FVector2D DeltaPos = End - Start;
+	const bool bGoingForward = DeltaPos.X >= 0.0f;
+
+	const float ClampedTensionX = FMath::Min<float>(FMath::Abs<float>(DeltaPos.X), bGoingForward ? ForwardSplineHorizontalDeltaRange : BackwardSplineHorizontalDeltaRange);
+	const float ClampedTensionY = FMath::Min<float>(FMath::Abs<float>(DeltaPos.Y), bGoingForward ? ForwardSplineVerticalDeltaRange : BackwardSplineVerticalDeltaRange);
+
+	if (bGoingForward)
+	{
+		return (ClampedTensionX * ForwardSplineTangentFromHorizontalDelta) + (ClampedTensionY * ForwardSplineTangentFromVerticalDelta);
+	}
+	else
+	{
+		return (ClampedTensionX * BackwardSplineTangentFromHorizontalDelta) + (ClampedTensionY * BackwardSplineTangentFromVerticalDelta);
+	}
+}
diff --git a/Engine/Source/Editor/GraphEditor/Public/GraphEditorSettings.h b/Engine/Source/Editor/GraphEditor/Public/GraphEditorSettings.h
--- a/Engine/Source/Editor/GraphEditor/Public/GraphEditorSettings.h
+++ b/Engine/Source/Editor/GraphEditor/Public/GraphEditorSettings.h
@@ -24,6 +24,41 @@ public:
 	UPROPERTY(config, EditAnywhere, Category=GeneralStyle)
 	TEnumAsByte<EBlueprintPinStyleType> DataPinStyle;
 
+public:
+	/** The maximum value to clamp the absolute value of the horizontal distance between endpoints when calculating tangents (when the wire is moving forward) */
+	UPROPERTY(config, EditAnywhere, Category=Splines, AdvancedDisplay)
+	float ForwardSplineHorizontalDeltaRange;
+
+	/** The maximum value to clamp the absolute value of the vertical distance between endpoints when calculating tangents (when the wire is moving forward) */
+	UPROPERTY(config, EditAnywhere, Category=Splines, AdvancedDisplay)
+	float ForwardSplineVerticalDeltaRange;
+
+	/** The amount that the horizontal delta affects the generated tangent handle of splines (when the wire is moving forward) */
+	UPROPERTY(config, EditAnywhere, Category=Splines, AdvancedDisplay)
+	FVector2D ForwardSplineTangentFromHorizontalDelta;
+
+	/** The amount that the vertical delta affects the generated tangent handle of splines (when the wire is moving forward) */
+	UPROPERTY(config, EditAnywhere, Category=Splines, AdvancedDisplay)
+	FVector2D ForwardSplineTangentFromVerticalDelta;
+
+	/** The maximum value to clamp the absolute value of the horizontal distance between endpoints when calculating tangents (when the wire is moving backwards) */
+	UPROPERTY(config, EditAnywhere, Category=Splines, AdvancedDisplay)
+	float BackwardSplineHorizontalDeltaRange;
+
+	/** The maximum value to clamp the absolute value of the vertical distance between endpoints when calculating tangents (when the wire is moving backwards) */
+	UPROPERTY(config, EditAnywhere, Category=Splines, AdvancedDisplay)
+	float BackwardSplineVerticalDeltaRange;
+
+	/** The amount that the horizontal delta affects the generated tangent handle of splines (when the wire is moving backwards) */
+	UPROPERTY(config, EditAnywhere, Category=Splines, AdvancedDisplay)
+	FVector2D BackwardSplineTangentFromHorizontalDelta;
+
+	/** The amount that the vertical delta affects the generated tangent handle of splines (when the wire is moving backwards) */
+	UPROPERTY(config, EditAnywhere, Category=Splines, AdvancedDisplay)
+	FVector2D BackwardSplineTangentFromVerticalDelta;
+
+	/** Compute a tangent given the start and end points of the spline */
+	FVector2D ComputeSplineTangent(const FVector2D& Start, const FVector2D& End) const;
 public:
 
 	/** The default color is used only for types not specifically defined below.  Generally if it's seen, it means another type needs to be defined so that the wire in question can have an appropriate color. */
