Blueprint: Remove some redundant code from SGraphNodeK2Event::AddPin, relying on the inherited behavior except for delegate pins

diff --git a/Engine/Source/Editor/GraphEditor/Private/KismetNodes/SGraphNodeK2Event.cpp b/Engine/Source/Editor/GraphEditor/Private/KismetNodes/SGraphNodeK2Event.cpp
--- a/Engine/Source/Editor/GraphEditor/Private/KismetNodes/SGraphNodeK2Event.cpp
+++ b/Engine/Source/Editor/GraphEditor/Private/KismetNodes/SGraphNodeK2Event.cpp
@@ -7,30 +7,13 @@
 
 void SGraphNodeK2Event::AddPin( const TSharedRef<SGraphPin>& PinToAdd ) 
 {
-	PinToAdd->SetOwner( SharedThis(this) );
-
 	const UEdGraphPin* PinObj = PinToAdd->GetPinObj();
-	const bool bAdvancedParameter = PinObj && PinObj->bAdvancedView;
-	if(bAdvancedParameter)
-	{
-		PinToAdd->SetVisibility( TAttribute<EVisibility>(PinToAdd, &SGraphPin::IsPinVisibleAsAdvanced) );
-	}
+	const bool bDelegateOutput = (PinObj != nullptr) && (UK2Node_Event::DelegateOutputName == PinObj->PinName);
 
-	const bool bDelegateOutput = PinObj && (UK2Node_Event::DelegateOutputName == PinObj->PinName);
-	if (PinToAdd->GetDirection() == EEdGraphPinDirection::EGPD_Input)
-	{
-		LeftNodeBox->AddSlot()
-			.AutoHeight()
-			.HAlign(HAlign_Left)
-			.VAlign(VAlign_Center)
-			.Padding(10,4)
-		[
-			PinToAdd
-		];
-		InputPins.Add(PinToAdd);
-	}
-	else if (bDelegateOutput && TitleAreaWidget.IsValid())
+	if (bDelegateOutput && TitleAreaWidget.IsValid())
 	{
+		PinToAdd->SetOwner(SharedThis(this));
+
 		bHasDelegateOutputPin = true;
 		PinToAdd->SetShowLabel(false);
 		TitleAreaWidget->AddSlot()
@@ -40,19 +23,12 @@ void SGraphNodeK2Event::AddPin( const TSharedRef<SGraphPin>& PinToAdd )
 		[
 			PinToAdd
 		];
+
 		OutputPins.Add(PinToAdd);
 	}
-	else // Direction == EEdGraphPinDirection::EGPD_Output
+	else
 	{
-		RightNodeBox->AddSlot()
-			.AutoHeight()
-			.HAlign(HAlign_Right)
-			.VAlign(VAlign_Center)
-			.Padding(10,4)
-		[
-			PinToAdd
-		];
-		OutputPins.Add(PinToAdd);
+		SGraphNodeK2Default::AddPin(PinToAdd);
 	}
 }
 
