GraphEditor: Pad the height of multi-line node titles to be a multiple of the graph snap grid taller than single-line nodes, so the pins will still line up if you place the node N cell snaps above

diff --git a/Engine/Source/Editor/GraphEditor/Private/SGraphNode.cpp b/Engine/Source/Editor/GraphEditor/Private/SGraphNode.cpp
--- a/Engine/Source/Editor/GraphEditor/Private/SGraphNode.cpp
+++ b/Engine/Source/Editor/GraphEditor/Private/SGraphNode.cpp
@@ -70,11 +70,32 @@ void SNodeTitle::RebuildWidget()
 	TArray<FString> Lines;
 	CachedTitle.ToString().ParseIntoArray(&Lines, TEXT("\n"), false);
 
-	if(Lines.Num())
+	if (Lines.Num())
 	{
 		CachedHeadTitle = FText::FromString(Lines[0]);
 	}
 
+	// Pad the height of multi-line node titles to be a multiple of the graph snap grid taller than
+	// single-line nodes, so the pins will still line up if you place the node N cell snaps above
+	if (Lines.Num() > 1)
+	{
+		// Note: This code a little fragile, and will need to be updated if the font or padding of titles
+		// changes in the future, but the failure mode is just a slight misalignment.
+		const int32 EstimatedExtraHeight = (Lines.Num() - 1) * 13;
+
+		const int32 SnapSize = (int32)SNodePanel::GetSnapGridSize();
+		const int32 PadSize = SnapSize - (EstimatedExtraHeight % SnapSize);
+
+		if (PadSize < SnapSize)
+		{
+			VerticalBox->AddSlot()
+			[
+				SNew(SSpacer)
+				.Size(FVector2D(1.0f, PadSize))
+			];
+		}
+	}
+
 	// Make a separate widget for each line, using a less obvious style for subsequent lines
 	for (int32 Index = 1; Index < Lines.Num(); ++Index)
 	{
@@ -568,9 +589,9 @@ TSharedPtr<SWidget>	SGraphNode::SetupErrorReporting()
 	UpdateErrorInfo();
 
 	// generate widget
-	SAssignNew(ErrorText, SErrorText )
-			.BackgroundColor( this, &SGraphNode::GetErrorColor )
-			.ToolTipText( this, &SGraphNode::GetErrorMsgToolTip );
+	SAssignNew(ErrorText, SErrorText)
+		.BackgroundColor( this, &SGraphNode::GetErrorColor )
+		.ToolTipText( this, &SGraphNode::GetErrorMsgToolTip );
 
 	ErrorReporting = ErrorText;
 	ErrorReporting->SetError(ErrorMsg);
@@ -901,14 +922,9 @@ TSharedPtr<SGraphPin> SGraphNode::CreatePinWidget(UEdGraphPin* Pin) const
 	return FNodeFactory::CreatePinWidget(Pin);
 }
 
-/**
- * Add a new pin to this graph node. The pin must be newly created.
- *
- * @param PinToAdd   A new pin to add to this GraphNode.
- */
-void SGraphNode::AddPin( const TSharedRef<SGraphPin>& PinToAdd )
+void SGraphNode::AddPin(const TSharedRef<SGraphPin>& PinToAdd)
 {	
-	PinToAdd->SetOwner( SharedThis(this) );
+	PinToAdd->SetOwner(SharedThis(this));
 
 	const UEdGraphPin* PinObj = PinToAdd->GetPinObj();
 	const bool bAdvancedParameter = PinObj && PinObj->bAdvancedView;
