GraphEditor: Make various padding values adjustable to iterate on more compact nodes (WIP / Experimental)

diff --git a/Engine/Source/Editor/GraphEditor/Private/ConnectionDrawingPolicy.cpp b/Engine/Source/Editor/GraphEditor/Private/ConnectionDrawingPolicy.cpp
--- a/Engine/Source/Editor/GraphEditor/Private/ConnectionDrawingPolicy.cpp
+++ b/Engine/Source/Editor/GraphEditor/Private/ConnectionDrawingPolicy.cpp
@@ -233,8 +233,12 @@ void FConnectionDrawingPolicy::DrawSplineWithArrow(const FVector2D& StartPoint,
 
 void FConnectionDrawingPolicy::DrawSplineWithArrow(FGeometry& StartGeom, FGeometry& EndGeom, const FLinearColor& WireColor, float WireThickness, bool bDrawBubbles, bool Bidirectional)
 {
-	const FVector2D StartPoint = FGeometryHelper::VerticalMiddleRightOf( StartGeom );
-	const FVector2D EndPoint = FGeometryHelper::VerticalMiddleLeftOf( EndGeom ) - FVector2D(ArrowRadius.X, 0);
+	//@TODO: These values should be pushed into the Slate style, they are compensating for a bit of
+	// empty space inside of the pin brush images.
+	const float StartFudgeX = 4.0f;
+	const float EndFudgeX = 4.0f;
+	const FVector2D StartPoint = FGeometryHelper::VerticalMiddleRightOf(StartGeom) - FVector2D(StartFudgeX, 0.0f);
+	const FVector2D EndPoint = FGeometryHelper::VerticalMiddleLeftOf(EndGeom) - FVector2D(ArrowRadius.X - EndFudgeX, 0);
 
 	DrawSplineWithArrow(StartPoint, EndPoint, WireColor, WireThickness, bDrawBubbles, Bidirectional);
 }
diff --git a/Engine/Source/Editor/GraphEditor/Private/GraphEditorSettings.cpp b/Engine/Source/Editor/GraphEditor/Private/GraphEditorSettings.cpp
--- a/Engine/Source/Editor/GraphEditor/Private/GraphEditorSettings.cpp
+++ b/Engine/Source/Editor/GraphEditor/Private/GraphEditorSettings.cpp
@@ -12,6 +12,11 @@
 
 UGraphEditorSettings::UGraphEditorSettings( const class FPostConstructInitializeProperties& PCIP )
 	: Super(PCIP)
+	, PaddingAbovePin(4.0f)
+	, PaddingBelowPin(4.0f)
+	, PaddingRightOfInput(10.0f)
+	, PaddingLeftOfOutput(10.0f)
+	, PaddingTowardsNodeEdge(10.0f)
 	, ForwardSplineHorizontalDeltaRange(1000.0f)
 	, ForwardSplineVerticalDeltaRange(1000.0f)
 	, ForwardSplineTangentFromHorizontalDelta(1.0f, 0.0f)
@@ -69,10 +74,50 @@ UGraphEditorSettings::UGraphEditorSettings( const class FPostConstructInitialize
 	ParentFunctionCallNodeTitleColor = FLinearColor(1.0f, 0.170000f, 0.0f, 1.0f);
 	FunctionTerminatorNodeTitleColor = FLinearColor(0.6f, 0.0f, 1.0f, 1.0f);
 	ExecBranchNodeTitleColor = FLinearColor(1.0f, 1.0f, 1.0f, 1.0f);
-	ExecSequenceNodeTitleColor= FLinearColor(0.8f, 0.4f, 0.4f, 1.0f);
+	ExecSequenceNodeTitleColor = FLinearColor(0.8f, 0.4f, 0.4f, 1.0f);
 	ResultNodeTitleColor = FLinearColor(1.0f, 0.65f, 0.4f, 1.0f);
 }
 
+#if WITH_EDITOR
+void UGraphEditorSettings::PostEditChangeProperty(FPropertyChangedEvent& PropertyChangedEvent)
+{
+	const FName PropertyName = (PropertyChangedEvent.Property != nullptr) ? PropertyChangedEvent.Property->GetFName() : NAME_None;
+
+	if ((PropertyName == GET_MEMBER_NAME_CHECKED(UGraphEditorSettings, DataPinStyle)) || (PropertyName == GET_MEMBER_NAME_CHECKED(UGraphEditorSettings, PaddingTowardsNodeEdge)))
+	{
+		// Invalidate all node graph editors
+		//@TODO: That thing I said
+	}
+
+	Super::PostEditChangeProperty(PropertyChangedEvent);
+}
+#endif
+
+FMargin UGraphEditorSettings::GetInputPinPadding() const
+{
+	const float HPad = FMath::Max<float>(PaddingTowardsNodeEdge, 0.0f);
+	return FMargin(HPad, PaddingAbovePin, PaddingRightOfInput, PaddingBelowPin);
+}
+
+FMargin UGraphEditorSettings::GetOutputPinPadding() const
+{
+	const float HPad = FMath::Max<float>(PaddingTowardsNodeEdge, 0.0f);
+	return FMargin(PaddingLeftOfOutput, PaddingAbovePin, HPad, PaddingBelowPin);
+}
+
+FMargin UGraphEditorSettings::GetNonPinNodeBodyPadding() const
+{
+	const float NegativeHPad = FMath::Max<float>(-PaddingTowardsNodeEdge, 0.0f);
+	return FMargin(NegativeHPad, 0.0f, NegativeHPad, 0.0f);
+}
+
+FVector2D UGraphEditorSettings::GetShadowDeltaSize() const
+{
+	FVector2D ShadowSize = FEditorStyle::GetVector(TEXT("Graph.Node.ShadowSize"));
+	ShadowSize.X += FMath::Min<float>(PaddingTowardsNodeEdge, 0.0f);
+	return ShadowSize;
+}
+
 FVector2D UGraphEditorSettings::ComputeSplineTangent(const FVector2D& Start, const FVector2D& End) const
 {
 	const FVector2D DeltaPos = End - Start;
diff --git a/Engine/Source/Editor/GraphEditor/Private/KismetNodes/SGraphNodeFormatText.cpp b/Engine/Source/Editor/GraphEditor/Private/KismetNodes/SGraphNodeFormatText.cpp
--- a/Engine/Source/Editor/GraphEditor/Private/KismetNodes/SGraphNodeFormatText.cpp
+++ b/Engine/Source/Editor/GraphEditor/Private/KismetNodes/SGraphNodeFormatText.cpp
@@ -45,10 +45,13 @@ void SGraphNodeFormatText::CreateInputSideAddButton(TSharedPtr<SVerticalBox> Inp
 										NSLOCTEXT("FormatTextNode", "FormatTextNodeAddPinButton_Tooltip", "Adds an argument to the node"),
 										false);
 
+	FMargin AddPinPadding = Settings->GetInputPinPadding();
+	AddPinPadding.Top += 6.0f;
+
 	InputBox->AddSlot()
 	.AutoHeight()
 	.VAlign(VAlign_Center)
-	.Padding(10,10,10,4)
+	.Padding(AddPinPadding)
 	[
 		AddPinButton
 	];
diff --git a/Engine/Source/Editor/GraphEditor/Private/KismetNodes/SGraphNodeK2Base.cpp b/Engine/Source/Editor/GraphEditor/Private/KismetNodes/SGraphNodeK2Base.cpp
--- a/Engine/Source/Editor/GraphEditor/Private/KismetNodes/SGraphNodeK2Base.cpp
+++ b/Engine/Source/Editor/GraphEditor/Private/KismetNodes/SGraphNodeK2Base.cpp
@@ -77,6 +77,7 @@ void SGraphNodeK2Base::UpdateCompactNode()
 			.ToolTip( NodeToolTip.ToSharedRef() )
 			.AddMetaData<FGraphNodeMetaData>(TagMeta)
 			+SOverlay::Slot()
+			.Padding(Settings->GetNonPinNodeBodyPadding())
 			[
 				SNew(SImage)
 				.Image( FEditorStyle::GetBrush("Graph.CompactNode.Body") )
@@ -84,6 +85,7 @@ void SGraphNodeK2Base::UpdateCompactNode()
 			+SOverlay::Slot()
 			.HAlign(HAlign_Center)
 			.VAlign(VAlign_Center)
+			.Padding(Settings->GetNonPinNodeBodyPadding())
 			[
 				// MIDDLE
 				SNew(STextBlock)
diff --git a/Engine/Source/Editor/GraphEditor/Private/KismetNodes/SGraphNodeK2Sequence.cpp b/Engine/Source/Editor/GraphEditor/Private/KismetNodes/SGraphNodeK2Sequence.cpp
--- a/Engine/Source/Editor/GraphEditor/Private/KismetNodes/SGraphNodeK2Sequence.cpp
+++ b/Engine/Source/Editor/GraphEditor/Private/KismetNodes/SGraphNodeK2Sequence.cpp
@@ -20,10 +20,13 @@ void SGraphNodeK2Sequence::CreateOutputSideAddButton(TSharedPtr<SVerticalBox> Ou
 		NSLOCTEXT("SequencerNode", "SequencerNodeAddPinButton", "Add pin"),
 		NSLOCTEXT("SequencerNode", "SequencerNodeAddPinButton_ToolTip", "Add new pin"));
 
+	FMargin AddPinPadding = Settings->GetOutputPinPadding();
+	AddPinPadding.Top += 6.0f;
+
 	OutputBox->AddSlot()
 	.AutoHeight()
 	.VAlign(VAlign_Center)
-	.Padding(10,10,10,4)
+	.Padding(AddPinPadding)
 	[
 		AddPinButton
 	];
diff --git a/Engine/Source/Editor/GraphEditor/Private/KismetNodes/SGraphNodeSwitchStatement.cpp b/Engine/Source/Editor/GraphEditor/Private/KismetNodes/SGraphNodeSwitchStatement.cpp
--- a/Engine/Source/Editor/GraphEditor/Private/KismetNodes/SGraphNodeSwitchStatement.cpp
+++ b/Engine/Source/Editor/GraphEditor/Private/KismetNodes/SGraphNodeSwitchStatement.cpp
@@ -83,10 +83,13 @@ void SGraphNodeSwitchStatement::CreateOutputSideAddButton(TSharedPtr<SVerticalBo
 		NSLOCTEXT("SwitchStatementNode", "SwitchStatementNodeAddPinButton", "Add pin"),
 		NSLOCTEXT("SwitchStatementNode", "SwitchStatementNodeAddPinButton_Tooltip", "Add new pin"));
 
+	FMargin AddPinPadding = Settings->GetOutputPinPadding();
+	AddPinPadding.Top += 6.0f;
+
 	OutputBox->AddSlot()
 	.AutoHeight()
 	.VAlign(VAlign_Center)
-	.Padding(10,10,10,4)
+	.Padding(AddPinPadding)
 	[
 		AddPinButton
 	];
diff --git a/Engine/Source/Editor/GraphEditor/Private/MaterialNodes/SGraphNodeMaterialBase.cpp b/Engine/Source/Editor/GraphEditor/Private/MaterialNodes/SGraphNodeMaterialBase.cpp
--- a/Engine/Source/Editor/GraphEditor/Private/MaterialNodes/SGraphNodeMaterialBase.cpp
+++ b/Engine/Source/Editor/GraphEditor/Private/MaterialNodes/SGraphNodeMaterialBase.cpp
@@ -265,11 +265,15 @@ void SGraphNodeMaterialBase::AddPin( const TSharedRef<SGraphPin>& PinToAdd )
 
 	if (PinToAdd->GetDirection() == EEdGraphPinDirection::EGPD_Input)
 	{
+		FMargin Padding = Settings->GetInputPinPadding();
+		Padding.Left *= 0.5f;
+		Padding.Right = 0.0f;
+
 		LeftNodeBox->AddSlot()
 		.AutoHeight()
 		.HAlign(HAlign_Left)
 		.VAlign(VAlign_Center)
-		.Padding(5,4,0,4)
+		.Padding(Padding)
 		[
 			PinToAdd
 		];
@@ -277,11 +281,15 @@ void SGraphNodeMaterialBase::AddPin( const TSharedRef<SGraphPin>& PinToAdd )
 	}
 	else // Direction == EEdGraphPinDirection::EGPD_Output
 	{
+		FMargin Padding = Settings->GetOutputPinPadding();
+		Padding.Left = 0.0f;
+		Padding.Right *= 0.5f;
+
 		RightNodeBox->AddSlot()
 		.AutoHeight()
 		.HAlign(HAlign_Right)
 		.VAlign(VAlign_Center)
-		.Padding(0,4,5,4)
+		.Padding(Padding)
 		[
 			PinToAdd
 		];
@@ -296,10 +304,14 @@ void SGraphNodeMaterialBase::CreateBelowPinControls(TSharedPtr<SVerticalBox> Mai
 		int32 LeftPinCount = InputPins.Num();
 		int32 RightPinCount = OutputPins.Num();
 
+		const float NegativeHPad = FMath::Max<float>(-Settings->PaddingTowardsNodeEdge, 0.0f);
+		const float ExtraPad = 0.0f;
+
 		// Place preview widget based on where the least pins are
-		if (LeftPinCount < RightPinCount || RightPinCount == 0)
+		if ((LeftPinCount < RightPinCount) || (RightPinCount == 0))
 		{
 			LeftNodeBox->AddSlot()
+			.Padding(FMargin(NegativeHPad + ExtraPad, 0.0f, 0.0f, 0.0f))
 			.AutoHeight()
 			[
 				CreatePreviewWidget()
@@ -308,6 +320,7 @@ void SGraphNodeMaterialBase::CreateBelowPinControls(TSharedPtr<SVerticalBox> Mai
 		else if (LeftPinCount > RightPinCount)
 		{
 			RightNodeBox->AddSlot()
+			.Padding(FMargin(NegativeHPad + ExtraPad, 0.0f, 0.0f, 0.0f))
 			.AutoHeight()
 			[
 				CreatePreviewWidget()
@@ -316,6 +329,7 @@ void SGraphNodeMaterialBase::CreateBelowPinControls(TSharedPtr<SVerticalBox> Mai
 		else
 		{
 			MainBox->AddSlot()
+			.Padding(Settings->GetNonPinNodeBodyPadding())
 			.AutoHeight()
 			[
 				SNew(SHorizontalBox)
diff --git a/Engine/Source/Editor/GraphEditor/Private/SGraphNode.cpp b/Engine/Source/Editor/GraphEditor/Private/SGraphNode.cpp
--- a/Engine/Source/Editor/GraphEditor/Private/SGraphNode.cpp
+++ b/Engine/Source/Editor/GraphEditor/Private/SGraphNode.cpp
@@ -733,6 +733,7 @@ void SGraphNode::UpdateGraphNode()
 	
 	TSharedPtr<SVerticalBox> InnerVerticalBox;
 	this->ContentScale.Bind( this, &SGraphNode::GetContentScale );
+
 	this->ChildSlot
 		.HAlign(HAlign_Center)
 		.VAlign(VAlign_Center)
@@ -741,31 +742,39 @@ void SGraphNode::UpdateGraphNode()
 			+SVerticalBox::Slot()
 			.AutoHeight()
 			[
-				SNew(SBorder)
-				.BorderImage( FEditorStyle::GetBrush( "Graph.Node.Body" ) )
-				.Padding(0)
+				SNew(SOverlay)
 				.AddMetaData<FGraphNodeMetaData>(TagMeta)
+				+SOverlay::Slot()
+				.Padding(Settings->GetNonPinNodeBodyPadding())
+				[
+					SNew(SImage)
+					.Image(FEditorStyle::GetBrush("Graph.Node.Body"))
+				]
+				+SOverlay::Slot()
 				[
 					SAssignNew(InnerVerticalBox, SVerticalBox)
 					+SVerticalBox::Slot()
 					.AutoHeight()
 					.HAlign(HAlign_Fill)
 					.VAlign(VAlign_Top)
+					.Padding(Settings->GetNonPinNodeBodyPadding())
 					[
 						TitleAreaWidget
 					]
+
 					+SVerticalBox::Slot()
 					.AutoHeight()
-					.Padding(1.0f)
+					.HAlign(HAlign_Fill)
+					.VAlign(VAlign_Top)
 					[
-						ErrorText->AsShared()
+						CreateNodeContentArea()
 					]
+
 					+SVerticalBox::Slot()
 					.AutoHeight()
-					.HAlign(HAlign_Fill)
-					.VAlign(VAlign_Top)
+					.Padding(Settings->GetNonPinNodeBodyPadding())
 					[
-						CreateNodeContentArea()
+						ErrorText->AsShared()
 					]
 				]
 			]			
@@ -927,19 +936,19 @@ void SGraphNode::AddPin(const TSharedRef<SGraphPin>& PinToAdd)
 	PinToAdd->SetOwner(SharedThis(this));
 
 	const UEdGraphPin* PinObj = PinToAdd->GetPinObj();
-	const bool bAdvancedParameter = PinObj && PinObj->bAdvancedView;
-	if(bAdvancedParameter)
+	const bool bAdvancedParameter = (PinObj != nullptr) && PinObj->bAdvancedView;
+	if (bAdvancedParameter)
 	{
 		PinToAdd->SetVisibility( TAttribute<EVisibility>(PinToAdd, &SGraphPin::IsPinVisibleAsAdvanced) );
 	}
-	
+
 	if (PinToAdd->GetDirection() == EEdGraphPinDirection::EGPD_Input)
 	{
 		LeftNodeBox->AddSlot()
 			.AutoHeight()
 			.HAlign(HAlign_Left)
 			.VAlign(VAlign_Center)
-			.Padding(10,4)
+			.Padding(Settings->GetInputPinPadding())
 		[
 			PinToAdd
 		];
@@ -951,7 +960,7 @@ void SGraphNode::AddPin(const TSharedRef<SGraphPin>& PinToAdd)
 			.AutoHeight()
 			.HAlign(HAlign_Right)
 			.VAlign(VAlign_Center)
-			.Padding(10,4)
+			.Padding(Settings->GetOutputPinPadding())
 		[
 			PinToAdd
 		];
@@ -1065,6 +1074,7 @@ SGraphNode::SGraphNode()
 	, bRenameIsPending( false )
 	, ErrorColor( FLinearColor::White )
 	, CachedUnscaledPosition( FVector2D::ZeroVector )
+	, Settings( GetDefault<UGraphEditorSettings>() )
 {
 	// Set up animation
 	{
diff --git a/Engine/Source/Editor/GraphEditor/Private/SGraphPanel.cpp b/Engine/Source/Editor/GraphEditor/Private/SGraphPanel.cpp
--- a/Engine/Source/Editor/GraphEditor/Private/SGraphPanel.cpp
+++ b/Engine/Source/Editor/GraphEditor/Private/SGraphPanel.cpp
@@ -141,7 +141,7 @@ int32 SGraphPanel::OnPaint( const FPaintArgs& Args, const FGeometry& AllottedGeo
 	const int32 NodeLayerId = NodeShadowsLayerId + 1;
 	int32 MaxLayerId = NodeLayerId;
 
-	const FVector2D NodeShadowSize = FEditorStyle::GetVector(TEXT("Graph.Node.ShadowSize"));
+	const FVector2D NodeShadowSize = GetDefault<UGraphEditorSettings>()->GetShadowDeltaSize();
 	const UEdGraphSchema* Schema = GraphObj->GetSchema();
 
 	// Draw the child nodes
diff --git a/Engine/Source/Editor/GraphEditor/Private/SoundNodes/SGraphNodeSoundBase.cpp b/Engine/Source/Editor/GraphEditor/Private/SoundNodes/SGraphNodeSoundBase.cpp
--- a/Engine/Source/Editor/GraphEditor/Private/SoundNodes/SGraphNodeSoundBase.cpp
+++ b/Engine/Source/Editor/GraphEditor/Private/SoundNodes/SGraphNodeSoundBase.cpp
@@ -26,10 +26,13 @@ void SGraphNodeSoundBase::CreateOutputSideAddButton(TSharedPtr<SVerticalBox> Out
 		NSLOCTEXT("SoundNode", "SoundNodeAddPinButton_Tooltip", "Adds an input to the sound node")
 	);
 
+	FMargin AddPinPadding = Settings->GetOutputPinPadding();
+	AddPinPadding.Top += 6.0f;
+
 	OutputBox->AddSlot()
 	.AutoHeight()
 	.VAlign(VAlign_Center)
-	.Padding(10,10,10,4)
+	.Padding(AddPinPadding)
 	[
 		AddPinButton
 	];
diff --git a/Engine/Source/Editor/GraphEditor/Public/GraphEditorSettings.h b/Engine/Source/Editor/GraphEditor/Public/GraphEditorSettings.h
--- a/Engine/Source/Editor/GraphEditor/Public/GraphEditorSettings.h
+++ b/Engine/Source/Editor/GraphEditor/Public/GraphEditorSettings.h
@@ -18,12 +18,52 @@ class GRAPHEDITOR_API UGraphEditorSettings
 {
 	GENERATED_UCLASS_BODY()
 
+public:
+
+#if WITH_EDITOR
+	// UObject interface
+	virtual void PostEditChangeProperty(FPropertyChangedEvent& PropertyChangedEvent) override;
+	// End of UObject interface
+#endif
+
 public:
 
 	/** The visual styling to use for graph editor pins (in Blueprints, materials, etc...) */
 	UPROPERTY(config, EditAnywhere, Category=GeneralStyle)
 	TEnumAsByte<EBlueprintPinStyleType> DataPinStyle;
 
+	/** The amount of padding above a pin (defaults to 4) */
+	UPROPERTY(/*config, */ EditAnywhere, Category=GeneralStyle, AdvancedDisplay)
+	float PaddingAbovePin;
+
+	/** The amount of padding below a pin (defaults to 4) */
+	UPROPERTY(/*config, */ EditAnywhere, Category=GeneralStyle, AdvancedDisplay)
+	float PaddingBelowPin;
+
+	/** The amount of horizontal padding towards the center of a node on an input pin (defaults to 10) */
+	UPROPERTY(/*config, */ EditAnywhere, Category=GeneralStyle, AdvancedDisplay)
+	float PaddingRightOfInput;
+
+	/** The amount of horizontal padding towards the center of a node on an output pin (defaults to 10) */
+	UPROPERTY(/*config, */ EditAnywhere, Category=GeneralStyle, AdvancedDisplay)
+	float PaddingLeftOfOutput;
+
+	/** The amount of padding towards the node edge (defaults to 10, can be negative to make pins overlap or stick out of a node */
+	UPROPERTY(/*config, */ EditAnywhere, Category=GeneralStyle, AdvancedDisplay)
+	float PaddingTowardsNodeEdge;
+
+	/** Returns the padding to use for an input pin on a graph node */
+	FMargin GetInputPinPadding() const;
+
+	/** Returns the padding to use for an output pin on a graph node */
+	FMargin GetOutputPinPadding() const;
+
+	/** Returns the padding to use for a node region that won't contain a pin */
+	FMargin GetNonPinNodeBodyPadding() const;
+
+	/** Returns the size adjustment from node render geometry to shadow render geometry */
+	FVector2D GetShadowDeltaSize() const;
+
 public:
 	/** The maximum value to clamp the absolute value of the horizontal distance between endpoints when calculating tangents (when the wire is moving forward) */
 	UPROPERTY(config, EditAnywhere, Category=Splines, AdvancedDisplay)
diff --git a/Engine/Source/Editor/GraphEditor/Public/SGraphNode.h b/Engine/Source/Editor/GraphEditor/Public/SGraphNode.h
--- a/Engine/Source/Editor/GraphEditor/Public/SGraphNode.h
+++ b/Engine/Source/Editor/GraphEditor/Public/SGraphNode.h
@@ -329,6 +329,7 @@ protected:
 
 	/** Is this node editable */
 	TAttribute<bool> IsEditable;
+
 	FSingleNodeEvent OnDoubleClick;
 
 	// Is the current tooltip a complex one that should be dropped when the tooltip is no longer displayed?
@@ -349,4 +350,7 @@ protected:
 
 	/** Caches true position of node */
 	FVector2D CachedUnscaledPosition;
+
+	/** Cached pointer to graph editor settings */
+	const class UGraphEditorSettings* Settings;
 };
