GraphEditor: Fix some node titles (such as material nodes) having extra blank lines due to \\r injected by the text builder class

diff --git a/Engine/Source/Editor/GraphEditor/Private/SGraphNode.cpp b/Engine/Source/Editor/GraphEditor/Private/SGraphNode.cpp
--- a/Engine/Source/Editor/GraphEditor/Private/SGraphNode.cpp
+++ b/Engine/Source/Editor/GraphEditor/Private/SGraphNode.cpp
@@ -68,7 +68,8 @@ void SNodeTitle::RebuildWidget()
 
 	// Break the title into lines
 	TArray<FString> Lines;
-	CachedTitle.ToString().ParseIntoArray(&Lines, TEXT("\n"), false);
+	const FString CachedTitleString = CachedTitle.ToString().Replace(TEXT("\r"), TEXT(""));
+	CachedTitleString.ParseIntoArray(&Lines, TEXT("\n"), false);
 
 	if (Lines.Num())
 	{
