Switch Node Polish
* Name and String have + Add Pin button
* Enum doesn't display Remove Execution Pin on right click
* Enum doesn't display Show Default Pin in details view
* All types have proper default in selection box

diff --git a/Engine/Source/Editor/BlueprintGraph/Classes/K2Node.h b/Engine/Source/Editor/BlueprintGraph/Classes/K2Node.h
--- a/Engine/Source/Editor/BlueprintGraph/Classes/K2Node.h
+++ b/Engine/Source/Editor/BlueprintGraph/Classes/K2Node.h
@@ -192,6 +192,9 @@ class UK2Node : public UEdGraphNode
 	/** Return whether the node's properties display in the blueprint details panel */
 	virtual bool ShouldShowNodeProperties() const { return false; }
 
+	/** Return whether the node's execution pins should support the remove execution pin action */
+	virtual bool CanEverRemoveExecutionPin() const { return false; }
+
 	/** Called when the connection list of one of the pins of this node is changed in the editor, after the pin has had it's literal cleared */
 	virtual void NotifyPinConnectionListChanged(UEdGraphPin* Pin) {}
 
diff --git a/Engine/Source/Editor/BlueprintGraph/Classes/K2Node_ExecutionSequence.h b/Engine/Source/Editor/BlueprintGraph/Classes/K2Node_ExecutionSequence.h
--- a/Engine/Source/Editor/BlueprintGraph/Classes/K2Node_ExecutionSequence.h
+++ b/Engine/Source/Editor/BlueprintGraph/Classes/K2Node_ExecutionSequence.h
@@ -23,6 +23,7 @@ class UK2Node_ExecutionSequence : public UK2Node
 	virtual class FNodeHandlingFunctor* CreateNodeHandler(class FKismetCompilerContext& CompilerContext) const override;
 	virtual void GetMenuActions(FBlueprintActionDatabaseRegistrar& ActionRegistrar) const override;
 	virtual FText GetMenuCategory() const override;
+	virtual bool CanEverRemoveExecutionPin() const override { return true; }
 	// End UK2Node interface
 
 	// K2Node_ExecutionSequence interface
diff --git a/Engine/Source/Editor/BlueprintGraph/Classes/K2Node_Switch.h b/Engine/Source/Editor/BlueprintGraph/Classes/K2Node_Switch.h
--- a/Engine/Source/Editor/BlueprintGraph/Classes/K2Node_Switch.h
+++ b/Engine/Source/Editor/BlueprintGraph/Classes/K2Node_Switch.h
@@ -36,6 +36,7 @@ class UK2Node_Switch : public UK2Node
 	virtual ERedirectType DoPinsMatchForReconstruction(const UEdGraphPin* NewPin, int32 NewPinIndex, const UEdGraphPin* OldPin, int32 OldPinIndex) const override;
 	virtual class FNodeHandlingFunctor* CreateNodeHandler(class FKismetCompilerContext& CompilerContext) const override;
 	virtual FText GetMenuCategory() const override;
+	virtual bool CanEverRemoveExecutionPin() const override { return true; }
 	// End of UK2Node interface
 
 	// UK2Node_Switch interface
diff --git a/Engine/Source/Editor/BlueprintGraph/Classes/K2Node_SwitchEnum.h b/Engine/Source/Editor/BlueprintGraph/Classes/K2Node_SwitchEnum.h
--- a/Engine/Source/Editor/BlueprintGraph/Classes/K2Node_SwitchEnum.h
+++ b/Engine/Source/Editor/BlueprintGraph/Classes/K2Node_SwitchEnum.h
@@ -32,12 +32,12 @@ class UK2Node_SwitchEnum : public UK2Node_Switch, public INodeDependingOnEnumInt
 	// UEdGraphNode interface
 	virtual FText GetTooltipText() const override;
 	virtual FText GetNodeTitle(ENodeTitleType::Type TitleType) const override;
-	virtual bool ShouldShowNodeProperties() const override { return true; }	
 	// End of UEdGraphNode interface
 
 	// UK2Node interface
 	virtual bool IsConnectionDisallowed(const UEdGraphPin* MyPin, const UEdGraphPin* OtherPin, FString& OutReason) const override;
 	virtual void GetMenuActions(FBlueprintActionDatabaseRegistrar& ActionRegistrar) const override;
+	virtual bool CanEverRemoveExecutionPin() const override { return false; }
 	// End of UK2Node interface
 
 	// UK2Node_Switch Interface
diff --git a/Engine/Source/Editor/BlueprintGraph/Private/EdGraphSchema_K2.cpp b/Engine/Source/Editor/BlueprintGraph/Private/EdGraphSchema_K2.cpp
--- a/Engine/Source/Editor/BlueprintGraph/Private/EdGraphSchema_K2.cpp
+++ b/Engine/Source/Editor/BlueprintGraph/Private/EdGraphSchema_K2.cpp
@@ -948,7 +948,7 @@ void UEdGraphSchema_K2::GetContextMenuActions(const UEdGraph* CurrentGraph, cons
 				// Conditionally add the execution path pin removal if this is an execution branching node
 				if( InGraphPin->Direction == EGPD_Output && InGraphPin->GetOwningNode())
 				{
-					if ( InGraphPin->GetOwningNode()->IsA(UK2Node_ExecutionSequence::StaticClass()) ||  InGraphPin->GetOwningNode()->IsA(UK2Node_Switch::StaticClass()) )
+					if (CastChecked<UK2Node>(InGraphPin->GetOwningNode())->CanEverRemoveExecutionPin())
 					{
 						MenuBuilder->AddMenuEntry( FGraphEditorCommands::Get().RemoveExecutionPin );
 					}
diff --git a/Engine/Source/Editor/BlueprintGraph/Private/K2Node_Switch.cpp b/Engine/Source/Editor/BlueprintGraph/Private/K2Node_Switch.cpp
--- a/Engine/Source/Editor/BlueprintGraph/Private/K2Node_Switch.cpp
+++ b/Engine/Source/Editor/BlueprintGraph/Private/K2Node_Switch.cpp
@@ -153,16 +153,8 @@ void UK2Node_Switch::PostEditChangeProperty(struct FPropertyChangedEvent& Proper
 	FName PropertyName = (PropertyChangedEvent.Property != NULL) ? PropertyChangedEvent.Property->GetFName() : NAME_None;
 	if (PropertyName == TEXT("bHasDefaultPin"))
 	{
-		// Disallow enabling the default pin on enums
-		if (IsA(UK2Node_SwitchEnum::StaticClass()))
-		{
-			bHasDefaultPin = false;
-		}
-		else
-		{
-			// Signal to the reconstruction logic that the default pin value has changed
-			bHasDefaultPinValueChanged = true;
-		}
+		// Signal to the reconstruction logic that the default pin value has changed
+		bHasDefaultPinValueChanged = true;
 		
 		if (!bHasDefaultPin)
 		{
diff --git a/Engine/Source/Editor/BlueprintGraph/Private/K2Node_SwitchInteger.cpp b/Engine/Source/Editor/BlueprintGraph/Private/K2Node_SwitchInteger.cpp
--- a/Engine/Source/Editor/BlueprintGraph/Private/K2Node_SwitchInteger.cpp
+++ b/Engine/Source/Editor/BlueprintGraph/Private/K2Node_SwitchInteger.cpp
@@ -47,7 +47,8 @@ void UK2Node_SwitchInteger::GetMenuActions(FBlueprintActionDatabaseRegistrar& Ac
 void UK2Node_SwitchInteger::CreateSelectionPin()
 {
 	const UEdGraphSchema_K2* K2Schema = GetDefault<UEdGraphSchema_K2>();
-	CreatePin(EGPD_Input, K2Schema->PC_Int, TEXT(""), NULL, false, false, TEXT("Selection"));
+	UEdGraphPin* Pin = CreatePin(EGPD_Input, K2Schema->PC_Int, TEXT(""), NULL, false, false, TEXT("Selection"));
+	K2Schema->SetPinDefaultValueBasedOnType(Pin);
 }
 
 FString UK2Node_SwitchInteger::GetPinNameGivenIndex(int32 Index)
diff --git a/Engine/Source/Editor/BlueprintGraph/Private/K2Node_SwitchName.cpp b/Engine/Source/Editor/BlueprintGraph/Private/K2Node_SwitchName.cpp
--- a/Engine/Source/Editor/BlueprintGraph/Private/K2Node_SwitchName.cpp
+++ b/Engine/Source/Editor/BlueprintGraph/Private/K2Node_SwitchName.cpp
@@ -62,7 +62,8 @@ void UK2Node_SwitchName::GetMenuActions(FBlueprintActionDatabaseRegistrar& Actio
 void UK2Node_SwitchName::CreateSelectionPin()
 {
 	const UEdGraphSchema_K2* K2Schema = GetDefault<UEdGraphSchema_K2>();
-	CreatePin(EGPD_Input, K2Schema->PC_Name, TEXT(""), NULL, false, false, TEXT("Selection"));
+	UEdGraphPin* Pin = CreatePin(EGPD_Input, K2Schema->PC_Name, TEXT(""), NULL, false, false, TEXT("Selection"));
+	K2Schema->SetPinDefaultValueBasedOnType(Pin);
 }
 
 FString UK2Node_SwitchName::GetPinNameGivenIndex(int32 Index)
diff --git a/Engine/Source/Editor/BlueprintGraph/Private/K2Node_SwitchString.cpp b/Engine/Source/Editor/BlueprintGraph/Private/K2Node_SwitchString.cpp
--- a/Engine/Source/Editor/BlueprintGraph/Private/K2Node_SwitchString.cpp
+++ b/Engine/Source/Editor/BlueprintGraph/Private/K2Node_SwitchString.cpp
@@ -74,7 +74,8 @@ void UK2Node_SwitchString::GetMenuActions(FBlueprintActionDatabaseRegistrar& Act
 void UK2Node_SwitchString::CreateSelectionPin()
 {
 	const UEdGraphSchema_K2* K2Schema = GetDefault<UEdGraphSchema_K2>();
-	CreatePin(EGPD_Input, K2Schema->PC_String, TEXT(""), NULL, false, false, TEXT("Selection"));
+	UEdGraphPin* Pin = CreatePin(EGPD_Input, K2Schema->PC_String, TEXT(""), NULL, false, false, TEXT("Selection"));
+	K2Schema->SetPinDefaultValueBasedOnType(Pin);
 }
 
 FString UK2Node_SwitchString::GetPinNameGivenIndex(int32 Index)
diff --git a/Engine/Source/Editor/GraphEditor/Private/KismetNodes/SGraphNodeSwitchStatement.cpp b/Engine/Source/Editor/GraphEditor/Private/KismetNodes/SGraphNodeSwitchStatement.cpp
--- a/Engine/Source/Editor/GraphEditor/Private/KismetNodes/SGraphNodeSwitchStatement.cpp
+++ b/Engine/Source/Editor/GraphEditor/Private/KismetNodes/SGraphNodeSwitchStatement.cpp
@@ -5,6 +5,7 @@
 #include "SGraphNodeSwitchStatement.h"
 #include "KismetPins/SGraphPinExec.h"
 #include "NodeFactory.h"
+#include "K2Node_SwitchEnum.h"
 
 #include "ScopedTransaction.h"
 #include "BlueprintEditorUtils.h"
@@ -97,7 +98,7 @@ void SGraphNodeSwitchStatement::CreateOutputSideAddButton(TSharedPtr<SVerticalBo
 
 EVisibility SGraphNodeSwitchStatement::IsAddPinButtonVisible() const
 {
-	if(NULL != Cast<UK2Node_SwitchInteger>(GraphNode))
+	if(!GraphNode->IsA<UK2Node_SwitchEnum>())
 	{
 		return SGraphNode::IsAddPinButtonVisible();
 	}
@@ -106,16 +107,16 @@ EVisibility SGraphNodeSwitchStatement::IsAddPinButtonVisible() const
 
 FReply SGraphNodeSwitchStatement::OnAddPin()
 {
-	if(UK2Node_SwitchInteger* SwitchNode = Cast<UK2Node_SwitchInteger>(GraphNode))
-	{
-		const FScopedTransaction Transaction( NSLOCTEXT("Kismet", "AddExecutionPin", "Add Execution Pin") );
-		SwitchNode->Modify();
+	UK2Node_Switch* SwitchNode = CastChecked<UK2Node_Switch>(GraphNode);
 
-		SwitchNode->AddPinToSwitchNode();
-		FBlueprintEditorUtils::MarkBlueprintAsModified(SwitchNode->GetBlueprint());
+	const FScopedTransaction Transaction( NSLOCTEXT("Kismet", "AddExecutionPin", "Add Execution Pin") );
+	SwitchNode->Modify();
+
+	SwitchNode->AddPinToSwitchNode();
+	FBlueprintEditorUtils::MarkBlueprintAsModified(SwitchNode->GetBlueprint());
+
+	UpdateGraphNode();
+	GraphNode->GetGraph()->NotifyGraphChanged();
 
-		UpdateGraphNode();
-		GraphNode->GetGraph()->NotifyGraphChanged();
-	}
 	return FReply::Handled();
 }
\ No newline at end of file
