Fix non-unity in UGraphEditorSettings (missing FMargin decl)

diff --git a/Engine/Source/Editor/GraphEditor/Public/GraphEditorSettings.h b/Engine/Source/Editor/GraphEditor/Public/GraphEditorSettings.h
--- a/Engine/Source/Editor/GraphEditor/Public/GraphEditorSettings.h
+++ b/Engine/Source/Editor/GraphEditor/Public/GraphEditorSettings.h
@@ -5,7 +5,7 @@
 =============================================================================*/
 
 #pragma once
-
+#include "Runtime/SlateCore/Public/Layout/Margin.h"
 #include "GraphEditorSettings.generated.h"
 
 
