Hiding the empty graph instruction box when there are no instructions.

diff --git a/Engine/Source/Editor/GraphEditor/Private/SGraphEditorImpl.cpp b/Engine/Source/Editor/GraphEditor/Private/SGraphEditorImpl.cpp
--- a/Engine/Source/Editor/GraphEditor/Private/SGraphEditorImpl.cpp
+++ b/Engine/Source/Editor/GraphEditor/Private/SGraphEditorImpl.cpp
@@ -667,7 +667,7 @@ FString SGraphEditorImpl::GetInstructionText() const
 
 EVisibility SGraphEditorImpl::InstructionTextVisibility() const
 {
-	if (GetInstructionTextFade() > 0.0f)
+	if (!GetInstructionText().IsEmpty() && (GetInstructionTextFade() > 0.0f))
 	{
 		return EVisibility::Visible;
 	}
