ttp346705 CRITICAL: Regression: BP: Interacting with the details panel within a blueprint interface collapses the details panel and de..

diff --git a/Engine/Source/Editor/GraphEditor/Private/SGraphActionMenu.cpp b/Engine/Source/Editor/GraphEditor/Private/SGraphActionMenu.cpp
--- a/Engine/Source/Editor/GraphEditor/Private/SGraphActionMenu.cpp
+++ b/Engine/Source/Editor/GraphEditor/Private/SGraphActionMenu.cpp
@@ -373,7 +373,8 @@ bool SGraphActionMenu::SelectItemByName(const FName& ItemName, ESelectInfo::Type
 							break;
 						}
 					}
-					else if (CurrentGraphNode->GetDisplayName().ToString() == FName::NameToDisplayString(ItemName.ToString(), false))
+					
+					if (CurrentGraphNode->GetDisplayName().ToString() == FName::NameToDisplayString(ItemName.ToString(), false))
 					{
 						SelectionNode = CurrentGraphNode;
 
