Modified the Composite node tooltips to be static when the documentation excerpts are shown ( ctrl + alt )
#TTP 344651 - Rich Tool Tip on placed Math Expression Node moves so you can't click the link
#Branch UE4
#Proj  GraphEditor

diff --git a/Engine/Source/Editor/GraphEditor/Private/KismetNodes/SGraphNodeK2Composite.cpp b/Engine/Source/Editor/GraphEditor/Private/KismetNodes/SGraphNodeK2Composite.cpp
--- a/Engine/Source/Editor/GraphEditor/Private/KismetNodes/SGraphNodeK2Composite.cpp
+++ b/Engine/Source/Editor/GraphEditor/Private/KismetNodes/SGraphNodeK2Composite.cpp
@@ -132,29 +132,39 @@ TSharedPtr<SToolTip> SGraphNodeK2Composite::GetComplexTooltip()
 {
 	if (UEdGraph* BoundGraph = GetInnerGraph())
 	{
+		struct LocalUtils
+		{
+			static bool IsInteractive()
+			{
+				const FModifierKeysState ModifierKeys = FSlateApplication::Get().GetModifierKeys();
+				return ( ModifierKeys.IsAltDown() && ModifierKeys.IsControlDown() );
+			}
+		};
+
 		TSharedPtr<SToolTip> FinalToolTip = NULL;
 		TSharedPtr<SVerticalBox> Container = NULL;
 		SAssignNew(FinalToolTip, SToolTip)
+		.IsInteractive_Static(&LocalUtils::IsInteractive)
+		[
+			SAssignNew(Container, SVerticalBox)
+			+SVerticalBox::Slot()
+			.AutoHeight()
 			[
-				SAssignNew(Container, SVerticalBox)
-				+SVerticalBox::Slot()
-				.AutoHeight()
-				[
-					SNew( STextBlock )
-					.Text(this, &SGraphNodeK2Composite::GetTooltipTextForNode)
-					.Font(FSlateFontInfo( FPaths::EngineContentDir() / TEXT("Slate/Fonts/Roboto-Regular.ttf"), 8))
-					.WrapTextAt(160.0f)
-				]
-				+SVerticalBox::Slot()
-				.AutoHeight()
-				[
-					// Create preview for the tooltip, make sure to disable state overlays to prevent
-					// PIE and read-only borders obscuring the graph
-					SNew(SGraphPreviewer, BoundGraph)
-					.CornerOverlayText(this, &SGraphNodeK2Composite::GetPreviewCornerText)
-					.ShowGraphStateOverlay(false)
-				]
-			];
+				SNew( STextBlock )
+				.Text(this, &SGraphNodeK2Composite::GetTooltipTextForNode)
+				.Font(FSlateFontInfo( FPaths::EngineContentDir() / TEXT("Slate/Fonts/Roboto-Regular.ttf"), 8))
+				.WrapTextAt(160.0f)
+			]
+			+SVerticalBox::Slot()
+			.AutoHeight()
+			[
+				// Create preview for the tooltip, make sure to disable state overlays to prevent
+				// PIE and read-only borders obscuring the graph
+				SNew(SGraphPreviewer, BoundGraph)
+				.CornerOverlayText(this, &SGraphNodeK2Composite::GetPreviewCornerText)
+				.ShowGraphStateOverlay(false)
+			]
+		];
 
 		// Check to see whether this node has a documentation excerpt. If it does, create a doc box for the tooltip
 		TSharedRef<IDocumentationPage> DocPage = IDocumentation::Get()->GetPage(GraphNode->GetDocumentationLink(), NULL);
