Layered blend per bone starts with 2 poses by default and has an add pin button

#TTP 308411 - Animation Blueprint: Layered blend per bone should have 2 blends by default, and should have ADD PIN button like regular blueprints

diff --git a/Engine/Source/Editor/GraphEditor/Private/AnimationNodes/SGraphNodeLayeredBoneBlend.cpp b/Engine/Source/Editor/GraphEditor/Private/AnimationNodes/SGraphNodeLayeredBoneBlend.cpp
new file mode 100644
--- /dev/null
+++ b/Engine/Source/Editor/GraphEditor/Private/AnimationNodes/SGraphNodeLayeredBoneBlend.cpp
@@ -0,0 +1,43 @@
+// Copyright 1998-2014 Epic Games, Inc. All Rights Reserved.
+
+
+#include "GraphEditorCommon.h"
+#include "SGraphNodeLayeredBoneBlend.h"
+
+/////////////////////////////////////////////////////
+// SGraphNodeLayeredBoneBlend
+
+void SGraphNodeLayeredBoneBlend::Construct(const FArguments& InArgs, UAnimGraphNode_LayeredBoneBlend* InNode)
+{
+	this->GraphNode = Node = InNode;
+
+	this->SetCursor(EMouseCursor::CardinalCross);
+
+	this->UpdateGraphNode();
+}
+
+void SGraphNodeLayeredBoneBlend::CreateInputSideAddButton(TSharedPtr<SVerticalBox> InputBox)
+{
+	TSharedRef<SWidget> AddPinButton = AddPinButtonContent(
+		NSLOCTEXT("LayeredBoneBlendNode", "LayeredBoneBlendNodeAddPinButton", "Add pin"),
+		NSLOCTEXT("LayeredBoneBlendNode", "LayeredBoneBlendNodeAddPinButton_Tooltip", "Adds a input pose to the node"),
+		false);
+
+	FMargin AddPinPadding = Settings->GetInputPinPadding();
+	AddPinPadding.Top += 6.0f;
+
+	InputBox->AddSlot()
+		.AutoHeight()
+		.VAlign(VAlign_Center)
+		.Padding(AddPinPadding)
+		[
+			AddPinButton
+		];
+}
+
+FReply SGraphNodeLayeredBoneBlend::OnAddPin()
+{
+	Node->AddPinToBlendByFilter();
+
+	return FReply::Handled();
+}
\ No newline at end of file
diff --git a/Engine/Source/Editor/GraphEditor/Private/AnimationNodes/SGraphNodeLayeredBoneBlend.h b/Engine/Source/Editor/GraphEditor/Private/AnimationNodes/SGraphNodeLayeredBoneBlend.h
new file mode 100644
--- /dev/null
+++ b/Engine/Source/Editor/GraphEditor/Private/AnimationNodes/SGraphNodeLayeredBoneBlend.h
@@ -0,0 +1,26 @@
+// Copyright 1998-2014 Epic Games, Inc. All Rights Reserved.
+
+#pragma once
+
+class UAnimGraphNode_LayeredBoneBlend;
+
+class SGraphNodeLayeredBoneBlend : public SGraphNode
+{
+public:
+	SLATE_BEGIN_ARGS(SGraphNodeLayeredBoneBlend){}
+	SLATE_END_ARGS()
+
+	void Construct(const FArguments& InArgs, UAnimGraphNode_LayeredBoneBlend* InNode);
+
+protected:
+	// SGraphNode interface
+	virtual void CreateInputSideAddButton(TSharedPtr<SVerticalBox> InputBox) override;
+	virtual FReply OnAddPin() override;
+	// End of SGraphNode interface
+
+private:
+
+	// The node that we represent
+	UAnimGraphNode_LayeredBoneBlend* Node;
+
+};
diff --git a/Engine/Source/Editor/GraphEditor/Private/NodeFactory.cpp b/Engine/Source/Editor/GraphEditor/Private/NodeFactory.cpp
--- a/Engine/Source/Editor/GraphEditor/Private/NodeFactory.cpp
+++ b/Engine/Source/Editor/GraphEditor/Private/NodeFactory.cpp
@@ -32,6 +32,7 @@
 #include "AnimGraphNode_Root.h"
 #include "AnimGraphNode_SequencePlayer.h"
 #include "AnimGraphNode_StateMachineBase.h"
+#include "AnimGraphNode_LayeredBoneBlend.h"
 #include "AnimStateNode.h"
 #include "AnimStateEntryNode.h"
 #include "AnimStateConduitNode.h"
@@ -46,6 +47,7 @@
 #include "AnimationNodes/SGraphNodeSequencePlayer.h"
 #include "AnimationNodes/SGraphNodeAnimationResult.h"
 #include "AnimationNodes/SGraphNodeStateMachineInstance.h"
+#include "AnimationNodes/SGraphNodeLayeredBoneBlend.h"
 #include "AnimationPins/SGraphPinPose.h"
 
 #include "KismetPins/SGraphPinBool.h"
@@ -107,6 +109,10 @@ TSharedPtr<SGraphNode> FNodeFactory::CreateNodeWidget(UEdGraphNode* InNode)
 		{
 			return SNew(SGraphNodeSequencePlayer, SequencePlayer);
 		}
+		else if (UAnimGraphNode_LayeredBoneBlend* LayeredBlend = Cast<UAnimGraphNode_LayeredBoneBlend>(InNode))
+		{
+			return SNew(SGraphNodeLayeredBoneBlend, LayeredBlend);
+		}
 	}
 
 	if (USoundCueGraphNode_Base* BaseSoundNode = Cast<USoundCueGraphNode_Base>(InNode))
diff --git a/Engine/Source/Runtime/Engine/Classes/Animation/AnimNode_LayeredBoneBlend.h b/Engine/Source/Runtime/Engine/Classes/Animation/AnimNode_LayeredBoneBlend.h
--- a/Engine/Source/Runtime/Engine/Classes/Animation/AnimNode_LayeredBoneBlend.h
+++ b/Engine/Source/Runtime/Engine/Classes/Animation/AnimNode_LayeredBoneBlend.h
@@ -38,6 +38,7 @@ protected:
 public:	
 	FAnimNode_LayeredBoneBlend()
 	{
+		AddPose();
 	}
 
 	// FAnimNode_Base interface
