[GitHub] 315 : Added an undo transaction for comment node resizing
#ttp 342708

diff --git a/Engine/Source/Editor/GraphEditor/Private/SGraphNodeResizable.cpp b/Engine/Source/Editor/GraphEditor/Private/SGraphNodeResizable.cpp
--- a/Engine/Source/Editor/GraphEditor/Private/SGraphNodeResizable.cpp
+++ b/Engine/Source/Editor/GraphEditor/Private/SGraphNodeResizable.cpp
@@ -3,6 +3,7 @@
 
 #include "GraphEditorCommon.h"
 #include "SGraphNodeResizable.h"
+#include "ScopedTransaction.h"
 
 namespace GraphNodeResizableDefs
 {
@@ -106,6 +107,10 @@ FReply SGraphNodeResizable::OnMouseButtonUp( const FGeometry& MyGeometry, const
 		UserSize.Y = FMath::RoundToFloat(UserSize.Y);
 
 		GetNodeObj()->ResizeNode(UserSize);
+
+		// End resize transaction
+		ResizeTransactionPtr.Reset();
+
 		return FReply::Handled().ReleaseMouseCapture();
 	}
 	return FReply::Unhandled();
@@ -191,6 +196,14 @@ FReply SGraphNodeResizable::OnMouseMove(const FGeometry& MyGeometry, const FPoin
 			GraphNode->ResizeNode( UserSize );
 			DeltaNodePos = GetCorrectedNodePosition() - GetPosition();
 		}
+
+		if (!ResizeTransactionPtr.IsValid() && UserSize != StoredUserSize)
+		{
+			// Start resize transaction.  The transaction is started here so all MoveTo actions are captured while empty
+			//	transactions are not created
+			ResizeTransactionPtr = MakeShareable(new FScopedTransaction(NSLOCTEXT("GraphEditor", "ResizeNodeAction", "Resize Node")));
+		}
+
 		SGraphNode::FNodeSet NodeFilter;
 		SGraphNode::MoveTo( GetPosition() + DeltaNodePos, NodeFilter );
 	}
diff --git a/Engine/Source/Editor/GraphEditor/Public/SGraphNodeResizable.h b/Engine/Source/Editor/GraphEditor/Public/SGraphNodeResizable.h
--- a/Engine/Source/Editor/GraphEditor/Public/SGraphNodeResizable.h
+++ b/Engine/Source/Editor/GraphEditor/Public/SGraphNodeResizable.h
@@ -73,6 +73,9 @@ protected:
 	/** The original size of the node while resizing */
 	FVector2D StoredUserSize;
 
+	/** The resize transaction */
+	TSharedPtr<FScopedTransaction> ResizeTransactionPtr;
+
 	/** Anchor point used to correct node position on resizing the node*/
 	FVector2D NodeAnchorPoint;
 
