Reducing the GuardBand on the node clipping because it seems wasteful and we have performance issues on big graphs.
#Branch UE4
#Proj GraphEditor

diff --git a/Engine/Source/Editor/GraphEditor/Private/SNodePanel.cpp b/Engine/Source/Editor/GraphEditor/Private/SNodePanel.cpp
--- a/Engine/Source/Editor/GraphEditor/Private/SNodePanel.cpp
+++ b/Engine/Source/Editor/GraphEditor/Private/SNodePanel.cpp
@@ -193,7 +193,7 @@ namespace NodePanelDefs
 	// Default Zoom Padding Value
 	static const float DefaultZoomPadding = 25.f;
 	// Node Culling Guardband Area
-	static const float GuardBandArea = 0.5f;
+	static const float GuardBandArea = 0.25f;
 	// Scaling factor to reduce speed of mouse zooming
 	static const float MouseZoomScaling = 0.05f;
 };
@@ -1092,15 +1092,22 @@ void SNodePanel::RemoveAllNodes()
 	VisibleChildren.Empty();
 }
 
+int32 VisibleNodes;
+int32 TotalNodes;
+
 void SNodePanel::PopulateVisibleChildren(const FGeometry& AllottedGeometry)
 {
 	VisibleChildren.Empty();
+	VisibleNodes = 0;
+	TotalNodes = 0;
 
 	for (int32 ChildIndex = 0; ChildIndex < Children.Num(); ++ChildIndex)
 	{
 		const TSharedRef<SNode>& SomeChild = Children[ChildIndex];
+		TotalNodes++;
 		if ( !IsNodeCulled(SomeChild, AllottedGeometry) )
 		{
+			VisibleNodes++;
 			VisibleChildren.Add(SomeChild);
 		}
 	}
