Removing debug spam committed in 2320435

diff --git a/Engine/Source/Editor/GraphEditor/Private/SNodePanel.cpp b/Engine/Source/Editor/GraphEditor/Private/SNodePanel.cpp
--- a/Engine/Source/Editor/GraphEditor/Private/SNodePanel.cpp
+++ b/Engine/Source/Editor/GraphEditor/Private/SNodePanel.cpp
@@ -1092,22 +1092,14 @@ void SNodePanel::RemoveAllNodes()
 	VisibleChildren.Empty();
 }
 
-int32 VisibleNodes;
-int32 TotalNodes;
-
 void SNodePanel::PopulateVisibleChildren(const FGeometry& AllottedGeometry)
 {
 	VisibleChildren.Empty();
-	VisibleNodes = 0;
-	TotalNodes = 0;
-
 	for (int32 ChildIndex = 0; ChildIndex < Children.Num(); ++ChildIndex)
 	{
 		const TSharedRef<SNode>& SomeChild = Children[ChildIndex];
-		TotalNodes++;
 		if ( !IsNodeCulled(SomeChild, AllottedGeometry) )
 		{
-			VisibleNodes++;
 			VisibleChildren.Add(SomeChild);
 		}
 	}
