This adds some simple clip testing for the connection bounds in an effort to increase graph performance on large blueprints
#TTP 347817 - Editing larger Blueprints is super slow.
#Branch UE4
#Proj GraphEditor

diff --git a/Engine/Source/Editor/GraphEditor/Private/ConnectionDrawingPolicy.cpp b/Engine/Source/Editor/GraphEditor/Private/ConnectionDrawingPolicy.cpp
--- a/Engine/Source/Editor/GraphEditor/Private/ConnectionDrawingPolicy.cpp
+++ b/Engine/Source/Editor/GraphEditor/Private/ConnectionDrawingPolicy.cpp
@@ -479,15 +479,16 @@ void FConnectionDrawingPolicy::Draw(TMap<TSharedRef<SWidget>, FArrangedWidget>&
 
 				DetermineLinkGeometry(PinGeometries, ArrangedNodes, SomePinWidget, ThePin, TargetPin, /*out*/ LinkStartWidgetGeometry, /*out*/ LinkEndWidgetGeometry);
 
-				if ((LinkEndWidgetGeometry != NULL) && (LinkStartWidgetGeometry != NULL))
+				if(( LinkEndWidgetGeometry && LinkStartWidgetGeometry ) && !IsConnectionCulled( *LinkStartWidgetGeometry, *LinkEndWidgetGeometry ))
 				{
+
 					float Thickness = 1.0f;
 					FLinearColor WireColor = FLinearColor::White;
 					bool bDrawBubbles = false;
 					bool bBidirectional = false;
 
 					DetermineWiringStyle(ThePin, TargetPin, /*inout*/ Thickness, /*inout*/ WireColor, /*inout*/ bDrawBubbles, /*inout*/ bBidirectional);
-
+					NumConnections++;
 					DrawSplineWithArrow(LinkStartWidgetGeometry->Geometry, LinkEndWidgetGeometry->Geometry, WireColor, Thickness, bDrawBubbles, bBidirectional);
 				}
 			}
@@ -495,6 +496,17 @@ void FConnectionDrawingPolicy::Draw(TMap<TSharedRef<SWidget>, FArrangedWidget>&
 	}
 }
 
+bool FConnectionDrawingPolicy::IsConnectionCulled( const FArrangedWidget& StartLink, const FArrangedWidget& EndLink ) const
+{
+	const float Top		= FMath::Min( StartLink.Geometry.AbsolutePosition.Y, EndLink.Geometry.AbsolutePosition.Y );
+	const float Left	= FMath::Min( StartLink.Geometry.AbsolutePosition.X, EndLink.Geometry.AbsolutePosition.X );
+	const float Bottom	= FMath::Max( StartLink.Geometry.AbsolutePosition.Y, EndLink.Geometry.AbsolutePosition.Y );
+	const float Right	= FMath::Max( StartLink.Geometry.AbsolutePosition.X, EndLink.Geometry.AbsolutePosition.X ); 
+
+	return	Left > ClippingRect.Right || Right < ClippingRect.Left || 
+			Bottom < ClippingRect.Top || Top > ClippingRect.Bottom;
+}
+
 void FConnectionDrawingPolicy::SetIncompatiblePinDrawState(const TSharedPtr<SGraphPin>& StartPin, const TSet< TSharedRef<SWidget> >& VisiblePins)
 {
 }
diff --git a/Engine/Source/Editor/GraphEditor/Private/ConnectionDrawingPolicy.h b/Engine/Source/Editor/GraphEditor/Private/ConnectionDrawingPolicy.h
--- a/Engine/Source/Editor/GraphEditor/Private/ConnectionDrawingPolicy.h
+++ b/Engine/Source/Editor/GraphEditor/Private/ConnectionDrawingPolicy.h
@@ -85,6 +85,8 @@ public:
 	virtual void ResetIncompatiblePinDrawState(const TSet< TSharedRef<SWidget> >& VisiblePins);
 
 	virtual void ApplyHoverDeemphasis(UEdGraphPin* OutputPin, UEdGraphPin* InputPin, /*inout*/ float& Thickness, /*inout*/ FLinearColor& WireColor);
+
+	virtual bool IsConnectionCulled( const FArrangedWidget& StartLink, const FArrangedWidget& EndLink ) const;
 };
 
 /////////////////////////////////////////////////////
