Build fix, left a dangling reference to removed counter in cl 2321492

diff --git a/Engine/Source/Editor/GraphEditor/Private/ConnectionDrawingPolicy.cpp b/Engine/Source/Editor/GraphEditor/Private/ConnectionDrawingPolicy.cpp
--- a/Engine/Source/Editor/GraphEditor/Private/ConnectionDrawingPolicy.cpp
+++ b/Engine/Source/Editor/GraphEditor/Private/ConnectionDrawingPolicy.cpp
@@ -488,7 +488,6 @@ void FConnectionDrawingPolicy::Draw(TMap<TSharedRef<SWidget>, FArrangedWidget>&
 					bool bBidirectional = false;
 
 					DetermineWiringStyle(ThePin, TargetPin, /*inout*/ Thickness, /*inout*/ WireColor, /*inout*/ bDrawBubbles, /*inout*/ bBidirectional);
-					NumConnections++;
 					DrawSplineWithArrow(LinkStartWidgetGeometry->Geometry, LinkEndWidgetGeometry->Geometry, WireColor, Thickness, bDrawBubbles, bBidirectional);
 				}
 			}
