[347817] BP perf - Removed a redundant node construction pass on graph panel updates.

- Double-checked all SGraphNode derivatives to make sure they were actually doing construction in their Construct() calls; fixed up a couple that weren't.
- Cuts update timings roughly in half (still needs work though).

diff --git a/Engine/Source/Editor/GraphEditor/Private/MaterialNodes/SGraphNodeMaterialBase.cpp b/Engine/Source/Editor/GraphEditor/Private/MaterialNodes/SGraphNodeMaterialBase.cpp
--- a/Engine/Source/Editor/GraphEditor/Private/MaterialNodes/SGraphNodeMaterialBase.cpp
+++ b/Engine/Source/Editor/GraphEditor/Private/MaterialNodes/SGraphNodeMaterialBase.cpp
@@ -216,8 +216,7 @@ void SGraphNodeMaterialBase::Construct(const FArguments& InArgs, UMaterialGraphN
 
 	this->SetCursor(EMouseCursor::CardinalCross);
 
-	// Surely we don't need to do this as SGraphPanel calls it after construction
-	//this->UpdateGraphNode();
+	this->UpdateGraphNode();
 }
 
 void SGraphNodeMaterialBase::CreatePinWidgets()
diff --git a/Engine/Source/Editor/GraphEditor/Private/SGraphPanel.cpp b/Engine/Source/Editor/GraphEditor/Private/SGraphPanel.cpp
--- a/Engine/Source/Editor/GraphEditor/Private/SGraphPanel.cpp
+++ b/Engine/Source/Editor/GraphEditor/Private/SGraphPanel.cpp
@@ -971,13 +971,8 @@ void SGraphPanel::AddNode(UEdGraphNode* Node)
 
 		NewNode->PlaySpawnEffect();
 
-		NewNode->UpdateGraphNode();
 		NewNode->RequestRenameOnSpawn();
 	}
-	else
-	{
-		NewNode->UpdateGraphNode();
-	}
 }
 
 TSharedPtr<SGraphNode> SGraphPanel::GetNodeWidgetFromGuid(FGuid Guid) const
@@ -987,6 +982,8 @@ TSharedPtr<SGraphNode> SGraphPanel::GetNodeWidgetFromGuid(FGuid Guid) const
 
 void SGraphPanel::Update()
 {
+//	SCOPE_LOG_TIME_FUNC();
+
 	// Add widgets for all the nodes that don't have one.
 	if(GraphObj != NULL)
 	{
