Backout CL# 2324267 - seems to be causing some issues w/ connection dragging.

diff --git a/Engine/Source/Editor/GraphEditor/Private/MaterialNodes/SGraphNodeMaterialBase.cpp b/Engine/Source/Editor/GraphEditor/Private/MaterialNodes/SGraphNodeMaterialBase.cpp
--- a/Engine/Source/Editor/GraphEditor/Private/MaterialNodes/SGraphNodeMaterialBase.cpp
+++ b/Engine/Source/Editor/GraphEditor/Private/MaterialNodes/SGraphNodeMaterialBase.cpp
@@ -216,7 +216,8 @@ void SGraphNodeMaterialBase::Construct(const FArguments& InArgs, UMaterialGraphN
 
 	this->SetCursor(EMouseCursor::CardinalCross);
 
-	this->UpdateGraphNode();
+	// Surely we don't need to do this as SGraphPanel calls it after construction
+	//this->UpdateGraphNode();
 }
 
 void SGraphNodeMaterialBase::CreatePinWidgets()
diff --git a/Engine/Source/Editor/GraphEditor/Private/SGraphPanel.cpp b/Engine/Source/Editor/GraphEditor/Private/SGraphPanel.cpp
--- a/Engine/Source/Editor/GraphEditor/Private/SGraphPanel.cpp
+++ b/Engine/Source/Editor/GraphEditor/Private/SGraphPanel.cpp
@@ -971,8 +971,13 @@ void SGraphPanel::AddNode(UEdGraphNode* Node)
 
 		NewNode->PlaySpawnEffect();
 
+		NewNode->UpdateGraphNode();
 		NewNode->RequestRenameOnSpawn();
 	}
+	else
+	{
+		NewNode->UpdateGraphNode();
+	}
 }
 
 TSharedPtr<SGraphNode> SGraphPanel::GetNodeWidgetFromGuid(FGuid Guid) const
@@ -982,8 +987,6 @@ TSharedPtr<SGraphNode> SGraphPanel::GetNodeWidgetFromGuid(FGuid Guid) const
 
 void SGraphPanel::Update()
 {
-//	SCOPE_LOG_TIME_FUNC();
-
 	// Add widgets for all the nodes that don't have one.
 	if(GraphObj != NULL)
 	{
