[347817] BP perf - Removed a redundant node construction pass on graph panel updates.

- Cuts update timings roughly in half.
- Fixed issue that this change prevously had introduced: we just needed to ensure that all pin widgets are enabled/disabled when the owning node widget is enabled/disabled.

diff --git a/Engine/Source/Editor/GraphEditor/Private/SGraphNode.cpp b/Engine/Source/Editor/GraphEditor/Private/SGraphNode.cpp
--- a/Engine/Source/Editor/GraphEditor/Private/SGraphNode.cpp
+++ b/Engine/Source/Editor/GraphEditor/Private/SGraphNode.cpp
@@ -131,6 +131,17 @@ bool SGraphNode::CanAllowInteractionUsingDragDropOp( const UEdGraphNode* GraphNo
 void SGraphNode::SetIsEditable(TAttribute<bool> InIsEditable)
 {
 	IsEditable = InIsEditable;
+
+	int32 i;
+	for(i = 0; i < InputPins.Num(); ++i)
+	{
+		InputPins[i]->SetIsEditable(IsEditable);
+	}
+
+	for(i = 0; i < OutputPins.Num(); ++i)
+	{
+		OutputPins[i]->SetIsEditable(IsEditable);
+	}
 }
 
 bool SGraphNode::IsNodeEditable() const
diff --git a/Engine/Source/Editor/GraphEditor/Private/SGraphPanel.cpp b/Engine/Source/Editor/GraphEditor/Private/SGraphPanel.cpp
--- a/Engine/Source/Editor/GraphEditor/Private/SGraphPanel.cpp
+++ b/Engine/Source/Editor/GraphEditor/Private/SGraphPanel.cpp
@@ -970,14 +970,8 @@ void SGraphPanel::AddNode(UEdGraphNode* Node)
 		VisibleChildren.Add(NewNode.ToSharedRef());
 
 		NewNode->PlaySpawnEffect();
-
-		NewNode->UpdateGraphNode();
 		NewNode->RequestRenameOnSpawn();
 	}
-	else
-	{
-		NewNode->UpdateGraphNode();
-	}
 }
 
 TSharedPtr<SGraphNode> SGraphPanel::GetNodeWidgetFromGuid(FGuid Guid) const
