[349626] fix material nodes not showing up after CL# 2329052

- restored previous changes that i forgot to include with that submission

diff --git a/Engine/Source/Editor/GraphEditor/Private/MaterialNodes/SGraphNodeMaterialBase.cpp b/Engine/Source/Editor/GraphEditor/Private/MaterialNodes/SGraphNodeMaterialBase.cpp
--- a/Engine/Source/Editor/GraphEditor/Private/MaterialNodes/SGraphNodeMaterialBase.cpp
+++ b/Engine/Source/Editor/GraphEditor/Private/MaterialNodes/SGraphNodeMaterialBase.cpp
@@ -216,8 +216,7 @@ void SGraphNodeMaterialBase::Construct(const FArguments& InArgs, UMaterialGraphN
 
 	this->SetCursor(EMouseCursor::CardinalCross);
 
-	// Surely we don't need to do this as SGraphPanel calls it after construction
-	//this->UpdateGraphNode();
+	this->UpdateGraphNode();
 }
 
 void SGraphNodeMaterialBase::CreatePinWidgets()
