Added some synchronisation to NameList GraphPin's to ensure the current values match the UI.
#TTP 349533 - LIVE: BP: Material Parameter Collections are not working when called through BP
#Branch UE4
#Proj GraphEditor

diff --git a/Engine/Source/Editor/GraphEditor/Private/KismetPins/SGraphPinNameList.cpp b/Engine/Source/Editor/GraphEditor/Private/KismetPins/SGraphPinNameList.cpp
--- a/Engine/Source/Editor/GraphEditor/Private/KismetPins/SGraphPinNameList.cpp
+++ b/Engine/Source/Editor/GraphEditor/Private/KismetPins/SGraphPinNameList.cpp
@@ -32,6 +32,11 @@ TSharedRef<SWidget>	SGraphPinNameList::GetDefaultValueWidget()
 			}
 		}
 		check(CurrentlySelectedName.IsValid());
+		// Ensure the pin value is synchronised with the graph pin widget
+		if( PreviousSelection != *CurrentlySelectedName.Get() )
+		{
+			GraphPinObj->GetSchema()->TrySetDefaultValue(*GraphPinObj, *CurrentlySelectedName->ToString());
+		}
 	}
 
 	// Create widget
