Fix for a compilation error. This may be temporary - it's the only way I know of right now to fix a missing symbol error for NAME_DefaultPinLabelStyle.

diff --git a/Engine/Source/Editor/GraphEditor/Private/SGraphPin.cpp b/Engine/Source/Editor/GraphEditor/Private/SGraphPin.cpp
--- a/Engine/Source/Editor/GraphEditor/Private/SGraphPin.cpp
+++ b/Engine/Source/Editor/GraphEditor/Private/SGraphPin.cpp
@@ -70,7 +70,6 @@ struct FKnotNetCollector
 /////////////////////////////////////////////////////
 // SGraphPin
 
-FName SGraphPin::NAME_DefaultPinLabelStyle(TEXT("Graph.Node.PinName"));
 
 SGraphPin::SGraphPin()
 	: bShowLabel(true)
diff --git a/Engine/Source/Editor/GraphEditor/Public/SGraphPin.h b/Engine/Source/Editor/GraphEditor/Public/SGraphPin.h
--- a/Engine/Source/Editor/GraphEditor/Public/SGraphPin.h
+++ b/Engine/Source/Editor/GraphEditor/Public/SGraphPin.h
@@ -4,13 +4,14 @@
 
 class SGraphNode;
 
+#define NAME_DefaultPinLabelStyle TEXT("Graph.Node.PinName")
+
 /**
  * Represents a pin on a node in the GraphEditor
  */
 class GRAPHEDITOR_API SGraphPin : public SBorder
 {
 public:
-	static FName NAME_DefaultPinLabelStyle;
 
 	SLATE_BEGIN_ARGS(SGraphPin)
 		: _PinLabelStyle(NAME_DefaultPinLabelStyle)
