De-emphasized connectors in the material editor are drawn much brighter, increases visibility with comment backgrounds

diff --git a/Engine/Source/Editor/GraphEditor/Private/ConnectionDrawingPolicy.cpp b/Engine/Source/Editor/GraphEditor/Private/ConnectionDrawingPolicy.cpp
--- a/Engine/Source/Editor/GraphEditor/Private/ConnectionDrawingPolicy.cpp
+++ b/Engine/Source/Editor/GraphEditor/Private/ConnectionDrawingPolicy.cpp
@@ -199,6 +199,7 @@ FConnectionDrawingPolicy::FConnectionDrawingPolicy(int32 InBackLayerID, int32 In
 	ArrowRadius = ArrowImage->ImageSize * ZoomFactor * 0.5f;
 	MidpointImage = NULL;
 	MidpointRadius = FVector2D::ZeroVector;
+	HoverDeemphasisDarkFraction = 0.8f;
 
 	BubbleImage = FEditorStyle::GetBrush( TEXT("Graph.ExecutionBubble") );
 }
@@ -520,7 +521,6 @@ void FConnectionDrawingPolicy::ApplyHoverDeemphasis(UEdGraphPin* OutputPin, UEdG
 	const float FadeInPeriod = 0.6f; // Time in seconds after the bias before the fade is fully complete
 	const float TimeFraction = FMath::SmoothStep(0.0f, FadeInPeriod, (float)(FSlateApplication::Get().GetCurrentTime() - LastHoverTimeEvent - FadeInBias));
 
-	const float DarkFraction = 0.8f;
 	const float LightFraction = 0.25f;
 	const FLinearColor DarkenedColor(0.0f, 0.0f, 0.0f, 0.5f);
 	const FLinearColor LightenedColor(1.0f, 1.0f, 1.0f, 1.0f);
@@ -535,7 +535,7 @@ void FConnectionDrawingPolicy::ApplyHoverDeemphasis(UEdGraphPin* OutputPin, UEdG
 	}
 	else
 	{
-		WireColor = FMath::Lerp<FLinearColor>(WireColor, DarkenedColor, DarkFraction * TimeFraction);
+		WireColor = FMath::Lerp<FLinearColor>(WireColor, DarkenedColor, HoverDeemphasisDarkFraction * TimeFraction);
 	}
 }
 
@@ -1405,6 +1405,9 @@ FMaterialGraphConnectionDrawingPolicy::FMaterialGraphConnectionDrawingPolicy(int
 	// Don't want to draw ending arrowheads
 	ArrowImage = nullptr;
 	ArrowRadius = FVector2D::ZeroVector;
+
+	// Still need to be able to perceive the graph while dragging connectors, esp over comment boxes
+	HoverDeemphasisDarkFraction = 0.4f;
 }
 
 void FMaterialGraphConnectionDrawingPolicy::DetermineWiringStyle(UEdGraphPin* OutputPin, UEdGraphPin* InputPin, /*inout*/ float& Thickness, /*inout*/ FLinearColor& WireColor, /*inout*/bool& bDrawBubbles, /*inout*/ bool& bBidirectional)
diff --git a/Engine/Source/Editor/GraphEditor/Private/ConnectionDrawingPolicy.h b/Engine/Source/Editor/GraphEditor/Private/ConnectionDrawingPolicy.h
--- a/Engine/Source/Editor/GraphEditor/Private/ConnectionDrawingPolicy.h
+++ b/Engine/Source/Editor/GraphEditor/Private/ConnectionDrawingPolicy.h
@@ -40,6 +40,7 @@ public:
 	FVector2D MidpointRadius;
 protected:
 	float ZoomFactor; 
+	float HoverDeemphasisDarkFraction;
 	const FSlateRect& ClippingRect;
 	FSlateWindowElementList& DrawElementsList;
 	TMap< UEdGraphPin*, TSharedRef<SGraphPin> > PinToPinWidgetMap;
