Default pin for switch nodes shows an exec output, not an output pin

diff --git a/Engine/Source/Editor/GraphEditor/Private/KismetNodes/SGraphNodeSwitchStatement.cpp b/Engine/Source/Editor/GraphEditor/Private/KismetNodes/SGraphNodeSwitchStatement.cpp
--- a/Engine/Source/Editor/GraphEditor/Private/KismetNodes/SGraphNodeSwitchStatement.cpp
+++ b/Engine/Source/Editor/GraphEditor/Private/KismetNodes/SGraphNodeSwitchStatement.cpp
@@ -22,6 +22,9 @@ public:
 	void Construct(const FArguments& InArgs, UEdGraphPin* InPin)
 	{
 		SGraphPin::Construct(SGraphPin::FArguments().PinLabelStyle(FName("Graph.Node.DefaultPinName")), InPin);
+
+		bWasEventPin = false;
+		CachePinBrushes(/*bForceCache=*/ true);
 	}
 };
 
