This change makes the comment bubbles interactive widgets which allow the user to double click the comment to edit and have extra optional pinning and visibility options.
#UE-2764 - ROCKET: TASK: K2: PUNTABLE: Allow double clicking on comment bubble to edit it.
#Branch UE4
#Proj Engine, UnrealEd, Slate, GraphEditor, BehaviourTreeEditor, EnvironmentQueryEditor, EditorStyle, ReferenceViewer, WorldBrowser

diff --git a/Engine/Source/Editor/BehaviorTreeEditor/Private/SGraphNode_BehaviorTree.cpp b/Engine/Source/Editor/BehaviorTreeEditor/Private/SGraphNode_BehaviorTree.cpp
--- a/Engine/Source/Editor/BehaviorTreeEditor/Private/SGraphNode_BehaviorTree.cpp
+++ b/Engine/Source/Editor/BehaviorTreeEditor/Private/SGraphNode_BehaviorTree.cpp
@@ -15,6 +15,7 @@
 #include "BehaviorTree/Tasks/BTTask_RunBehavior.h"
 #include "IDocumentation.h"
 #include "SInlineEditableTextBlock.h"
+#include "SCommentBubble.h"
 
 #define LOCTEXT_NAMESPACE "SGraphNode_BehaviorTree"
 
@@ -365,7 +366,7 @@ void SGraphNode_BehaviorTree::UpdateGraphNode()
 		.OnGetIndexColor(this, &SGraphNode_BehaviorTree::GetIndexColor);
 
 	this->ContentScale.Bind( this, &SGraphNode::GetContentScale );
-	this->ChildSlot
+	this->GetOrAddSlot( ENodeZone::Center )
 		.HAlign(HAlign_Fill)
 		.VAlign(VAlign_Center)
 		[
@@ -523,6 +524,27 @@ void SGraphNode_BehaviorTree::UpdateGraphNode()
 				]
 			]
 		];
+	// Create comment bubble
+	TSharedPtr<SCommentBubble> CommentBubble;
+
+	SAssignNew( CommentBubble, SCommentBubble )
+	.GraphNode( GraphNode )
+	.Text( this, &SGraphNode::GetNodeComment )
+	.ColorAndOpacity( this, &SGraphNode_BehaviorTree::GetCommentColor )
+	.AllowPinning( true )
+	.EnableTitleBarBubble( true )
+	.EnableBubbleCtrls( true )
+	.GraphLOD( this, &SGraphNode::GetCurrentLOD )
+	.IsGraphNodeHovered( this, &SGraphNode::IsHovered );
+
+	GetOrAddSlot( ENodeZone::TopCenter )
+	.SlotOffset( TAttribute<FVector2D>( CommentBubble.Get(), &SCommentBubble::GetOffset ))
+	.SlotSize( TAttribute<FVector2D>( CommentBubble.Get(), &SCommentBubble::GetSize ))
+	.AllowScaling( TAttribute<bool>( CommentBubble.Get(), &SCommentBubble::IsScalingAllowed ))
+	.VAlign( VAlign_Top )
+	[
+		CommentBubble.ToSharedRef()
+	];
 
 	ErrorReporting = ErrorText;
 	ErrorReporting->SetError(ErrorMsg);
diff --git a/Engine/Source/Editor/EditorStyle/Private/SlateEditorStyle.cpp b/Engine/Source/Editor/EditorStyle/Private/SlateEditorStyle.cpp
--- a/Engine/Source/Editor/EditorStyle/Private/SlateEditorStyle.cpp
+++ b/Engine/Source/Editor/EditorStyle/Private/SlateEditorStyle.cpp
@@ -3226,7 +3226,40 @@ void FSlateEditorStyle::FStyle::SetupGraphEditorStyles()
 		Set( "Graph.Node.CommentArrow", new IMAGE_BRUSH( "Old/Graph/CommentBubbleArrow", FVector2D(8,8) ) );
 		Set( "Graph.Node.CommentFont", TTF_CORE_FONT( "Fonts/Roboto-Regular", 10 ) );
 		Set( "Graph.Node.Comment.TextColor", FLinearColor::Black );
-		Set( "Graph.Node.Comment.BubblePadding", FVector2D(3,3) );
+		Set( "Graph.Node.Comment.BubbleOffset", FMargin(8,0,0,0) );
+		Set( "Graph.Node.Comment.PinIconPadding", FMargin(0,2,0,0) );
+		Set( "Graph.Node.Comment.BubblePadding", FMargin(4,4) );
+
+		const FCheckBoxStyle CommentTitleButton = FCheckBoxStyle()
+			.SetCheckBoxType( ESlateCheckBoxType::CheckBox )
+			.SetUncheckedImage( IMAGE_BRUSH( "Icons/icon_Blueprint_CommentBubbleOn_16x", FVector2D(16,16), FLinearColor(1.f, 1.f, 1.f, 0.8f)))
+			.SetUncheckedHoveredImage( IMAGE_BRUSH( "Icons/icon_Blueprint_CommentBubbleOn_16x", FVector2D(16,16), FLinearColor(1.f, 1.f, 1.f, 0.9f)))
+			.SetUncheckedPressedImage( IMAGE_BRUSH( "Icons/icon_Blueprint_CommentBubbleOn_16x", FVector2D(16,16), FLinearColor(1.f, 1.f, 1.f, 1.f)))
+			.SetCheckedImage( IMAGE_BRUSH( "Icons/icon_Blueprint_CommentBubbleOn_16x", FVector2D(16,16), FLinearColor(1.f, 1.f, 1.f, 0.8f)))
+			.SetCheckedHoveredImage( IMAGE_BRUSH( "Icons/icon_Blueprint_CommentBubbleOn_16x", FVector2D(16,16), FLinearColor(1.f, 1.f, 1.f, 1.f)))
+			.SetCheckedPressedImage( IMAGE_BRUSH( "Icons/icon_Blueprint_CommentBubbleOff_16x", FVector2D(16,16), FLinearColor(1.f, 1.f, 1.f, 0.6f)));
+		Set( "CommentTitleButton", CommentTitleButton );
+
+		const FCheckBoxStyle CommentBubbleButton = FCheckBoxStyle()
+			.SetCheckBoxType( ESlateCheckBoxType::CheckBox )
+			.SetUncheckedImage( IMAGE_BRUSH( "Icons/icon_Blueprint_CommentBubbleOn_16x", FVector2D(10,10), FLinearColor(1.f, 1.f, 1.f, 0.5f)))
+			.SetUncheckedHoveredImage( IMAGE_BRUSH( "Icons/icon_Blueprint_CommentBubbleOn_16x", FVector2D(10,10), FLinearColor(1.f, 1.f, 1.f, 0.9f)))
+			.SetUncheckedPressedImage( IMAGE_BRUSH( "Icons/icon_Blueprint_CommentBubbleOn_16x", FVector2D(10,10), FLinearColor(1.f, 1.f, 1.f, 1.f)))
+			.SetCheckedImage( IMAGE_BRUSH( "Icons/icon_Blueprint_CommentBubbleOn_16x", FVector2D(10,10), FLinearColor(1.f, 1.f, 1.f, 0.5f)))
+			.SetCheckedHoveredImage( IMAGE_BRUSH( "Icons/icon_Blueprint_CommentBubbleOn_16x", FVector2D(10,10), FLinearColor(1.f, 1.f, 1.f, 1.f)))
+			.SetCheckedPressedImage( IMAGE_BRUSH( "Icons/icon_Blueprint_CommentBubbleOn_16x", FVector2D(10,10), FLinearColor(1.f, 1.f, 1.f, 0.4f)));
+		Set( "CommentBubbleButton", CommentBubbleButton );
+
+		const FCheckBoxStyle CommentBubblePin = FCheckBoxStyle()
+			.SetCheckBoxType( ESlateCheckBoxType::CheckBox )
+			.SetUncheckedImage( IMAGE_BRUSH( "Icons/icon_Blueprint_CommentBubbleUnPin_16x", FVector2D(10,10), FLinearColor(1.f, 1.f, 1.f, 0.3f)))
+			.SetUncheckedHoveredImage( IMAGE_BRUSH( "Icons/icon_Blueprint_CommentBubbleUnPin_16x", FVector2D(10,10), FLinearColor(1.f, 1.f, 1.f, 0.5f)))
+			.SetUncheckedPressedImage( IMAGE_BRUSH( "Icons/icon_Blueprint_CommentBubblePin_16x", FVector2D(10,10), FLinearColor(1.f, 1.f, 1.f, 0.8f)))
+			.SetCheckedImage( IMAGE_BRUSH( "Icons/icon_Blueprint_CommentBubblePin_16x", FVector2D(10,10), FLinearColor(1.f, 1.f, 1.f, 0.3f)))
+			.SetCheckedHoveredImage( IMAGE_BRUSH( "Icons/icon_Blueprint_CommentBubblePin_16x", FVector2D(10,10), FLinearColor(1.f, 1.f, 1.f, 0.5f)))
+			.SetCheckedPressedImage( IMAGE_BRUSH( "Icons/icon_Blueprint_CommentBubbleUnPin_16x", FVector2D(10,10), FLinearColor(1.f, 1.f, 1.f, 0.8f)));
+		Set( "CommentBubblePin", CommentBubblePin );
+
 
 		Set( "Graph.VarNode.Body", new BOX_BRUSH( "/Graph/VarNode_body", FMargin(16.f/64.f, 12.f/28.f) ) );
 		Set( "Graph.VarNode.ColorSpill", new IMAGE_BRUSH( "/Graph/VarNode_color_spill", FVector2D(132,28) ) );
diff --git a/Engine/Source/Editor/EnvironmentQueryEditor/Private/SGraphNode_EnvironmentQuery.cpp b/Engine/Source/Editor/EnvironmentQueryEditor/Private/SGraphNode_EnvironmentQuery.cpp
--- a/Engine/Source/Editor/EnvironmentQueryEditor/Private/SGraphNode_EnvironmentQuery.cpp
+++ b/Engine/Source/Editor/EnvironmentQueryEditor/Private/SGraphNode_EnvironmentQuery.cpp
@@ -187,7 +187,7 @@ void SGraphNode_EnvironmentQuery::UpdateGraphNode()
 	TSharedPtr<SNodeTitle> NodeTitle = SNew(SNodeTitle, GraphNode);
 
 	this->ContentScale.Bind( this, &SGraphNode::GetContentScale );
-	this->ChildSlot
+	this->GetOrAddSlot( ENodeZone::Center )
 		.HAlign(HAlign_Fill)
 		.VAlign(VAlign_Center)
 		[
diff --git a/Engine/Source/Editor/GraphEditor/Private/AnimationStateNodes/SGraphNodeAnimState.cpp b/Engine/Source/Editor/GraphEditor/Private/AnimationStateNodes/SGraphNodeAnimState.cpp
--- a/Engine/Source/Editor/GraphEditor/Private/AnimationStateNodes/SGraphNodeAnimState.cpp
+++ b/Engine/Source/Editor/GraphEditor/Private/AnimationStateNodes/SGraphNodeAnimState.cpp
@@ -178,7 +178,7 @@ void SGraphNodeAnimState::UpdateGraphNode()
 	TSharedPtr<SNodeTitle> NodeTitle = SNew(SNodeTitle, GraphNode);
 
 	this->ContentScale.Bind( this, &SGraphNode::GetContentScale );
-	this->ChildSlot
+	this->GetOrAddSlot( ENodeZone::Center )
 		.HAlign(HAlign_Center)
 		.VAlign(VAlign_Center)
 		[
diff --git a/Engine/Source/Editor/GraphEditor/Private/AnimationStateNodes/SGraphNodeAnimStateEntry.cpp b/Engine/Source/Editor/GraphEditor/Private/AnimationStateNodes/SGraphNodeAnimStateEntry.cpp
--- a/Engine/Source/Editor/GraphEditor/Private/AnimationStateNodes/SGraphNodeAnimStateEntry.cpp
+++ b/Engine/Source/Editor/GraphEditor/Private/AnimationStateNodes/SGraphNodeAnimStateEntry.cpp
@@ -47,7 +47,7 @@ void SGraphNodeAnimStateEntry::UpdateGraphNode()
 	FLinearColor TitleShadowColor(0.6f, 0.6f, 0.6f);
 
 	this->ContentScale.Bind( this, &SGraphNode::GetContentScale );
-	this->ChildSlot
+	this->GetOrAddSlot( ENodeZone::Center )
 		.HAlign(HAlign_Center)
 		.VAlign(VAlign_Center)
 		[
diff --git a/Engine/Source/Editor/GraphEditor/Private/AnimationStateNodes/SGraphNodeAnimTransition.cpp b/Engine/Source/Editor/GraphEditor/Private/AnimationStateNodes/SGraphNodeAnimTransition.cpp
--- a/Engine/Source/Editor/GraphEditor/Private/AnimationStateNodes/SGraphNodeAnimTransition.cpp
+++ b/Engine/Source/Editor/GraphEditor/Private/AnimationStateNodes/SGraphNodeAnimTransition.cpp
@@ -165,7 +165,7 @@ void SGraphNodeAnimTransition::UpdateGraphNode()
 	LeftNodeBox.Reset();
 
 	this->ContentScale.Bind( this, &SGraphNode::GetContentScale );
-	this->ChildSlot
+	this->GetOrAddSlot( ENodeZone::Center )
 		.HAlign(HAlign_Center)
 		.VAlign(VAlign_Center)
 		[
diff --git a/Engine/Source/Editor/GraphEditor/Private/KismetNodes/SGraphNodeK2ArrayFunction.cpp b/Engine/Source/Editor/GraphEditor/Private/KismetNodes/SGraphNodeK2ArrayFunction.cpp
--- a/Engine/Source/Editor/GraphEditor/Private/KismetNodes/SGraphNodeK2ArrayFunction.cpp
+++ b/Engine/Source/Editor/GraphEditor/Private/KismetNodes/SGraphNodeK2ArrayFunction.cpp
@@ -70,7 +70,7 @@ void SGraphNodeK2ArrayFunction::UpdateGraphNode()
 			.Text(this, &SGraphNodeK2ArrayFunction::GetNodeCompactTitle);
 
 		this->ContentScale.Bind( this, &SGraphNode::GetContentScale );
-		this->ChildSlot
+		this->GetOrAddSlot( ENodeZone::Center )
 			.HAlign(HAlign_Center)
 			.VAlign(VAlign_Center)
 			[
diff --git a/Engine/Source/Editor/GraphEditor/Private/KismetNodes/SGraphNodeK2Base.cpp b/Engine/Source/Editor/GraphEditor/Private/KismetNodes/SGraphNodeK2Base.cpp
--- a/Engine/Source/Editor/GraphEditor/Private/KismetNodes/SGraphNodeK2Base.cpp
+++ b/Engine/Source/Editor/GraphEditor/Private/KismetNodes/SGraphNodeK2Base.cpp
@@ -59,7 +59,8 @@ void SGraphNodeK2Base::UpdateCompactNode()
 	//            |_______|______|_______|
 	//
 	this->ContentScale.Bind( this, &SGraphNode::GetContentScale );
-	this->ChildSlot
+	this->GetOrAddSlot( ENodeZone::Center )
+
 	.HAlign(HAlign_Center)
 	.VAlign(VAlign_Center)
 	[
diff --git a/Engine/Source/Editor/GraphEditor/Private/KismetNodes/SGraphNodeK2Composite.cpp b/Engine/Source/Editor/GraphEditor/Private/KismetNodes/SGraphNodeK2Composite.cpp
--- a/Engine/Source/Editor/GraphEditor/Private/KismetNodes/SGraphNodeK2Composite.cpp
+++ b/Engine/Source/Editor/GraphEditor/Private/KismetNodes/SGraphNodeK2Composite.cpp
@@ -44,7 +44,7 @@ void SGraphNodeK2Composite::UpdateGraphNode()
 	//            |_______|______|_______|
 	//
 	this->ContentScale.Bind( this, &SGraphNode::GetContentScale );
-	this->ChildSlot
+	this->GetOrAddSlot( ENodeZone::Center )
 		.HAlign(HAlign_Center)
 		.VAlign(VAlign_Center)
 		[
@@ -119,6 +119,27 @@ void SGraphNodeK2Composite::UpdateGraphNode()
 				]
 			]
 		];
+	// Create comment bubble
+	TSharedPtr<SCommentBubble> CommentBubble;
+
+	SAssignNew( CommentBubble, SCommentBubble )
+	.GraphNode( GraphNode )
+	.Text( this, &SGraphNode::GetNodeComment )
+	.ColorAndOpacity( this, &SGraphNodeK2Composite::GetCommentColor )
+	.AllowPinning( true )
+	.EnableTitleBarBubble( true )
+	.EnableBubbleCtrls( true )
+	.GraphLOD( this, &SGraphNode::GetCurrentLOD )
+	.IsGraphNodeHovered( this, &SGraphNode::IsHovered );
+
+	GetOrAddSlot( ENodeZone::TopCenter )
+	.SlotOffset( TAttribute<FVector2D>( CommentBubble.Get(), &SCommentBubble::GetOffset ))
+	.SlotSize( TAttribute<FVector2D>( CommentBubble.Get(), &SCommentBubble::GetSize ))
+	.AllowScaling( TAttribute<bool>( CommentBubble.Get(), &SCommentBubble::IsScalingAllowed ))
+	.VAlign( VAlign_Top )
+	[
+		CommentBubble.ToSharedRef()
+	];
 
 	CreatePinWidgets();
 }
diff --git a/Engine/Source/Editor/GraphEditor/Private/KismetNodes/SGraphNodeK2Terminator.cpp b/Engine/Source/Editor/GraphEditor/Private/KismetNodes/SGraphNodeK2Terminator.cpp
--- a/Engine/Source/Editor/GraphEditor/Private/KismetNodes/SGraphNodeK2Terminator.cpp
+++ b/Engine/Source/Editor/GraphEditor/Private/KismetNodes/SGraphNodeK2Terminator.cpp
@@ -54,7 +54,7 @@ void SGraphNodeK2Terminator::UpdateGraphNode()
 	//            |_______|______|_______|
 	//
 	this->ContentScale.Bind( this, &SGraphNode::GetContentScale );
-	this->ChildSlot
+	this->GetOrAddSlot( ENodeZone::Center )
 	.HAlign(HAlign_Center)
 	.VAlign(VAlign_Center)
 	[
diff --git a/Engine/Source/Editor/GraphEditor/Private/KismetNodes/SGraphNodeK2Var.cpp b/Engine/Source/Editor/GraphEditor/Private/KismetNodes/SGraphNodeK2Var.cpp
--- a/Engine/Source/Editor/GraphEditor/Private/KismetNodes/SGraphNodeK2Var.cpp
+++ b/Engine/Source/Editor/GraphEditor/Private/KismetNodes/SGraphNodeK2Var.cpp
@@ -173,7 +173,7 @@ void SGraphNodeK2Var::UpdateGraphNode()
 	//            |_______|________|
 	//
 	this->ContentScale.Bind( this, &SGraphNode::GetContentScale );
-	this->ChildSlot
+	this->GetOrAddSlot( ENodeZone::Center )
 	.HAlign(HAlign_Center)
 	.VAlign(VAlign_Center)
 	[
diff --git a/Engine/Source/Editor/GraphEditor/Private/SCommentBubble.cpp b/Engine/Source/Editor/GraphEditor/Private/SCommentBubble.cpp
new file mode 100644
--- /dev/null
+++ b/Engine/Source/Editor/GraphEditor/Private/SCommentBubble.cpp
@@ -0,0 +1,382 @@
+// Copyright 1998-2013 Epic Games, Inc. All Rights Reserved.
+
+#include "GraphEditorCommon.h"
+#include "SCommentBubble.h"
+#include "SInlineEditableTextBlock.h"
+
+namespace SCommentBubbleDefs
+{
+	/** Mouse double click disable */
+	static const float DoubleClickDisable = -1.f;
+
+	/** Mouse double click delay */
+	static const float DoubleClickDelay = 0.5f;
+
+	/** Bubble Full opacity */
+	static const float FullOpacity = 1.f;
+
+	/** Bubble Toggle Icon Fade Down Speed */
+	static const float FadeDownSpeed = 5.f;
+
+	/** Height of the arrow connecting the bubble to the node */
+	static const float BubbleArrowHeight = 8.f;
+};
+
+
+void SCommentBubble::Construct( const FArguments& InArgs )
+{
+	check( InArgs._Text.IsBound() );
+	check( InArgs._GraphNode );
+
+	GraphNode				= InArgs._GraphNode;
+	CommentAttribute		= InArgs._Text;
+	ColorAndOpacity			= InArgs._ColorAndOpacity;
+	bAllowPinning			= InArgs._AllowPinning;
+	bEnableTitleBarBubble	= InArgs._EnableTitleBarBubble;
+	bEnableBubbleCtrls		= InArgs._EnableBubbleCtrls;
+	GraphLOD				= InArgs._GraphLOD;
+	IsGraphNodeHovered		= InArgs._IsGraphNodeHovered;
+
+	DoubleClickDelayTime	= SCommentBubbleDefs::DoubleClickDisable;
+	OpacityValue			= 0.f;
+
+	// Create Widget
+	UpdateBubble();
+}
+
+FCursorReply SCommentBubble::OnCursorQuery( const FGeometry& MyGeometry, const FPointerEvent& CursorEvent ) const
+{
+	const FVector2D Size( GetDesiredSize().X, GetDesiredSize().Y - SCommentBubbleDefs::BubbleArrowHeight );
+	const FSlateRect TestRect( MyGeometry.AbsolutePosition, MyGeometry.AbsolutePosition + Size );
+
+	if( TestRect.ContainsPoint( CursorEvent.GetScreenSpacePosition() ))
+	{
+		return FCursorReply::Cursor( EMouseCursor::TextEditBeam );
+	}
+	return FCursorReply::Cursor( EMouseCursor::Default );
+}
+
+FReply SCommentBubble::OnMouseButtonDown( const FGeometry& MyGeometry, const FPointerEvent& CursorEvent )
+{
+	if( TextBlock.IsValid() )
+	{
+		const FVector2D Size( GetDesiredSize().X, GetDesiredSize().Y - SCommentBubbleDefs::BubbleArrowHeight );
+		const FSlateRect TestRect( MyGeometry.AbsolutePosition, MyGeometry.AbsolutePosition + Size );
+
+		if( TestRect.ContainsPoint( CursorEvent.GetScreenSpacePosition() ))
+		{
+			if( DoubleClickDelayTime >= SCommentBubbleDefs::DoubleClickDelay )
+			{
+				TextBlock->EnterEditingMode();
+				DoubleClickDelayTime = SCommentBubbleDefs::DoubleClickDisable;
+			}
+			else
+			{
+				FSlateApplication::Get().SetKeyboardFocus( TextBlock, EKeyboardFocusCause::SetDirectly );
+				DoubleClickDelayTime = 0.f;
+			}
+		}
+	}
+	return FReply::Handled();
+}
+
+FReply SCommentBubble::OnMouseButtonDoubleClick( const FGeometry& MyGeometry, const FPointerEvent& MouseEvent )
+{
+	return FReply::Handled();
+}
+
+void SCommentBubble::Tick( const FGeometry& AllottedGeometry, const double InCurrentTime, const float InDeltaTime )
+{
+	SCompoundWidget::Tick( AllottedGeometry, InCurrentTime, InDeltaTime );
+
+	bool bIsHovered = IsHovered();
+
+	if( DoubleClickDelayTime >= 0.f )
+	{
+		DoubleClickDelayTime += InDeltaTime;
+		
+		if( !bIsHovered )
+		{
+			DoubleClickDelayTime = -1.f;
+		}
+	}
+	if( bEnableTitleBarBubble && IsGraphNodeHovered.IsBound() )
+	{
+		bIsHovered |= IsGraphNodeHovered.Execute();
+
+		if( !GraphNode->bCommentBubbleVisible && !bIsHovered )
+		{
+			if( OpacityValue > 0.f )
+			{
+				const float FadeDownAmt = InDeltaTime * SCommentBubbleDefs::FadeDownSpeed;
+				OpacityValue = FMath::Max( OpacityValue - FadeDownAmt, 0.f );
+			}
+		}
+		else
+		{
+			OpacityValue = SCommentBubbleDefs::FullOpacity;
+		}
+	}
+}
+
+void SCommentBubble::UpdateBubble()
+{
+	if( GraphNode->bCommentBubbleVisible )
+	{
+		const FSlateBrush* CommentCalloutArrowBrush = FEditorStyle::GetBrush(TEXT("Graph.Node.CommentArrow"));
+		const FMargin BubblePadding = FEditorStyle::GetMargin( TEXT("Graph.Node.Comment.BubblePadding"));
+		const FMargin PinIconPadding = FEditorStyle::GetMargin( TEXT("Graph.Node.Comment.PinIconPadding"));
+		const FMargin BubbleOffset = FEditorStyle::GetMargin( TEXT("Graph.Node.Comment.BubbleOffset"));
+		// Conditionally create bubble controls
+		TSharedPtr<SWidget> BubbleControls = SNullWidget::NullWidget;
+
+		if( bEnableBubbleCtrls )
+		{
+			if( bAllowPinning )
+			{
+				SAssignNew( BubbleControls, SVerticalBox )
+				+SVerticalBox::Slot()
+				.Padding( 1.f )
+				.AutoHeight()
+				.VAlign( VAlign_Top )
+				[
+					SNew(SHorizontalBox)
+					+SHorizontalBox::Slot()
+					.AutoWidth()
+					.Padding( PinIconPadding )
+					[
+						SNew( SCheckBox )
+						.Style( &FEditorStyle::Get().GetWidgetStyle<FCheckBoxStyle>( "CommentBubblePin" ))
+						.IsChecked( GraphNode->bCommentBubblePinned ? ESlateCheckBoxState::Checked : ESlateCheckBoxState::Unchecked )
+						.OnCheckStateChanged( this, &SCommentBubble::OnPinStateToggle )
+						.ToolTipText( this, &SCommentBubble::GetScaleButtonTooltip )
+						.Cursor( EMouseCursor::Default )
+					]
+				]
+				+SVerticalBox::Slot()
+				.AutoHeight()
+				.Padding( 1.f )
+				.VAlign( VAlign_Top )
+				[
+					SNew( SCheckBox )
+					.Style( &FEditorStyle::Get().GetWidgetStyle< FCheckBoxStyle >( "CommentBubbleButton" ))
+					.IsChecked( GraphNode->bCommentBubbleVisible ? ESlateCheckBoxState::Checked : ESlateCheckBoxState::Unchecked )
+					.OnCheckStateChanged( this, &SCommentBubble::OnCommentBubbleToggle )
+					.ToolTipText( NSLOCTEXT( "CommentBubble", "ToggleCommentTooltip", "Toggle Comment Bubble" ))
+					.Cursor( EMouseCursor::Default )
+					.ForegroundColor( this, &SCommentBubble::GetToggleButtonColour )
+				];
+			}
+			else
+			{
+				SAssignNew( BubbleControls, SVerticalBox )
+				+SVerticalBox::Slot()
+				.AutoHeight()
+				.Padding( 1.f )
+				.VAlign( VAlign_Top )
+				[
+					SNew( SCheckBox )
+					.Style( &FEditorStyle::Get().GetWidgetStyle< FCheckBoxStyle >( "CommentBubbleButton6" ))
+					.IsChecked( GraphNode->bCommentBubbleVisible ? ESlateCheckBoxState::Checked : ESlateCheckBoxState::Unchecked )
+					.OnCheckStateChanged( this, &SCommentBubble::OnCommentBubbleToggle )
+					.ToolTipText( NSLOCTEXT( "CommentBubble", "ToggleCommentTooltip", "Toggle Comment Bubble" ))
+					.Cursor( EMouseCursor::Default )
+					.ForegroundColor( this, &SCommentBubble::GetToggleButtonColour )
+				];
+			}
+		}
+		// Create the comment bubble widget
+		ChildSlot
+		[
+			SNew(SVerticalBox)
+			+SVerticalBox::Slot()
+			.AutoHeight()
+			[
+				SNew(SHorizontalBox)
+				+SHorizontalBox::Slot()
+				.AutoWidth()
+				[
+					SNew(SOverlay)
+					+SOverlay::Slot()
+					[
+						SNew(SImage)
+						.Image( FEditorStyle::GetBrush( TEXT("Graph.Node.CommentBubble")) )
+						.ColorAndOpacity( ColorAndOpacity )
+					]
+					+SOverlay::Slot()
+					.HAlign(HAlign_Left)
+					.VAlign(VAlign_Center)
+					[
+						SNew(SHorizontalBox)
+						+SHorizontalBox::Slot()
+						.Padding( BubblePadding )
+						.AutoWidth()
+						[
+							SAssignNew(TextBlock, SInlineEditableTextBlock)
+							.Text( this, &SCommentBubble::GetCommentText )
+							.Font( FEditorStyle::GetFontStyle( TEXT("Graph.Node.CommentFont")) )
+							.ColorAndOpacity( FEditorStyle::GetColor("Graph.Node.Comment.TextColor"))
+							.OnTextCommitted( this, &SCommentBubble::OnCommentTextCommitted )
+						]
+						+SHorizontalBox::Slot()
+						.AutoWidth()
+						.HAlign( HAlign_Right )
+						.Padding( 0.f )
+						[
+							BubbleControls.ToSharedRef()
+						]
+					]
+				]
+			]
+			+SVerticalBox::Slot()
+			.AutoHeight()
+			[
+				SNew(SHorizontalBox)
+				+SHorizontalBox::Slot()
+				.Padding( BubbleOffset )
+				.MaxWidth( CommentCalloutArrowBrush->ImageSize.X )
+				[
+					SNew(SImage)
+					.Image( CommentCalloutArrowBrush )
+					.ColorAndOpacity( ColorAndOpacity )
+				]
+			]
+		];
+	}
+	else 
+	{
+		TSharedPtr<SWidget> TitleBarBubble = SNullWidget::NullWidget;
+
+		if( bEnableTitleBarBubble )
+		{
+			const FMargin BubbleOffset = FEditorStyle::GetMargin( TEXT("Graph.Node.Comment.BubbleOffset"));
+			// Create Title bar bubble toggle widget
+			SAssignNew( TitleBarBubble, SHorizontalBox )
+			.Visibility( this, &SCommentBubble::GetToggleButtonVisibility )
+			+SHorizontalBox::Slot()
+			.AutoWidth()
+			.HAlign( HAlign_Center )
+			.VAlign( VAlign_Top )
+			.Padding( BubbleOffset )
+			[
+				SNew( SCheckBox )
+				.Style( &FEditorStyle::Get().GetWidgetStyle< FCheckBoxStyle >( "CommentTitleButton" ))
+				.IsChecked( ESlateCheckBoxState::Unchecked )
+				.OnCheckStateChanged( this, &SCommentBubble::OnCommentBubbleToggle )
+				.ToolTipText( NSLOCTEXT( "CommentBubble", "ToggleCommentTooltip", "Toggle Comment Bubble" ))
+				.Cursor( EMouseCursor::Default )
+				.ForegroundColor( this, &SCommentBubble::GetToggleButtonColour )
+			];
+		}
+		ChildSlot
+		[
+			TitleBarBubble.ToSharedRef()
+		];
+	}
+}
+
+
+FVector2D SCommentBubble::GetOffset() const
+{
+	return FVector2D( 0.f, -GetDesiredSize().Y );
+}
+
+FVector2D SCommentBubble::GetSize() const
+{
+	return GetDesiredSize() + FVector2D( 20.f, 0.f );
+}
+
+bool SCommentBubble::IsBubbleVisible() const
+{
+	EGraphRenderingLOD::Type CurrLOD = GraphLOD.Get();
+	const bool bShowScaled = CurrLOD > EGraphRenderingLOD::LowestDetail;
+	const bool bShowPinned = CurrLOD <= EGraphRenderingLOD::MediumDetail;
+
+	if( bAllowPinning )
+	{
+		return GraphNode->bCommentBubblePinned ? true : bShowScaled;
+	}
+	return bShowPinned;
+}
+
+bool SCommentBubble::IsScalingAllowed() const
+{
+	return !GraphNode->bCommentBubblePinned || !GraphNode->bCommentBubbleVisible;
+}
+
+FText SCommentBubble::GetCommentText() const
+{
+	if( CachedComment != CommentAttribute.Get() )
+	{
+		CachedComment = CommentAttribute.Get();
+		CachedCommentText = FText::FromString( CachedComment );
+	}
+	return CachedCommentText;
+}
+
+FText SCommentBubble::GetScaleButtonTooltip() const
+{
+	if( GraphNode->bCommentBubblePinned )
+	{
+		return NSLOCTEXT( "CommentBubble", "AllowScaleButtonTooltip", "Allow this bubble to scale with zoom" );
+	}
+	else
+	{
+		return NSLOCTEXT( "CommentBubble", "PreventScaleButtonTooltip", "Prevent this bubble scaling with zoom" );
+	}
+}
+
+FSlateColor SCommentBubble::GetToggleButtonColour() const
+{
+	const FLinearColor BubbleColor = ColorAndOpacity.Get().GetSpecifiedColor();
+	return FLinearColor( BubbleColor.R, BubbleColor.G, BubbleColor.B, OpacityValue * OpacityValue );
+}
+
+void SCommentBubble::OnCommentTextCommitted( const FText& NewText, ETextCommit::Type CommitInfo )
+{
+	if( GraphNode )
+	{
+		const FScopedTransaction Transaction( NSLOCTEXT( "CommentBubble", "CommentCommitted", "Comment Changed" ) );
+		GraphNode->Modify();
+		GraphNode->NodeComment = NewText.ToString();
+	}
+}
+
+EVisibility SCommentBubble::GetToggleButtonVisibility() const
+{
+	EVisibility Visibility = EVisibility::Hidden;
+
+	if( OpacityValue > 0.f && !GraphNode->bCommentBubbleVisible )
+	{
+		Visibility = EVisibility::Visible;
+	}
+	return Visibility;
+}
+
+EVisibility SCommentBubble::GetBubbleVisibility() const
+{
+	const bool bIsVisible = !CachedCommentText.IsEmpty() && IsBubbleVisible();
+	return bIsVisible ? EVisibility::Visible : EVisibility::Hidden;
+}
+
+void SCommentBubble::OnCommentBubbleToggle( ESlateCheckBoxState::Type State )
+{
+	if( GraphNode )
+	{
+		const FScopedTransaction Transaction( NSLOCTEXT( "CommentBubble", "BubbleVisibility", "Comment Bubble Visibilty" ) );
+		GraphNode->Modify();
+		GraphNode->bCommentBubbleVisible = !GraphNode->bCommentBubbleVisible;
+		OpacityValue = 0.f;
+		UpdateBubble();
+	}
+}
+
+void SCommentBubble::OnPinStateToggle( ESlateCheckBoxState::Type State ) const
+{
+	if( GraphNode )
+	{
+		const FScopedTransaction Transaction( NSLOCTEXT( "CommentBubble", "BubblePinned", "Comment Bubble Pin" ) );
+		GraphNode->Modify();
+		GraphNode->bCommentBubblePinned = State == ESlateCheckBoxState::Checked;
+	}
+}
\ No newline at end of file
diff --git a/Engine/Source/Editor/GraphEditor/Private/SGraphNode.cpp b/Engine/Source/Editor/GraphEditor/Private/SGraphNode.cpp
--- a/Engine/Source/Editor/GraphEditor/Private/SGraphNode.cpp
+++ b/Engine/Source/Editor/GraphEditor/Private/SGraphNode.cpp
@@ -748,7 +748,7 @@ void SGraphNode::UpdateGraphNode()
 	TSharedPtr<SVerticalBox> InnerVerticalBox;
 	this->ContentScale.Bind( this, &SGraphNode::GetContentScale );
 
-	this->ChildSlot
+	this->GetOrAddSlot( ENodeZone::Center )
 		.HAlign(HAlign_Center)
 		.VAlign(VAlign_Center)
 		[
@@ -794,6 +794,28 @@ void SGraphNode::UpdateGraphNode()
 			]			
 		];
 
+	// Create comment bubble
+	TSharedPtr<SCommentBubble> CommentBubble;
+
+	SAssignNew( CommentBubble, SCommentBubble )
+	.GraphNode( GraphNode )
+	.Text( this, &SGraphNode::GetNodeComment )
+	.ColorAndOpacity( this, &SGraphNode::GetCommentColor )
+	.AllowPinning( true )
+	.EnableTitleBarBubble( true )
+	.EnableBubbleCtrls( true )
+	.GraphLOD( this, &SGraphNode::GetCurrentLOD )
+	.IsGraphNodeHovered( this, &SGraphNode::IsHovered );
+
+	GetOrAddSlot( ENodeZone::TopCenter )
+	.SlotOffset( TAttribute<FVector2D>( CommentBubble.Get(), &SCommentBubble::GetOffset ))
+	.SlotSize( TAttribute<FVector2D>( CommentBubble.Get(), &SCommentBubble::GetSize ))
+	.AllowScaling( TAttribute<bool>( CommentBubble.Get(), &SCommentBubble::IsScalingAllowed ))
+	.VAlign( VAlign_Top )
+	[
+		CommentBubble.ToSharedRef()
+	];
+
 	CreateBelowWidgetControls(MainVerticalBox);
 	CreatePinWidgets();
 	CreateInputSideAddButton(LeftNodeBox);
@@ -1163,11 +1185,6 @@ FString SGraphNode::GetErrorMsgToolTip( ) const
 	return GraphNode->ErrorMsg;
 }
 
-bool SGraphNode::ShouldScaleNodeComment() const
-{
-	return true;
-}
-
 bool SGraphNode::IsNameReadOnly() const
 {
 	return !GraphNode->bCanRenameNode;
@@ -1343,4 +1360,9 @@ void SGraphNode::PopulateMetaTag(FGraphNodeMetaData* TagMeta) const
 		TagMeta->GUID = GraphNode->NodeGuid;
 		TagMeta->FriendlyName = FString::Printf(TEXT("%s in %s"), *GraphNode->GetNodeTitle(ENodeTitleType::ListView).ToString(), *TagMeta->OuterName);		
 	}
+}
+
+EGraphRenderingLOD::Type SGraphNode::GetCurrentLOD() const
+{
+	return OwnerGraphPanelPtr.IsValid() ? OwnerGraphPanelPtr.Pin()->GetCurrentLOD() : EGraphRenderingLOD::DefaultDetail;
 }
\ No newline at end of file
diff --git a/Engine/Source/Editor/GraphEditor/Private/SGraphNodeComment.cpp b/Engine/Source/Editor/GraphEditor/Private/SGraphNodeComment.cpp
--- a/Engine/Source/Editor/GraphEditor/Private/SGraphNodeComment.cpp
+++ b/Engine/Source/Editor/GraphEditor/Private/SGraphNodeComment.cpp
@@ -122,7 +122,7 @@ void SGraphNodeComment::UpdateGraphNode()
 
 	bool bIsSet = GraphNode->IsA(UEdGraphNode_Comment::StaticClass());
 	this->ContentScale.Bind( this, &SGraphNode::GetContentScale );
-	this->ChildSlot
+	this->GetOrAddSlot( ENodeZone::Center )
 		.HAlign(HAlign_Fill)
 		.VAlign(VAlign_Fill)
 		[
@@ -179,6 +179,27 @@ void SGraphNodeComment::UpdateGraphNode()
 				]
 			]			
 		];
+	// Create comment bubble
+	TSharedPtr<SCommentBubble> CommentBubble;
+
+	SAssignNew( CommentBubble, SCommentBubble )
+	.GraphNode( GraphNode )
+	.Text( this, &SGraphNode::GetNodeComment )
+	.ColorAndOpacity( this, &SGraphNodeComment::GetCommentColor )
+	.AllowPinning( true )
+	.EnableTitleBarBubble( true )
+	.EnableBubbleCtrls( true )
+	.GraphLOD( this, &SGraphNode::GetCurrentLOD )
+	.IsGraphNodeHovered( this, &SGraphNode::IsHovered );
+
+	GetOrAddSlot( ENodeZone::TopCenter )
+	.SlotOffset( TAttribute<FVector2D>( CommentBubble.Get(), &SCommentBubble::GetOffset ))
+	.SlotSize( TAttribute<FVector2D>( CommentBubble.Get(), &SCommentBubble::GetSize ))
+	.AllowScaling( TAttribute<bool>( CommentBubble.Get(), &SCommentBubble::IsScalingAllowed ))
+	.VAlign( VAlign_Top )
+	[
+		CommentBubble.ToSharedRef()
+	];
 }
 
 FVector2D SGraphNodeComment::ComputeDesiredSize() const
@@ -321,11 +342,6 @@ FVector2D SGraphNodeComment::GetNodeMaximumSize() const
 	return FVector2D( UserSize.X + 100, UserSize.Y + 100 );
 }
 
-bool SGraphNodeComment::ShouldScaleNodeComment() const 
-{
-	return false;
-}
-
 FSlateColor SGraphNodeComment::GetCommentBodyColor() const
 {
 	UEdGraphNode_Comment* CommentNode = Cast<UEdGraphNode_Comment>(GraphNode);
diff --git a/Engine/Source/Editor/GraphEditor/Private/SGraphNodeDocumentation.cpp b/Engine/Source/Editor/GraphEditor/Private/SGraphNodeDocumentation.cpp
--- a/Engine/Source/Editor/GraphEditor/Private/SGraphNodeDocumentation.cpp
+++ b/Engine/Source/Editor/GraphEditor/Private/SGraphNodeDocumentation.cpp
@@ -130,7 +130,7 @@ void SGraphNodeDocumentation::UpdateGraphNode()
 	TSharedPtr<SWidget> DocumentationPage = CreateDocumentationPage();
 
 	TSharedPtr<SVerticalBox> InnerVerticalBox;
-	ChildSlot
+	GetOrAddSlot( ENodeZone::Center )
 	.HAlign( HAlign_Center )
 	.VAlign( VAlign_Center )
 	[
diff --git a/Engine/Source/Editor/GraphEditor/Private/SGraphNodeKnot.cpp b/Engine/Source/Editor/GraphEditor/Private/SGraphNodeKnot.cpp
--- a/Engine/Source/Editor/GraphEditor/Private/SGraphNodeKnot.cpp
+++ b/Engine/Source/Editor/GraphEditor/Private/SGraphNodeKnot.cpp
@@ -279,7 +279,7 @@ void SGraphNodeKnot::UpdateGraphNode()
 
 	this->ContentScale.Bind( this, &SGraphNode::GetContentScale );
 
-	this->ChildSlot
+	this->GetOrAddSlot( ENodeZone::Center )
 		.HAlign(HAlign_Center)
 		.VAlign(VAlign_Center)
 		[
diff --git a/Engine/Source/Editor/GraphEditor/Private/SGraphNodeResizable.cpp b/Engine/Source/Editor/GraphEditor/Private/SGraphNodeResizable.cpp
--- a/Engine/Source/Editor/GraphEditor/Private/SGraphNodeResizable.cpp
+++ b/Engine/Source/Editor/GraphEditor/Private/SGraphNodeResizable.cpp
@@ -34,7 +34,7 @@ void SGraphNodeResizable::OnMouseEnter( const FGeometry& MyGeometry, const FPoin
 	{
 		FVector2D LocalMouseCoordinates = MyGeometry.AbsoluteToLocal( MouseEvent.GetScreenSpacePosition() );
 		MouseZone = FindMouseZone(LocalMouseCoordinates);
-		SCompoundWidget::OnMouseEnter( MyGeometry, MouseEvent );
+		SNodePanel::SNode::OnMouseEnter( MyGeometry, MouseEvent );
 	}
 }
 
@@ -44,7 +44,7 @@ void SGraphNodeResizable::OnMouseLeave( const FPointerEvent& MouseEvent )
 	{
 		// Reset our mouse zone
 		MouseZone = CRWZ_NotInWindow;
-		SCompoundWidget::OnMouseLeave( MouseEvent );
+		SNodePanel::SNode::OnMouseLeave( MouseEvent );
 	}
 }
 
diff --git a/Engine/Source/Editor/GraphEditor/Private/SGraphPanel.cpp b/Engine/Source/Editor/GraphEditor/Private/SGraphPanel.cpp
--- a/Engine/Source/Editor/GraphEditor/Private/SGraphPanel.cpp
+++ b/Engine/Source/Editor/GraphEditor/Private/SGraphPanel.cpp
@@ -123,8 +123,6 @@ int32 SGraphPanel::OnPaint( const FPaintArgs& Args, const FGeometry& AllottedGeo
 	ArrangeChildNodes(AllottedGeometry, ArrangedChildren);
 
 	// Determine some 'global' settings based on current LOD
-	const bool bDrawScaledCommentBubblesThisFrame = GetCurrentLOD() > EGraphRenderingLOD::LowestDetail;
-	const bool bDrawUnscaledCommentBubblesThisFrame = GetCurrentLOD() <= EGraphRenderingLOD::MediumDetail;
 	const bool bDrawShadowsThisFrame = GetCurrentLOD() > EGraphRenderingLOD::LowestDetail;
 
 	// Because we paint multiple children, we must track the maximum layer id that they produced in case one of our parents
@@ -208,23 +206,6 @@ int32 SGraphPanel::OnPaint( const FPaintArgs& Args, const FGeometry& AllottedGeo
 				// Draw the comments and information popups for this node, if it has any.
 				{
 					float CommentBubbleY = 0.0f;
-
-					const FString NodeComment = ChildNode->GetNodeComment();
-					if (!NodeComment.IsEmpty())
-					{
-						// Comment bubbles have different LOD behavior:
-						//   A comment box comment (bScaleComments=false) will only be shown when zoomed out (the title bar is readable instead when up close)
-						//   A per-node comment (bScaleComments=true) will only be show when zoomed in (it gets too small to read)
-						const bool bScaleComments = ChildNode->ShouldScaleNodeComment();
-						const bool bShowCommentBubble = bScaleComments ? bDrawScaledCommentBubblesThisFrame : bDrawUnscaledCommentBubblesThisFrame;
-
-						if (bShowCommentBubble)
-						{
-							FGeometry CommentGeometry = CurWidget.Geometry.MakeChild(CurWidget.Geometry.Size, FSlateLayoutTransform(bScaleComments ? 1.0f : Inverse(CurWidget.Geometry.Scale)));
-							PaintComment(NodeComment, CommentGeometry, MyClippingRect, OutDrawElements, ChildLayerId, ChildNode->GetNodeCommentColor().GetColor( InWidgetStyle ), /*inout*/ CommentBubbleY, InWidgetStyle);
-						}
-					}
-
 					Context.bSelected = bSelected;
 					TArray<FGraphInformationPopupInfo> Popups;
 
diff --git a/Engine/Source/Editor/GraphEditor/Public/SCommentBubble.h b/Engine/Source/Editor/GraphEditor/Public/SCommentBubble.h
new file mode 100644
--- /dev/null
+++ b/Engine/Source/Editor/GraphEditor/Public/SCommentBubble.h
@@ -0,0 +1,138 @@
+// Copyright 1998-2013 Epic Games, Inc. All Rights Reserved.
+
+#pragma once
+
+DECLARE_DELEGATE_RetVal( bool, FIsGraphNodeHovered );
+
+class GRAPHEDITOR_API SCommentBubble : public SCompoundWidget
+{
+	SLATE_BEGIN_ARGS( SCommentBubble )
+		: _GraphNode( NULL )
+		, _ColorAndOpacity( FLinearColor::White )
+		, _AllowPinning( false )
+		, _EnableTitleBarBubble( false )
+		, _EnableBubbleCtrls( false )
+		, _GraphLOD( EGraphRenderingLOD::DefaultDetail )
+	{}
+
+		/** the GraphNode this bubble should interact with */
+		SLATE_ARGUMENT( UEdGraphNode*, GraphNode )
+
+		/** The comment text for the bubble */
+		SLATE_ATTRIBUTE( FString, Text )
+
+		/** Color and opacity for the comment bubble */
+		SLATE_ATTRIBUTE( FSlateColor, ColorAndOpacity )
+
+		/** Allow bubble to be pinned */
+		SLATE_ARGUMENT( bool, AllowPinning )
+
+		/** Enable the title bar bubble to toggle */
+		SLATE_ARGUMENT( bool, EnableTitleBarBubble )
+
+		/** Enable the controls within the bubble */
+		SLATE_ARGUMENT( bool, EnableBubbleCtrls )
+
+		/** The current level of detail */
+		SLATE_ATTRIBUTE( EGraphRenderingLOD::Type, GraphLOD )
+
+		/** Delegate to determine if the parent node is in the hovered state. */
+		SLATE_EVENT( FIsGraphNodeHovered, IsGraphNodeHovered )
+
+	SLATE_END_ARGS()
+
+	/**
+	 * Construct this widget.  Called by the SNew() Slate macro.
+	 *
+	 * @param	InArgs				Declaration used by the SNew() macro to construct this widget
+	 */
+	void Construct( const FArguments& InArgs );
+
+
+	// Begin SWidget interface
+	virtual FCursorReply OnCursorQuery( const FGeometry& MyGeometry, const FPointerEvent& CursorEvent ) const override;
+	virtual FReply OnMouseButtonDown( const FGeometry& MyGeometry, const FPointerEvent& MouseEvent ) override;
+	virtual FReply OnMouseButtonDoubleClick( const FGeometry& MyGeometry, const FPointerEvent& MouseEvent ) override;
+	virtual void Tick( const FGeometry& AllottedGeometry, const double InCurrentTime, const float InDeltaTime ) override;
+	// End SWidget Interface
+
+	/** Returns the offset from the SNode center slot */
+	FVector2D GetOffset() const;
+
+	/** Returns the bubble size */
+	FVector2D GetSize() const;
+
+	/** Returns if comment bubble is visible */
+	bool IsBubbleVisible() const;
+
+	/** Returns if graph scaling can be applied to this bubble */
+	bool IsScalingAllowed() const;
+
+	/** Called to display/hide the comment bubble */
+	void OnCommentBubbleToggle( ESlateCheckBoxState::Type State );
+
+	/** Called when a node's comment bubble pinned state is changed */
+	void OnPinStateToggle( ESlateCheckBoxState::Type State ) const;
+
+protected:
+
+	/** Called to update the bubble widget layout */
+	void UpdateBubble();
+
+	/** Returns the current comment text, converting from FString into FText */
+	FText GetCommentText() const;
+
+	/** Returns the current scale button tooltip */
+	FText GetScaleButtonTooltip() const;
+
+	/** Called to determine if the comment bubble's current visibility */
+	EVisibility GetBubbleVisibility() const;
+
+	/** Called to determine if the toggle button's current visibility */
+	EVisibility GetToggleButtonVisibility() const;
+
+	/** Returns the color for the toggle bubble including the opacity value */
+	FSlateColor GetToggleButtonColour() const;
+
+	/** Called when the comment text is committed */
+	void OnCommentTextCommitted( const FText& NewText, ETextCommit::Type CommitInfo );
+
+	/** Called when a node's comment bubble visibility state is changed */
+	void OnNodeBubbleVisibilityChange( UEdGraphNode* NodeBeingChanged, bool bIsVisible );
+
+protected:
+
+	/** The GraphNode this widget interacts with */
+	UEdGraphNode* GraphNode;
+	/** Cached inline editable text box */
+	TSharedPtr<SInlineEditableTextBlock> TextBlock;
+
+	/** The Comment Bubble color and opacity value */
+	TAttribute<FSlateColor> ColorAndOpacity;
+	/** Attribute to query node comment */
+	TAttribute<FString> CommentAttribute;
+	/** Attribute to query node pinned state */
+	TAttribute<bool> PinnedState;
+	/** Attribute to query node visibility state */
+	TAttribute<bool> VisibilityState;
+	/** Attribute to query current LOD */
+	TAttribute<EGraphRenderingLOD::Type> GraphLOD;
+
+	/** Delegate to determine if the graph node is currently hovered */
+	FIsGraphNodeHovered IsGraphNodeHovered;
+
+	/** Allow pin behaviour */
+	bool bAllowPinning;
+	/** Allow in bubble controls */
+	bool bEnableBubbleCtrls;
+	/** Enable the title bar bubble toggle */
+	bool bEnableTitleBarBubble;
+	/** Used to control the delay on double click selection */
+	float DoubleClickDelayTime;
+	/** Used to Control hover fade up/down for widgets */
+	float OpacityValue;
+	/** Cached comment */
+	mutable FString CachedComment;
+	/** Cached FText version of the comment */
+	mutable FText CachedCommentText;
+};
diff --git a/Engine/Source/Editor/GraphEditor/Public/SGraphNode.h b/Engine/Source/Editor/GraphEditor/Public/SGraphNode.h
--- a/Engine/Source/Editor/GraphEditor/Public/SGraphNode.h
+++ b/Engine/Source/Editor/GraphEditor/Public/SGraphNode.h
@@ -111,9 +111,6 @@ public:
 	/** @return the tint for the node's comment */
 	FSlateColor GetNodeCommentColor() const;
 
-	/** @return if node's comment should be effected by scaling of panel(zoom) */
-	virtual bool ShouldScaleNodeComment()const;
-
 	/** @return the tooltip to display when over the node */
 	FText GetNodeTooltip() const;
 	
@@ -195,6 +192,9 @@ public:
 	/** Gets the unscaled position of the node from the last tick */
 	FVector2D GetUnscaledPosition() const {return CachedUnscaledPosition;}
 
+	/** Returns the current Node LOD or Highest LOD if unable to query */
+	EGraphRenderingLOD::Type GetCurrentLOD() const;
+
 protected:
 	SGraphNode();
 
@@ -259,6 +259,9 @@ protected:
 	// Should we use low-detail node titles?
 	virtual bool UseLowDetailNodeTitles() const;
 
+	/** Return the desired comment bubble color */
+	virtual FSlateColor GetCommentColor() const { return FLinearColor::White; }
+
 	///// ADVANCED VIEW FUNCTIONS /////
 
 	/** Create button to show/hide advanced pins */
diff --git a/Engine/Source/Editor/GraphEditor/Public/SGraphNodeComment.h b/Engine/Source/Editor/GraphEditor/Public/SGraphNodeComment.h
--- a/Engine/Source/Editor/GraphEditor/Public/SGraphNodeComment.h
+++ b/Engine/Source/Editor/GraphEditor/Public/SGraphNodeComment.h
@@ -27,7 +27,8 @@ public:
 
 	// Begin SGraphNode Interface
 	virtual bool IsNameReadOnly() const override;
-	// ENd SGraphNode Interface
+	virtual FSlateColor GetCommentColor() const override { return GetCommentBodyColor(); }
+	// End SGraphNode Interface
 
 	void Construct( const FArguments& InArgs, UEdGraphNode* InNode );
 
@@ -43,7 +44,6 @@ public:
 protected:
 	// SGraphNode Interface
 	virtual void UpdateGraphNode() override;
-	virtual bool ShouldScaleNodeComment()const override;
 	virtual void PopulateMetaTag(class FGraphNodeMetaData* TagMeta) const override;
 
 	/**
diff --git a/Engine/Source/Editor/GraphEditor/Public/SNodePanel.h b/Engine/Source/Editor/GraphEditor/Public/SNodePanel.h
--- a/Engine/Source/Editor/GraphEditor/Public/SNodePanel.h
+++ b/Engine/Source/Editor/GraphEditor/Public/SNodePanel.h
@@ -211,15 +211,196 @@ public:
  * code desired.  The other main restriction is that each SNode instance must have a unique UObject* associated
  * with it.
  */
+namespace ENodeZone
+{
+	enum Type
+	{
+		TopLeft,
+		TopCenter,
+		TopRight,
+
+		Left,
+		Center,
+		Right,
+
+		BottomLeft,
+		BottomCenter,
+		BottomRight,
+
+		Count
+	};
+}
+
 class GRAPHEDITOR_API SNodePanel : public SPanel
 {
 public:
-	class SNode : public SBorder
+
+	class SNode : public SPanel
 	{
 	public:
 
+		/** A slot that support alignment of content and padding and z-order */
+		class FNodeSlot : public TSlotBase<FNodeSlot>
+		{
+		public:
+			FNodeSlot()
+				: TSlotBase<FNodeSlot>()
+				, HAlignment(HAlign_Fill)
+				, VAlignment(VAlign_Fill)
+				, SlotPadding(0.0f)
+				, Offset( FVector2D::ZeroVector )
+				, AllowScale( true )
+			{ }
+
+			FNodeSlot& HAlign( EHorizontalAlignment InHAlignment )
+			{
+				HAlignment = InHAlignment;
+				return *this;
+			}
+
+			FNodeSlot& VAlign( EVerticalAlignment InVAlignment )
+			{
+				VAlignment = InVAlignment;
+				return *this;
+			}
+
+			FNodeSlot& Padding( const TAttribute<FMargin> InPadding )
+			{
+				SlotPadding = InPadding;
+				return *this;
+			}
+
+			FNodeSlot& SlotOffset( const TAttribute<FVector2D> InOffset )
+			{
+				Offset = InOffset;
+				return *this;
+			}
+
+			FNodeSlot& SlotSize( const TAttribute<FVector2D> InSize )
+			{
+				Size = InSize;
+				return *this;
+			}
+
+			FNodeSlot& AllowScaling( const TAttribute<bool> InAllowScale )
+			{
+				AllowScale = InAllowScale;
+				return *this;
+			}
+
+		public:
+
+			/** The child widget contained in this slot. */
+			ENodeZone::Type Zone;
+			EHorizontalAlignment HAlignment;
+			EVerticalAlignment VAlignment;
+			TAttribute<FMargin> SlotPadding;
+			TAttribute<FVector2D> Offset;
+			TAttribute<FVector2D> Size;
+			TAttribute<bool> AllowScale;
+		};
+
 		typedef TSet<TWeakPtr<SNodePanel::SNode>> FNodeSet;
 
+		// SPanel Interface
+		virtual FChildren* GetChildren() override
+		{
+			return &Children;
+		}
+
+		virtual FVector2D ComputeDesiredSize() const override
+		{
+			for( int32 ChildIndex = 0; ChildIndex < Children.Num(); ++ChildIndex )
+			{
+				if( Children[ ChildIndex ].Zone == ENodeZone::Center )
+				{
+					const FNodeSlot& CenterZone = Children[ ChildIndex ];
+					const EVisibility ChildVisibility = CenterZone.GetWidget()->GetVisibility();
+					if( ChildVisibility != EVisibility::Collapsed )
+					{
+						return ( CenterZone.GetWidget()->GetDesiredSize() + CenterZone.SlotPadding.Get().GetDesiredSize() ) * DesiredSizeScale.Get();
+					}
+				}
+			}
+			return FVector2D::ZeroVector;
+		}
+
+		virtual void OnArrangeChildren( const FGeometry& AllottedGeometry, FArrangedChildren& ArrangedChildren ) const override
+		{
+			for( int32 ChildIndex = 0; ChildIndex < Children.Num(); ++ChildIndex )
+			{
+				const FNodeSlot& CurChild = Children[ChildIndex];
+				const EVisibility ChildVisibility = CurChild.GetWidget()->GetVisibility();
+				if ( ArrangedChildren.Accepts(ChildVisibility) )
+				{
+					const FMargin SlotPadding(CurChild.SlotPadding.Get());
+					const float GeometryScale = CurChild.AllowScale.Get() ? 1.f : 1.f / AllottedGeometry.Scale;
+					FVector2D Size;
+
+					if( CurChild.Size.IsBound() )
+					{
+						Size = CurChild.Size.Get();
+					}
+					else
+					{
+						AlignmentArrangeResult XResult = AlignChild<Orient_Horizontal>(AllottedGeometry.Size.X, CurChild, SlotPadding);
+						AlignmentArrangeResult YResult = AlignChild<Orient_Vertical>(AllottedGeometry.Size.Y, CurChild, SlotPadding);
+						Size = FVector2D( XResult.Size, YResult.Size );
+					}
+					const FArrangedWidget ChildGeom = 
+					AllottedGeometry.MakeChild(
+						CurChild.GetWidget(),
+						CurChild.Offset.Get(),
+						Size,
+						GeometryScale
+					);
+					ArrangedChildren.AddWidget( ChildVisibility, ChildGeom );
+				}
+			}
+		}
+
+		virtual int32 OnPaint( const FPaintArgs& Args, const FGeometry& AllottedGeometry, const FSlateRect& MyClippingRect, FSlateWindowElementList& OutDrawElements, int32 LayerId, const FWidgetStyle& InWidgetStyle, bool bParentEnabled ) const override
+		{
+			FArrangedChildren ArrangedChildren( EVisibility::Visible );
+			{
+				ArrangeChildren( AllottedGeometry, ArrangedChildren );
+			}
+
+			int32 MaxLayerId = LayerId;
+			for( int32 ChildIndex = 0; ChildIndex < ArrangedChildren.Num(); ++ChildIndex )
+			{
+				const FArrangedWidget& CurWidget = ArrangedChildren[ ChildIndex ];
+				FSlateRect ChildClipRect = MyClippingRect.IntersectionWith(CurWidget.Geometry.GetClippingRect());
+
+				if( ChildClipRect.GetSize().Size() > 0 )
+				{
+					const FSlateRect ClipRect = Children[ ChildIndex ].Zone == ENodeZone::Center ? ChildClipRect : MyClippingRect;
+					const int32 CurWidgetsMaxLayerId = CurWidget.Widget->Paint( Args.WithNewParent( this ), CurWidget.Geometry, ClipRect, OutDrawElements, LayerId, InWidgetStyle, ShouldBeEnabled( bParentEnabled ));
+					MaxLayerId = FMath::Max( MaxLayerId, CurWidgetsMaxLayerId );
+				}
+			}
+			return MaxLayerId;
+		}
+		// End of SPanel Interface
+
+		FNodeSlot& GetOrAddSlot( const ENodeZone::Type SlotId )
+		{
+			// Return existing
+			for( int32 ChildIndex = 0; ChildIndex < Children.Num(); ++ChildIndex )
+			{
+				if( Children[ ChildIndex ].Zone == SlotId )
+				{
+					return Children[ ChildIndex ];
+				}
+			}
+			// Add Zone
+			FNodeSlot& NewSlot = *new FNodeSlot();
+			NewSlot.Zone = SlotId;
+			Children.Add( &NewSlot );
+
+			return NewSlot;
+		}
+
 		/**
 		* @param NewPosition	The Node should be relocated to this position in the graph panel
 		* @param NodeFilter		Set of nodes to prevent movement on, after moving successfully a node is added to this set.
@@ -303,10 +484,40 @@ public:
 
 	protected:
 		SNode()
+		: BorderImage( FCoreStyle::Get().GetBrush( "NoBorder" ) )
+		, BorderBackgroundColor( FEditorStyle::GetColor("Graph.ForegroundColor"))
+		, DesiredSizeScale(FVector2D(1,1))
 		{
-			BorderImage = FEditorStyle::GetBrush("NoBorder");
-			ForegroundColor = FEditorStyle::GetColor("Graph.ForegroundColor");
 		}
+
+	protected:
+
+		// SBorder Begin
+		TAttribute<const FSlateBrush*> BorderImage;
+		TAttribute<FSlateColor> BorderBackgroundColor;
+		TAttribute<FVector2D> DesiredSizeScale;
+		/** Whether or not to show the disabled effect when this border is disabled */
+		TAttribute<bool> ShowDisabledEffect;
+		/** Mouse event handlers */
+		FPointerEventHandler MouseButtonDownHandler;
+		FPointerEventHandler MouseButtonUpHandler;
+		FPointerEventHandler MouseMoveHandler;
+		FPointerEventHandler MouseDoubleClickHandler;
+		// SBorder End
+
+		// SPanel Begin
+		/** The layout scale to apply to this widget's contents; useful for animation. */
+		TAttribute<FVector2D> ContentScale;
+		/** The color and opacity to apply to this widget and all its descendants. */
+		TAttribute<FLinearColor> ColorAndOpacity;
+		/** Optional foreground color that will be inherited by all of this widget's contents */
+		TAttribute<FSlateColor> ForegroundColor;
+		// SPanel End
+
+	private:
+
+		TPanelChildren<FNodeSlot> Children;
+
 	};
 
 	SNodePanel();
diff --git a/Engine/Source/Editor/ReferenceViewer/Private/SReferenceNode.cpp b/Engine/Source/Editor/ReferenceViewer/Private/SReferenceNode.cpp
--- a/Engine/Source/Editor/ReferenceViewer/Private/SReferenceNode.cpp
+++ b/Engine/Source/Editor/ReferenceViewer/Private/SReferenceNode.cpp
@@ -5,6 +5,7 @@
 #include "AssetThumbnail.h"
 #include "AssetRegistryModule.h"
 #include "SInlineEditableTextBlock.h"
+#include "SCommentBubble.h"
 
 #define LOCTEXT_NAMESPACE "ReferenceViewer"
 
@@ -68,7 +69,7 @@ void SReferenceNode::UpdateGraphNode()
 	}
 
 	ContentScale.Bind( this, &SReferenceNode::GetContentScale );
-	ChildSlot
+	GetOrAddSlot( ENodeZone::Center )
 	.HAlign(HAlign_Center)
 	.VAlign(VAlign_Center)
 	[
@@ -194,6 +195,26 @@ void SReferenceNode::UpdateGraphNode()
 			]
 		]
 	];
+	// Create comment bubble if comment text is valid
+	GetNodeObj()->bCommentBubbleVisible = !GetNodeObj()->NodeComment.IsEmpty();
+	if( GetNodeObj()->bCommentBubbleVisible )
+	{
+		TSharedPtr<SCommentBubble> CommentBubble;
+
+		SAssignNew( CommentBubble, SCommentBubble )
+		.GraphNode( GraphNode )
+		.Text( this, &SGraphNode::GetNodeComment )
+		.ColorAndOpacity( this, &SReferenceNode::GetCommentColor );
+
+		GetOrAddSlot( ENodeZone::TopCenter )
+		.SlotOffset( TAttribute<FVector2D>( CommentBubble.Get(), &SCommentBubble::GetOffset ))
+		.SlotSize( TAttribute<FVector2D>( CommentBubble.Get(), &SCommentBubble::GetSize ))
+		.AllowScaling( TAttribute<bool>( CommentBubble.Get(), &SCommentBubble::IsScalingAllowed ))
+		.VAlign( VAlign_Top )
+		[
+			CommentBubble.ToSharedRef()
+		];
+	}
 
 	ErrorReporting = ErrorText;
 	ErrorReporting->SetError(ErrorMsg);
diff --git a/Engine/Source/Editor/UnrealEd/Private/EdGraphNode_Comment.cpp b/Engine/Source/Editor/UnrealEd/Private/EdGraphNode_Comment.cpp
--- a/Engine/Source/Editor/UnrealEd/Private/EdGraphNode_Comment.cpp
+++ b/Engine/Source/Editor/UnrealEd/Private/EdGraphNode_Comment.cpp
@@ -24,6 +24,8 @@ UEdGraphNode_Comment::UEdGraphNode_Comment(const FObjectInitializer& ObjectIniti
 	bColorCommentBubble = false;
 	MoveMode = ECommentBoxMode::GroupMovement;
 
+	bCommentBubblePinned = true;
+	bCommentBubbleVisible = true;
 	bCanResizeNode = true;
 	bCanRenameNode = true;
 }
diff --git a/Engine/Source/Editor/WorldBrowser/Private/Tiles/SWorldTileItem.cpp b/Engine/Source/Editor/WorldBrowser/Private/Tiles/SWorldTileItem.cpp
--- a/Engine/Source/Editor/WorldBrowser/Private/Tiles/SWorldTileItem.cpp
+++ b/Engine/Source/Editor/WorldBrowser/Private/Tiles/SWorldTileItem.cpp
@@ -153,7 +153,7 @@ void SWorldTileItem::Construct(const FArguments& InArgs)
 
 	TileModel->ChangedEvent.AddSP(this, &SWorldTileItem::RequestRefresh);
 			
-	ChildSlot
+	GetOrAddSlot( ENodeZone::Center )
 	[
 		ThumbnailViewport.ToSharedRef()
 	];
diff --git a/Engine/Source/Runtime/Core/Public/UObject/ObjectVersion.h b/Engine/Source/Runtime/Core/Public/UObject/ObjectVersion.h
--- a/Engine/Source/Runtime/Core/Public/UObject/ObjectVersion.h
+++ b/Engine/Source/Runtime/Core/Public/UObject/ObjectVersion.h
@@ -640,7 +640,8 @@ enum EUnrealEngineObjectUE4Version
 	VER_UE4_ADD_PROJECTILE_FRICTION_BEHAVIOR,
 	// Add axis settings enum to MovementComponent.
 	VER_UE4_MOVEMENTCOMPONENT_AXIS_SETTINGS,
-
+	// Switch to new interactive comments, requires boundry conversion to preserve previous states
+	VER_UE4_GRAPH_INTERACTIVE_COMMENTBUBBLES,
 	// -----<new versions can be added before this line>-------------------------------------------------
 	// - this needs to be the last line (see note below)
 	VER_UE4_AUTOMATIC_VERSION_PLUS_ONE,
diff --git a/Engine/Source/Runtime/Engine/Classes/EdGraph/EdGraphNode.h b/Engine/Source/Runtime/Engine/Classes/EdGraph/EdGraphNode.h
--- a/Engine/Source/Runtime/Engine/Classes/EdGraph/EdGraphNode.h
+++ b/Engine/Source/Runtime/Engine/Classes/EdGraph/EdGraphNode.h
@@ -127,6 +127,14 @@ class ENGINE_API UEdGraphNode : public UObject
 	UPROPERTY()
 	FString NodeComment;
 
+	/** Comment bubble pinned state */
+	UPROPERTY()
+	bool bCommentBubblePinned;
+
+	/** Comment bubble visibility */
+	UPROPERTY()
+	bool bCommentBubbleVisible;
+
 	/** Flag to store node specific compile error/warning*/
 	UPROPERTY()
 	int32 ErrorType;
diff --git a/Engine/Source/Runtime/Engine/Private/EdGraph/EdGraphNode.cpp b/Engine/Source/Runtime/Engine/Private/EdGraph/EdGraphNode.cpp
--- a/Engine/Source/Runtime/Engine/Private/EdGraph/EdGraphNode.cpp
+++ b/Engine/Source/Runtime/Engine/Private/EdGraph/EdGraphNode.cpp
@@ -45,6 +45,8 @@ UEdGraphNode::UEdGraphNode(const FObjectInitializer& ObjectInitializer)
 {
 
 #if WITH_EDITORONLY_DATA
+	bCommentBubblePinned = false;
+	bCommentBubbleVisible = false;
 	bCanResizeNode = false;
 #endif // WITH_EDITORONLY_DATA
 }
@@ -253,6 +255,11 @@ void UEdGraphNode::PostLoad()
 		// Generate new one
 		CreateNewGuid();
 	}
+	// Moving to the new style comments requires conversion to preserve previous state
+	if(GetLinkerUE4Version() < VER_UE4_GRAPH_INTERACTIVE_COMMENTBUBBLES)
+	{
+		bCommentBubbleVisible = !NodeComment.IsEmpty();
+	}
 }
 
 void UEdGraphNode::CreateNewGuid()
