Build fix for non unity

diff --git a/Engine/Source/Editor/GraphEditor/Private/GraphEditorCommon.h b/Engine/Source/Editor/GraphEditor/Private/GraphEditorCommon.h
--- a/Engine/Source/Editor/GraphEditor/Private/GraphEditorCommon.h
+++ b/Engine/Source/Editor/GraphEditor/Private/GraphEditorCommon.h
@@ -12,6 +12,7 @@
 #include "GraphEditorSettings.h"
 #include "SNodePanel.h"
 #include "SGraphPanel.h"
+#include "SCommentBubble.h"
 #include "SGraphNode.h"
 #include "SGraphNodeResizable.h"
 #include "SGraphPin.h"
diff --git a/Engine/Source/Editor/GraphEditor/Private/SCommentBubble.cpp b/Engine/Source/Editor/GraphEditor/Private/SCommentBubble.cpp
--- a/Engine/Source/Editor/GraphEditor/Private/SCommentBubble.cpp
+++ b/Engine/Source/Editor/GraphEditor/Private/SCommentBubble.cpp
@@ -3,6 +3,7 @@
 #include "GraphEditorCommon.h"
 #include "SCommentBubble.h"
 #include "SInlineEditableTextBlock.h"
+#include "ScopedTransaction.h"
 
 namespace SCommentBubbleDefs
 {
