[UE-3442] Can now use scientific notation in float inputs for BP nodes.

diff --git a/Engine/Source/Editor/GraphEditor/Private/KismetPins/SGraphPinNum.cpp b/Engine/Source/Editor/GraphEditor/Private/KismetPins/SGraphPinNum.cpp
--- a/Engine/Source/Editor/GraphEditor/Private/KismetPins/SGraphPinNum.cpp
+++ b/Engine/Source/Editor/GraphEditor/Private/KismetPins/SGraphPinNum.cpp
@@ -5,6 +5,7 @@
 #include "SGraphPin.h"
 #include "SGraphPinString.h"
 #include "SGraphPinNum.h"
+#include "DefaultValueHelper.h"
 
 void SGraphPinNum::Construct(const FArguments& InArgs, UEdGraphPin* InGraphPinObj)
 {
@@ -14,7 +15,7 @@ void SGraphPinNum::Construct(const FArguments& InArgs, UEdGraphPin* InGraphPinOb
 void SGraphPinNum::SetTypeInValue(const FText& NewTypeInValue, ETextCommit::Type /*CommitInfo*/)
 {
 	FString TypeValueString = NewTypeInValue.ToString();
-	if (TypeValueString.IsNumeric())
+	if (FDefaultValueHelper::IsStringValidFloat(TypeValueString))
 	{
 		//This is rubbish.
 		//TODO: Make not rubbish.
