Modified the comment bubbles on SGraphNodeKnot to display on node hover.
#UE-2353 - Feature Request: BP: Allow users to hide comments on nodes
#Branch UE4
#Proj GraphEditor

diff --git a/Engine/Source/Editor/GraphEditor/Private/SCommentBubble.cpp b/Engine/Source/Editor/GraphEditor/Private/SCommentBubble.cpp
--- a/Engine/Source/Editor/GraphEditor/Private/SCommentBubble.cpp
+++ b/Engine/Source/Editor/GraphEditor/Private/SCommentBubble.cpp
@@ -284,7 +284,7 @@ FVector2D SCommentBubble::GetOffset() const
 
 FVector2D SCommentBubble::GetSize() const
 {
-	return GetDesiredSize() + FVector2D( 20.f, 0.f );
+	return GetDesiredSize();
 }
 
 bool SCommentBubble::IsBubbleVisible() const
diff --git a/Engine/Source/Editor/GraphEditor/Private/SGraphNodeKnot.cpp b/Engine/Source/Editor/GraphEditor/Private/SGraphNodeKnot.cpp
--- a/Engine/Source/Editor/GraphEditor/Private/SGraphNodeKnot.cpp
+++ b/Engine/Source/Editor/GraphEditor/Private/SGraphNodeKnot.cpp
@@ -319,16 +319,24 @@ void SGraphNodeKnot::UpdateGraphNode()
 						]
 					]
 				]
-
-				+SVerticalBox::Slot()
-				.VAlign(VAlign_Fill)
-				.HAlign(HAlign_Fill)
-				[
-					InlineEditableText.ToSharedRef()
-				]
 			]
 		];
-
+	// Create comment bubble
+	SAssignNew( CommentBubble, SCommentBubble )
+	.GraphNode( GraphNode )
+	.Text( this, &SGraphNode::GetNodeComment )
+	.ColorAndOpacity( FLinearColor::White )
+	.GraphLOD( this, &SGraphNode::GetCurrentLOD )
+	.IsGraphNodeHovered( this, &SGraphNode::IsHovered );
+
+	GetOrAddSlot( ENodeZone::TopCenter )
+	.SlotOffset( TAttribute<FVector2D>( this, &SGraphNodeKnot::GetCommentOffset ))
+	.SlotSize( TAttribute<FVector2D>( CommentBubble.Get(), &SCommentBubble::GetSize ))
+	.AllowScaling( TAttribute<bool>( CommentBubble.Get(), &SCommentBubble::IsScalingAllowed ))
+	.VAlign( VAlign_Top )
+	[
+		CommentBubble.ToSharedRef()
+	];
 	CreatePinWidgets();
 }
 
@@ -373,3 +381,27 @@ void SGraphNodeKnot::AddPin(const TSharedRef<SGraphPin>& PinToAdd)
 		OutputPins.Add(PinToAdd);
 	}
 }
+
+FVector2D SGraphNodeKnot::GetCommentOffset() const
+{
+	static const FVector2D CenterAdjust( 8.f, 0.f );
+	return CommentBubble->GetOffset() + CenterAdjust;
+}
+
+void SGraphNodeKnot::OnMouseEnter( const FGeometry& MyGeometry, const FPointerEvent& MouseEvent )
+{
+	if( !GraphNode->NodeComment.IsEmpty() )
+	{
+		GraphNode->bCommentBubbleVisible = true;
+		CommentBubble->UpdateBubble();
+	}
+}
+
+void SGraphNodeKnot::OnMouseLeave( const FPointerEvent& MouseEvent )
+{
+	if( !GraphNode->NodeComment.IsEmpty() )
+	{
+		GraphNode->bCommentBubbleVisible = false;
+		CommentBubble->UpdateBubble();
+	}
+}
diff --git a/Engine/Source/Editor/GraphEditor/Private/SGraphNodeKnot.h b/Engine/Source/Editor/GraphEditor/Private/SGraphNodeKnot.h
--- a/Engine/Source/Editor/GraphEditor/Private/SGraphNodeKnot.h
+++ b/Engine/Source/Editor/GraphEditor/Private/SGraphNodeKnot.h
@@ -5,6 +5,8 @@
 #include "SGraphNodeDefault.h"
 #include "K2Node_Knot.h"
 
+class SCommentBubble;
+
 class SGraphNodeKnot : public SGraphNodeDefault
 {
 public:
@@ -20,4 +22,18 @@ public:
 	virtual void AddPin(const TSharedRef<SGraphPin>& PinToAdd) override;
 	virtual void RequestRenameOnSpawn() override { }
 	// End of SGraphNode interface
+
+	// SWidget interface
+	virtual void OnMouseEnter( const FGeometry& MyGeometry, const FPointerEvent& MouseEvent ) override;
+	virtual void OnMouseLeave( const FPointerEvent& MouseEvent ) override;
+	// End of SWidget interface
+
+	/** Returns Offset to center comment on the node's only pin */
+	FVector2D GetCommentOffset() const;
+
+protected:
+
+	/** SharedPtr to comment bubble */
+	TSharedPtr<SCommentBubble> CommentBubble;
+
 };
diff --git a/Engine/Source/Editor/GraphEditor/Public/SCommentBubble.h b/Engine/Source/Editor/GraphEditor/Public/SCommentBubble.h
--- a/Engine/Source/Editor/GraphEditor/Public/SCommentBubble.h
+++ b/Engine/Source/Editor/GraphEditor/Public/SCommentBubble.h
@@ -74,11 +74,11 @@ class GRAPHEDITOR_API SCommentBubble : public SCompoundWidget
 	/** Called when a node's comment bubble pinned state is changed */
 	void OnPinStateToggle( ESlateCheckBoxState::Type State ) const;
 
-protected:
-
 	/** Called to update the bubble widget layout */
 	void UpdateBubble();
 
+protected:
+
 	/** Returns the current comment text, converting from FString into FText */
 	FText GetCommentText() const;
 
@@ -97,9 +97,6 @@ protected:
 	/** Called when the comment text is committed */
 	void OnCommentTextCommitted( const FText& NewText, ETextCommit::Type CommitInfo );
 
-	/** Called when a node's comment bubble visibility state is changed */
-	void OnNodeBubbleVisibilityChange( UEdGraphNode* NodeBeingChanged, bool bIsVisible );
-
 protected:
 
 	/** The GraphNode this widget interacts with */
diff --git a/Engine/Source/Editor/GraphEditor/Public/SNodePanel.h b/Engine/Source/Editor/GraphEditor/Public/SNodePanel.h
--- a/Engine/Source/Editor/GraphEditor/Public/SNodePanel.h
+++ b/Engine/Source/Editor/GraphEditor/Public/SNodePanel.h
@@ -337,7 +337,7 @@ public:
 					const float GeometryScale = CurChild.AllowScale.Get() ? 1.f : 1.f / AllottedGeometry.Scale;
 					FVector2D Size;
 
-					if( CurChild.Size.IsBound() )
+					if( CurChild.Size.IsSet() )
 					{
 						Size = CurChild.Size.Get();
 					}
