Modified the comment nodes to only consider their contents when they already know their own size. Previously they determined the contents on a null size which resulted in this issue.
#UE-4912 - Create comment around selected nodes fails to move nodes
#Branch UE4
#Proj GraphEditor

diff --git a/Engine/Source/Editor/GraphEditor/Private/SGraphNodeComment.cpp b/Engine/Source/Editor/GraphEditor/Private/SGraphNodeComment.cpp
--- a/Engine/Source/Editor/GraphEditor/Private/SGraphNodeComment.cpp
+++ b/Engine/Source/Editor/GraphEditor/Private/SGraphNodeComment.cpp
@@ -232,43 +232,47 @@ FReply SGraphNodeComment::OnMouseButtonDoubleClick( const FGeometry& InMyGeometr
 
 void SGraphNodeComment::HandleSelection(bool bSelected, bool bUpdateNodesUnderComment) const
 {
-	if ((!this->bIsSelected && bSelected) || bUpdateNodesUnderComment)
+	const FVector2D NodeSize = GetDesiredSize();
+	// we only want to do this after the comment has a valid desired size
+	if( !NodeSize.IsZero() )
 	{
-		SGraphNodeComment* Comment = const_cast<SGraphNodeComment*> (this);
-
-		UEdGraphNode_Comment* CommentNode = Cast<UEdGraphNode_Comment>(GraphNode);
-
-		if (CommentNode)
+		if ((!this->bIsSelected && bSelected) || bUpdateNodesUnderComment)
 		{
+			SGraphNodeComment* Comment = const_cast<SGraphNodeComment*> (this);
 
-			// Get our geo
-			const FVector2D NodePosition = GetPosition();
-			const FVector2D NodeSize = GetDesiredSize();
-			const FSlateRect CommentRect( NodePosition.X, NodePosition.Y, NodePosition.X + NodeSize.X, NodePosition.Y + NodeSize.Y );
+			UEdGraphNode_Comment* CommentNode = Cast<UEdGraphNode_Comment>(GraphNode);
 
-			TSharedPtr< SGraphPanel > Panel = Comment->GetOwnerPanel();
-			FChildren* PanelChildren = Panel->GetChildren();
-			int32 NumChildren = PanelChildren->Num();
-			CommentNode->ClearNodesUnderComment();
-
-			for ( int32 NodeIndex=0; NodeIndex < NumChildren; ++NodeIndex )
+			if (CommentNode)
 			{
-				const TSharedRef<SGraphNode> SomeNodeWidget = StaticCastSharedRef<SGraphNode>(PanelChildren->GetChildAt(NodeIndex));
 
-				UObject* GraphObject = SomeNodeWidget->GetObjectBeingDisplayed();
+				// Get our geo
+				const FVector2D NodePosition = GetPosition();
+				const FSlateRect CommentRect( NodePosition.X, NodePosition.Y, NodePosition.X + NodeSize.X, NodePosition.Y + NodeSize.Y );
 
-				const FVector2D SomeNodePosition = SomeNodeWidget->GetPosition();
-				const FVector2D SomeNodeSize = SomeNodeWidget->GetDesiredSize();
+				TSharedPtr< SGraphPanel > Panel = Comment->GetOwnerPanel();
+				FChildren* PanelChildren = Panel->GetChildren();
+				int32 NumChildren = PanelChildren->Num();
+				CommentNode->ClearNodesUnderComment();
 
-				const FSlateRect NodeGeometryGraphSpace( SomeNodePosition.X, SomeNodePosition.Y, SomeNodePosition.X + SomeNodeSize.X, SomeNodePosition.Y + SomeNodeSize.Y );
-				if ( FSlateRect::IsRectangleContained( CommentRect, NodeGeometryGraphSpace ) )
+				for ( int32 NodeIndex=0; NodeIndex < NumChildren; ++NodeIndex )
 				{
-					CommentNode->AddNodeUnderComment(GraphObject);
+					const TSharedRef<SGraphNode> SomeNodeWidget = StaticCastSharedRef<SGraphNode>(PanelChildren->GetChildAt(NodeIndex));
+
+					UObject* GraphObject = SomeNodeWidget->GetObjectBeingDisplayed();
+
+					const FVector2D SomeNodePosition = SomeNodeWidget->GetPosition();
+					const FVector2D SomeNodeSize = SomeNodeWidget->GetDesiredSize();
+
+					const FSlateRect NodeGeometryGraphSpace( SomeNodePosition.X, SomeNodePosition.Y, SomeNodePosition.X + SomeNodeSize.X, SomeNodePosition.Y + SomeNodeSize.Y );
+					if ( FSlateRect::IsRectangleContained( CommentRect, NodeGeometryGraphSpace ) )
+					{
+						CommentNode->AddNodeUnderComment(GraphObject);
+					}
 				}
 			}
 		}
+		bIsSelected = bSelected;
 	}
-	this->bIsSelected = bSelected;
 }
 
 const FSlateBrush* SGraphNodeComment::GetShadowBrush(bool bSelected) const
