Fix for dark comment colors obscuring text. The text should now select a light or dark foreground color based on the bubble luminance.
#Branche UE4
#Proj GraphEditor

diff --git a/Engine/Source/Editor/EditorStyle/Private/SlateEditorStyle.cpp b/Engine/Source/Editor/EditorStyle/Private/SlateEditorStyle.cpp
--- a/Engine/Source/Editor/EditorStyle/Private/SlateEditorStyle.cpp
+++ b/Engine/Source/Editor/EditorStyle/Private/SlateEditorStyle.cpp
@@ -3227,7 +3227,6 @@ void FSlateEditorStyle::FStyle::SetupGraphEditorStyles()
 		Set( "Graph.Node.CommentBubble", new BOX_BRUSH( "Old/Graph/CommentBubble", FMargin(8/32.0f) ) );
 		Set( "Graph.Node.CommentArrow", new IMAGE_BRUSH( "Old/Graph/CommentBubbleArrow", FVector2D(8,8) ) );
 		Set( "Graph.Node.CommentFont", TTF_CORE_FONT( "Fonts/Roboto-Regular", 10 ) );
-		Set( "Graph.Node.Comment.TextColor", FLinearColor::Black );
 		Set( "Graph.Node.Comment.BubbleOffset", FMargin(8,0,0,0) );
 		Set( "Graph.Node.Comment.PinIconPadding", FMargin(0,2,0,0) );
 		Set( "Graph.Node.Comment.BubblePadding", FMargin(4,4) );
@@ -3254,12 +3253,12 @@ void FSlateEditorStyle::FStyle::SetupGraphEditorStyles()
 
 		const FCheckBoxStyle CommentBubblePin = FCheckBoxStyle()
 			.SetCheckBoxType( ESlateCheckBoxType::CheckBox )
-			.SetUncheckedImage( IMAGE_BRUSH( "Icons/icon_Blueprint_CommentBubbleUnPin_16x", FVector2D(10,10), FLinearColor(1.f, 1.f, 1.f, 0.3f)))
-			.SetUncheckedHoveredImage( IMAGE_BRUSH( "Icons/icon_Blueprint_CommentBubbleUnPin_16x", FVector2D(10,10), FLinearColor(1.f, 1.f, 1.f, 0.5f)))
-			.SetUncheckedPressedImage( IMAGE_BRUSH( "Icons/icon_Blueprint_CommentBubblePin_16x", FVector2D(10,10), FLinearColor(1.f, 1.f, 1.f, 0.8f)))
-			.SetCheckedImage( IMAGE_BRUSH( "Icons/icon_Blueprint_CommentBubblePin_16x", FVector2D(10,10), FLinearColor(1.f, 1.f, 1.f, 0.3f)))
-			.SetCheckedHoveredImage( IMAGE_BRUSH( "Icons/icon_Blueprint_CommentBubblePin_16x", FVector2D(10,10), FLinearColor(1.f, 1.f, 1.f, 0.5f)))
-			.SetCheckedPressedImage( IMAGE_BRUSH( "Icons/icon_Blueprint_CommentBubbleUnPin_16x", FVector2D(10,10), FLinearColor(1.f, 1.f, 1.f, 0.8f)));
+			.SetUncheckedImage( IMAGE_BRUSH( "Icons/icon_Blueprint_CommentBubbleUnPin_16x", FVector2D(10,10), FLinearColor::White))
+			.SetUncheckedHoveredImage( IMAGE_BRUSH( "Icons/icon_Blueprint_CommentBubbleUnPin_16x", FVector2D(10,10), FLinearColor::White))
+			.SetUncheckedPressedImage( IMAGE_BRUSH( "Icons/icon_Blueprint_CommentBubblePin_16x", FVector2D(10,10), FLinearColor::White))
+			.SetCheckedImage( IMAGE_BRUSH( "Icons/icon_Blueprint_CommentBubblePin_16x", FVector2D(10,10), FLinearColor::White))
+			.SetCheckedHoveredImage( IMAGE_BRUSH( "Icons/icon_Blueprint_CommentBubblePin_16x", FVector2D(10,10), FLinearColor::White))
+			.SetCheckedPressedImage( IMAGE_BRUSH( "Icons/icon_Blueprint_CommentBubbleUnPin_16x", FVector2D(10,10), FLinearColor::White));
 		Set( "CommentBubblePin", CommentBubblePin );
 
 
diff --git a/Engine/Source/Editor/GraphEditor/Private/SCommentBubble.cpp b/Engine/Source/Editor/GraphEditor/Private/SCommentBubble.cpp
--- a/Engine/Source/Editor/GraphEditor/Private/SCommentBubble.cpp
+++ b/Engine/Source/Editor/GraphEditor/Private/SCommentBubble.cpp
@@ -21,6 +21,15 @@ namespace SCommentBubbleDefs
 
 	/** Height of the arrow connecting the bubble to the node */
 	static const float BubbleArrowHeight = 8.f;
+
+	/** Luminance CoEficients */
+	static const FLinearColor LuminanceCoEff( 0.2126f, 0.7152f, 0.0722f, 0.f );
+
+	/** Light foreground color */
+	static const FLinearColor LightForegroundClr( 0.f, 0.f, 0.f, 0.65f );
+
+	/** Dark foreground color */
+	static const FLinearColor DarkForegroundClr( 1.f, 1.f, 1.f, 0.65f );
 };
 
 
@@ -92,6 +101,11 @@ void SCommentBubble::Tick( const FGeometry& AllottedGeometry, const double InCur
 
 	bool bIsHovered = IsHovered();
 
+	const FLinearColor BubbleColor = ColorAndOpacity.Get().GetSpecifiedColor() * SCommentBubbleDefs::LuminanceCoEff;
+	const float BubbleLuminance = BubbleColor.R + BubbleColor.G + BubbleColor.B;
+	ForegroundColor = BubbleLuminance < 0.5f ? SCommentBubbleDefs::DarkForegroundClr : SCommentBubbleDefs::LightForegroundClr;
+
+
 	if( DoubleClickDelayTime >= 0.f )
 	{
 		DoubleClickDelayTime += InDeltaTime;
@@ -152,6 +166,7 @@ void SCommentBubble::UpdateBubble()
 						.OnCheckStateChanged( this, &SCommentBubble::OnPinStateToggle )
 						.ToolTipText( this, &SCommentBubble::GetScaleButtonTooltip )
 						.Cursor( EMouseCursor::Default )
+						.ForegroundColor( this, &SCommentBubble::GetForegroundColor )
 					]
 				]
 				+SVerticalBox::Slot()
@@ -165,7 +180,7 @@ void SCommentBubble::UpdateBubble()
 					.OnCheckStateChanged( this, &SCommentBubble::OnCommentBubbleToggle )
 					.ToolTipText( NSLOCTEXT( "CommentBubble", "ToggleCommentTooltip", "Toggle Comment Bubble" ))
 					.Cursor( EMouseCursor::Default )
-					.ForegroundColor( this, &SCommentBubble::GetToggleButtonColour )
+					.ForegroundColor( FLinearColor::White )
 				];
 			}
 			else
@@ -177,12 +192,12 @@ void SCommentBubble::UpdateBubble()
 				.VAlign( VAlign_Top )
 				[
 					SNew( SCheckBox )
-					.Style( &FEditorStyle::Get().GetWidgetStyle< FCheckBoxStyle >( "CommentBubbleButton6" ))
+					.Style( &FEditorStyle::Get().GetWidgetStyle< FCheckBoxStyle >( "CommentBubbleButton" ))
 					.IsChecked( GraphNode->bCommentBubbleVisible ? ESlateCheckBoxState::Checked : ESlateCheckBoxState::Unchecked )
 					.OnCheckStateChanged( this, &SCommentBubble::OnCommentBubbleToggle )
 					.ToolTipText( NSLOCTEXT( "CommentBubble", "ToggleCommentTooltip", "Toggle Comment Bubble" ))
 					.Cursor( EMouseCursor::Default )
-					.ForegroundColor( this, &SCommentBubble::GetToggleButtonColour )
+					.ForegroundColor( FLinearColor::White )
 				];
 			}
 		}
@@ -216,7 +231,7 @@ void SCommentBubble::UpdateBubble()
 							SAssignNew(TextBlock, SInlineEditableTextBlock)
 							.Text( this, &SCommentBubble::GetCommentText )
 							.Font( FEditorStyle::GetFontStyle( TEXT("Graph.Node.CommentFont")) )
-							.ColorAndOpacity( FEditorStyle::GetColor("Graph.Node.Comment.TextColor"))
+							.ColorAndOpacity( this, &SCommentBubble::GetForegroundColor)
 							.OnTextCommitted( this, &SCommentBubble::OnCommentTextCommitted )
 						]
 						+SHorizontalBox::Slot()
@@ -266,7 +281,7 @@ void SCommentBubble::UpdateBubble()
 				.OnCheckStateChanged( this, &SCommentBubble::OnCommentBubbleToggle )
 				.ToolTipText( NSLOCTEXT( "CommentBubble", "ToggleCommentTooltip", "Toggle Comment Bubble" ))
 				.Cursor( EMouseCursor::Default )
-				.ForegroundColor( this, &SCommentBubble::GetToggleButtonColour )
+				.ForegroundColor( this, &SCommentBubble::GetToggleButtonColor )
 			];
 		}
 		ChildSlot
@@ -327,10 +342,10 @@ FText SCommentBubble::GetScaleButtonTooltip() const
 	}
 }
 
-FSlateColor SCommentBubble::GetToggleButtonColour() const
+FSlateColor SCommentBubble::GetToggleButtonColor() const
 {
 	const FLinearColor BubbleColor = ColorAndOpacity.Get().GetSpecifiedColor();
-	return FLinearColor( BubbleColor.R, BubbleColor.G, BubbleColor.B, OpacityValue * OpacityValue );
+	return FLinearColor( 1.f, 1.f, 1.f, OpacityValue * OpacityValue );
 }
 
 void SCommentBubble::OnCommentTextCommitted( const FText& NewText, ETextCommit::Type CommitInfo )
diff --git a/Engine/Source/Editor/GraphEditor/Public/SCommentBubble.h b/Engine/Source/Editor/GraphEditor/Public/SCommentBubble.h
--- a/Engine/Source/Editor/GraphEditor/Public/SCommentBubble.h
+++ b/Engine/Source/Editor/GraphEditor/Public/SCommentBubble.h
@@ -92,7 +92,10 @@ protected:
 	EVisibility GetToggleButtonVisibility() const;
 
 	/** Returns the color for the toggle bubble including the opacity value */
-	FSlateColor GetToggleButtonColour() const;
+	FSlateColor GetToggleButtonColor() const;
+
+	/** Returns the foreground color for the text and buttons, taking into account the bubble color */
+	FSlateColor GetForegroundColor() const { return ForegroundColor; }
 
 	/** Called when the comment text is committed */
 	void OnCommentTextCommitted( const FText& NewText, ETextCommit::Type CommitInfo );
@@ -118,6 +121,8 @@ protected:
 	/** Delegate to determine if the graph node is currently hovered */
 	FIsGraphNodeHovered IsGraphNodeHovered;
 
+	/** Current Foreground Color */
+	FLinearColor ForegroundColor;
 	/** Allow pin behaviour */
 	bool bAllowPinning;
 	/** Allow in bubble controls */
