Blueprint/UMG - Fixing the class picker so that it's not clipping the contents of the dropdown list.

diff --git a/Engine/Source/Editor/GraphEditor/Private/KismetPins/SGraphPinClass.cpp b/Engine/Source/Editor/GraphEditor/Private/KismetPins/SGraphPinClass.cpp
--- a/Engine/Source/Editor/GraphEditor/Private/KismetPins/SGraphPinClass.cpp
+++ b/Engine/Source/Editor/GraphEditor/Private/KismetPins/SGraphPinClass.cpp
@@ -80,8 +80,9 @@ TSharedRef<SWidget> SGraphPinClass::GenerateAssetPicker()
 		.WidthOverride(280)
 		[
 			SNew(SVerticalBox)
-			+SVerticalBox::Slot()
-			.AutoHeight()
+
+			+ SVerticalBox::Slot()
+			.FillHeight(1.0f)
 			.MaxHeight(500)
 			[ 
 				SNew(SBorder)
