Modified the comment bubbles to support multiline editing.
#UE-2051 - ROCKET: TASK: The ability to write multi-line comments.
#Branch UE4
#Proj GraphEditor

diff --git a/Engine/Source/Editor/GraphEditor/Private/SCommentBubble.cpp b/Engine/Source/Editor/GraphEditor/Private/SCommentBubble.cpp
--- a/Engine/Source/Editor/GraphEditor/Private/SCommentBubble.cpp
+++ b/Engine/Source/Editor/GraphEditor/Private/SCommentBubble.cpp
@@ -30,6 +30,10 @@ namespace SCommentBubbleDefs
 
 	/** Dark foreground color */
 	static const FLinearColor DarkForegroundClr( 1.f, 1.f, 1.f, 0.65f );
+
+	/** Clear text box background color */
+	static const FLinearColor TextClearBackground( 0.f, 0.f, 0.f, 0.f );
+
 };
 
 
@@ -47,7 +51,6 @@ void SCommentBubble::Construct( const FArguments& InArgs )
 	GraphLOD				= InArgs._GraphLOD;
 	IsGraphNodeHovered		= InArgs._IsGraphNodeHovered;
 
-	DoubleClickDelayTime	= SCommentBubbleDefs::DoubleClickDisable;
 	OpacityValue			= 0.f;
 
 	// Create Widget
@@ -66,35 +69,6 @@ FCursorReply SCommentBubble::OnCursorQuery( const FGeometry& MyGeometry, const F
 	return FCursorReply::Cursor( EMouseCursor::Default );
 }
 
-FReply SCommentBubble::OnMouseButtonDown( const FGeometry& MyGeometry, const FPointerEvent& CursorEvent )
-{
-	if( TextBlock.IsValid() )
-	{
-		const FVector2D Size( GetDesiredSize().X, GetDesiredSize().Y - SCommentBubbleDefs::BubbleArrowHeight );
-		const FSlateRect TestRect( MyGeometry.AbsolutePosition, MyGeometry.AbsolutePosition + Size );
-
-		if( TestRect.ContainsPoint( CursorEvent.GetScreenSpacePosition() ))
-		{
-			if( DoubleClickDelayTime >= SCommentBubbleDefs::DoubleClickDelay )
-			{
-				TextBlock->EnterEditingMode();
-				DoubleClickDelayTime = SCommentBubbleDefs::DoubleClickDisable;
-			}
-			else
-			{
-				FSlateApplication::Get().SetKeyboardFocus( TextBlock, EKeyboardFocusCause::SetDirectly );
-				DoubleClickDelayTime = 0.f;
-			}
-		}
-	}
-	return FReply::Handled();
-}
-
-FReply SCommentBubble::OnMouseButtonDoubleClick( const FGeometry& MyGeometry, const FPointerEvent& MouseEvent )
-{
-	return FReply::Handled();
-}
-
 void SCommentBubble::Tick( const FGeometry& AllottedGeometry, const double InCurrentTime, const float InDeltaTime )
 {
 	SCompoundWidget::Tick( AllottedGeometry, InCurrentTime, InDeltaTime );
@@ -105,16 +79,6 @@ void SCommentBubble::Tick( const FGeometry& AllottedGeometry, const double InCur
 	const float BubbleLuminance = BubbleColor.R + BubbleColor.G + BubbleColor.B;
 	ForegroundColor = BubbleLuminance < 0.5f ? SCommentBubbleDefs::DarkForegroundClr : SCommentBubbleDefs::LightForegroundClr;
 
-
-	if( DoubleClickDelayTime >= 0.f )
-	{
-		DoubleClickDelayTime += InDeltaTime;
-		
-		if( !bIsHovered )
-		{
-			DoubleClickDelayTime = -1.f;
-		}
-	}
 	if( bEnableTitleBarBubble && IsGraphNodeHovered.IsBound() )
 	{
 		bIsHovered |= IsGraphNodeHovered.Execute();
@@ -228,10 +192,12 @@ void SCommentBubble::UpdateBubble()
 						.Padding( BubblePadding )
 						.AutoWidth()
 						[
-							SAssignNew(TextBlock, SInlineEditableTextBlock)
+							SAssignNew(TextBlock, SMultiLineEditableTextBox)
 							.Text( this, &SCommentBubble::GetCommentText )
+							.HintText( NSLOCTEXT( "CommentBubble", "EditCommentHint", "Type here to edit the comment" ))
 							.Font( FEditorStyle::GetFontStyle( TEXT("Graph.Node.CommentFont")) )
-							.ColorAndOpacity( this, &SCommentBubble::GetForegroundColor)
+							.ForegroundColor( this, &SCommentBubble::GetTextForegroundColor )
+							.BackgroundColor( this, &SCommentBubble::GetTextBackgroundColor )
 							.OnTextCommitted( this, &SCommentBubble::OnCommentTextCommitted )
 						]
 						+SHorizontalBox::Slot()
@@ -348,6 +314,16 @@ FSlateColor SCommentBubble::GetToggleButtonColor() const
 	return FLinearColor( 1.f, 1.f, 1.f, OpacityValue * OpacityValue );
 }
 
+FSlateColor SCommentBubble::GetTextBackgroundColor() const
+{
+	return TextBlock->HasKeyboardFocus() ? FLinearColor::White : SCommentBubbleDefs::TextClearBackground;
+}
+
+FSlateColor SCommentBubble::GetTextForegroundColor() const
+{
+	return TextBlock->HasKeyboardFocus() ? FLinearColor::Black : ForegroundColor;
+}
+
 void SCommentBubble::OnCommentTextCommitted( const FText& NewText, ETextCommit::Type CommitInfo )
 {
 	if( GraphNode )
diff --git a/Engine/Source/Editor/GraphEditor/Public/SCommentBubble.h b/Engine/Source/Editor/GraphEditor/Public/SCommentBubble.h
--- a/Engine/Source/Editor/GraphEditor/Public/SCommentBubble.h
+++ b/Engine/Source/Editor/GraphEditor/Public/SCommentBubble.h
@@ -51,8 +51,6 @@ class GRAPHEDITOR_API SCommentBubble : public SCompoundWidget
 
 	// Begin SWidget interface
 	virtual FCursorReply OnCursorQuery( const FGeometry& MyGeometry, const FPointerEvent& CursorEvent ) const override;
-	virtual FReply OnMouseButtonDown( const FGeometry& MyGeometry, const FPointerEvent& MouseEvent ) override;
-	virtual FReply OnMouseButtonDoubleClick( const FGeometry& MyGeometry, const FPointerEvent& MouseEvent ) override;
 	virtual void Tick( const FGeometry& AllottedGeometry, const double InCurrentTime, const float InDeltaTime ) override;
 	// End SWidget Interface
 
@@ -94,6 +92,12 @@ protected:
 	/** Returns the color for the toggle bubble including the opacity value */
 	FSlateColor GetToggleButtonColor() const;
 
+	/** Returns the current background color for the textbox */
+	FSlateColor GetTextBackgroundColor() const;
+	
+	/** Returns the current foreground color for the textbox */
+	FSlateColor GetTextForegroundColor() const;
+
 	/** Returns the foreground color for the text and buttons, taking into account the bubble color */
 	FSlateColor GetForegroundColor() const { return ForegroundColor; }
 
@@ -105,7 +109,7 @@ protected:
 	/** The GraphNode this widget interacts with */
 	UEdGraphNode* GraphNode;
 	/** Cached inline editable text box */
-	TSharedPtr<SInlineEditableTextBlock> TextBlock;
+	TSharedPtr<SMultiLineEditableTextBox> TextBlock;
 
 	/** The Comment Bubble color and opacity value */
 	TAttribute<FSlateColor> ColorAndOpacity;
@@ -129,8 +133,6 @@ protected:
 	bool bEnableBubbleCtrls;
 	/** Enable the title bar bubble toggle */
 	bool bEnableTitleBarBubble;
-	/** Used to control the delay on double click selection */
-	float DoubleClickDelayTime;
 	/** Used to Control hover fade up/down for widgets */
 	float OpacityValue;
 	/** Cached comment */
