Addressing some feedback and fixing a few issues to do with the recent node comment bubble changes.
#Branch UE4
#Proj Engine, GraphEditor

diff --git a/Engine/Source/Editor/GraphEditor/Private/SCommentBubble.cpp b/Engine/Source/Editor/GraphEditor/Private/SCommentBubble.cpp
--- a/Engine/Source/Editor/GraphEditor/Private/SCommentBubble.cpp
+++ b/Engine/Source/Editor/GraphEditor/Private/SCommentBubble.cpp
@@ -7,16 +7,10 @@
 
 namespace SCommentBubbleDefs
 {
-	/** Mouse double click disable */
-	static const float DoubleClickDisable = -1.f;
+	/** Bubble fade up/down delay */
+	static const float FadeDelay = -3.5f;
 
-	/** Mouse double click delay */
-	static const float DoubleClickDelay = 0.5f;
-
-	/** Bubble Full opacity */
-	static const float FullOpacity = 1.f;
-
-	/** Bubble Toggle Icon Fade Down Speed */
+	/** Bubble Toggle Icon Fade Speed */
 	static const float FadeDownSpeed = 5.f;
 
 	/** Height of the arrow connecting the bubble to the node */
@@ -50,8 +44,8 @@ void SCommentBubble::Construct( const FArguments& InArgs )
 	bEnableBubbleCtrls		= InArgs._EnableBubbleCtrls;
 	GraphLOD				= InArgs._GraphLOD;
 	IsGraphNodeHovered		= InArgs._IsGraphNodeHovered;
-
-	OpacityValue			= 0.f;
+	HintText				= InArgs._HintText.IsSet() ? InArgs._HintText : NSLOCTEXT( "CommentBubble", "EditCommentHint", "Click to edit" );
+	OpacityValue			= SCommentBubbleDefs::FadeDelay;
 
 	// Create Widget
 	UpdateBubble();
@@ -83,17 +77,24 @@ void SCommentBubble::Tick( const FGeometry& AllottedGeometry, const double InCur
 	{
 		bIsHovered |= IsGraphNodeHovered.Execute();
 
-		if( !GraphNode->bCommentBubbleVisible && !bIsHovered )
+		if( !GraphNode->bCommentBubbleVisible )
 		{
-			if( OpacityValue > 0.f )
+			if( bIsHovered )
 			{
-				const float FadeDownAmt = InDeltaTime * SCommentBubbleDefs::FadeDownSpeed;
-				OpacityValue = FMath::Max( OpacityValue - FadeDownAmt, 0.f );
+				if( OpacityValue < 1.f )
+				{
+					const float FadeUpAmt = InDeltaTime * SCommentBubbleDefs::FadeDownSpeed;
+					OpacityValue = FMath::Min( OpacityValue + FadeUpAmt, 1.f );
+				}
+			}
+			else
+			{
+				if( OpacityValue > SCommentBubbleDefs::FadeDelay )
+				{
+					const float FadeDownAmt = InDeltaTime * SCommentBubbleDefs::FadeDownSpeed;
+					OpacityValue = FMath::Max( OpacityValue - FadeDownAmt, SCommentBubbleDefs::FadeDelay );
+				}
 			}
-		}
-		else
-		{
-			OpacityValue = SCommentBubbleDefs::FullOpacity;
 		}
 	}
 }
@@ -194,8 +195,9 @@ void SCommentBubble::UpdateBubble()
 						[
 							SAssignNew(TextBlock, SMultiLineEditableTextBox)
 							.Text( this, &SCommentBubble::GetCommentText )
-							.HintText( NSLOCTEXT( "CommentBubble", "EditCommentHint", "Type here to edit the comment" ))
-							.Font( FEditorStyle::GetFontStyle( TEXT("Graph.Node.CommentFont")) )
+							.HintText( NSLOCTEXT( "CommentBubble", "EditCommentHint", "Click to edit" ))
+							.Font( FEditorStyle::GetFontStyle( TEXT("Graph.Node.CommentFont")))
+							.SelectAllTextWhenFocused( true )
 							.ForegroundColor( this, &SCommentBubble::GetTextForegroundColor )
 							.BackgroundColor( this, &SCommentBubble::GetTextBackgroundColor )
 							.OnTextCommitted( this, &SCommentBubble::OnCommentTextCommitted )
@@ -288,7 +290,11 @@ bool SCommentBubble::IsScalingAllowed() const
 
 FText SCommentBubble::GetCommentText() const
 {
-	if( CachedComment != CommentAttribute.Get() )
+	if( CommentAttribute.Get().IsEmpty() )
+	{
+		return HintText.Get();
+	}
+	else if( CachedComment != CommentAttribute.Get() )
 	{
 		CachedComment = CommentAttribute.Get();
 		CachedCommentText = FText::FromString( CachedComment );
@@ -326,7 +332,8 @@ FSlateColor SCommentBubble::GetTextForegroundColor() const
 
 void SCommentBubble::OnCommentTextCommitted( const FText& NewText, ETextCommit::Type CommitInfo )
 {
-	if( GraphNode )
+	const bool bValidText = !NewText.EqualTo( HintText.Get() );
+	if( GraphNode && bValidText )
 	{
 		const FScopedTransaction Transaction( NSLOCTEXT( "CommentBubble", "CommentCommitted", "Comment Changed" ) );
 		GraphNode->Modify();
diff --git a/Engine/Source/Editor/GraphEditor/Public/SCommentBubble.h b/Engine/Source/Editor/GraphEditor/Public/SCommentBubble.h
--- a/Engine/Source/Editor/GraphEditor/Public/SCommentBubble.h
+++ b/Engine/Source/Editor/GraphEditor/Public/SCommentBubble.h
@@ -21,6 +21,9 @@ class GRAPHEDITOR_API SCommentBubble : public SCompoundWidget
 		/** The comment text for the bubble */
 		SLATE_ATTRIBUTE( FString, Text )
 
+		/** The comment hint text for the bubble */
+		SLATE_ATTRIBUTE( FText, HintText )
+
 		/** Color and opacity for the comment bubble */
 		SLATE_ATTRIBUTE( FSlateColor, ColorAndOpacity )
 
@@ -121,6 +124,8 @@ protected:
 	TAttribute<bool> VisibilityState;
 	/** Attribute to query current LOD */
 	TAttribute<EGraphRenderingLOD::Type> GraphLOD;
+	/** Hint Text */
+	TAttribute<FText> HintText;
 
 	/** Delegate to determine if the graph node is currently hovered */
 	FIsGraphNodeHovered IsGraphNodeHovered;
diff --git a/Engine/Source/Runtime/Engine/Private/EdGraph/EdGraphSchema.cpp b/Engine/Source/Runtime/Engine/Private/EdGraph/EdGraphSchema.cpp
--- a/Engine/Source/Runtime/Engine/Private/EdGraph/EdGraphSchema.cpp
+++ b/Engine/Source/Runtime/Engine/Private/EdGraph/EdGraphSchema.cpp
@@ -614,12 +614,10 @@ void UEdGraphSchema::GetContextMenuActions(const UEdGraph* CurrentGraph, const U
 				.VAlign(VAlign_Center)
 				.FillWidth( 1.0f )
 				[
-					SNew(SEditableTextBox)
+					SNew(SMultiLineEditableTextBox)
 					.Text( NodeCommentText )
 					.ToolTipText(LOCTEXT("NodeComment_ToolTip", "Comment for this node"))
 					.OnTextCommitted_Static(&Local::OnNodeCommentTextCommitted, SelectedNodeWeakPtr)
-					.MinDesiredWidth( 120.f )
-					//.IsReadOnly( !SelectedActors(0)->IsActorLabelEditable() )
 					.SelectAllTextWhenFocused( true )
 					.RevertTextOnEscape( true )
 				];
