Blueprints: Add padding to variable set nodes so that the exec wires will line up with other nodes

diff --git a/Engine/Source/Editor/GraphEditor/Private/KismetNodes/SGraphNodeK2Var.cpp b/Engine/Source/Editor/GraphEditor/Private/KismetNodes/SGraphNodeK2Var.cpp
--- a/Engine/Source/Editor/GraphEditor/Private/KismetNodes/SGraphNodeK2Var.cpp
+++ b/Engine/Source/Editor/GraphEditor/Private/KismetNodes/SGraphNodeK2Var.cpp
@@ -238,22 +238,38 @@ void SGraphNodeK2Var::UpdateGraphNode()
 		]
 	];
 
+	float VerticalPaddingAmount = 0.0f;
+
 	// Add padding widgets at the top of the pin boxes if it's a struct operation with a long title
 	if (bPadTitle)
+	{
+		VerticalPaddingAmount += 16.0f;
+	}
+
+	// Add padding to offset the exec pin so that it will align with other nodes
+	if (UK2Node* K2Node = Cast<UK2Node>(GraphNode))
+	{
+		if (!K2Node->IsNodePure())
+		{
+			VerticalPaddingAmount += 7.0f;
+		}
+	}
+
+	if (VerticalPaddingAmount > 0.0f)
 	{
 		LeftNodeBox->AddSlot()
 			.AutoHeight()
 			.HAlign(HAlign_Left)
 			.VAlign(VAlign_Center)
 			[
-				SNew(SSpacer).Size(FVector2D(0.0f, 16.0f))
+				SNew(SSpacer).Size(FVector2D(0.0f, VerticalPaddingAmount))
 			];
 		RightNodeBox->AddSlot()
 			.AutoHeight()
 			.HAlign(HAlign_Left)
 			.VAlign(VAlign_Center)
 			[
-				SNew(SSpacer).Size(FVector2D(0.0f, 16.0f))
+				SNew(SSpacer).Size(FVector2D(0.0f, VerticalPaddingAmount))
 			];
 	}
 
