GraphEditor: Split connection drawing policies out into individual files

diff --git a/Engine/Source/Editor/GraphEditor/Private/AnimGraphConnectionDrawingPolicy.cpp b/Engine/Source/Editor/GraphEditor/Private/AnimGraphConnectionDrawingPolicy.cpp
new file mode 100644
--- /dev/null
+++ b/Engine/Source/Editor/GraphEditor/Private/AnimGraphConnectionDrawingPolicy.cpp
@@ -0,0 +1,88 @@
+// Copyright 1998-2014 Epic Games, Inc. All Rights Reserved.
+
+#include "GraphEditorCommon.h"
+#include "Editor/UnrealEd/Public/Kismet2/BlueprintEditorUtils.h"
+#include "AnimGraphConnectionDrawingPolicy.h"
+#include "AnimationGraphSchema.h"
+#include "AnimGraphNode_Base.h"
+
+/////////////////////////////////////////////////////
+// FAnimGraphConnectionDrawingPolicy
+
+FAnimGraphConnectionDrawingPolicy::FAnimGraphConnectionDrawingPolicy(int32 InBackLayerID, int32 InFrontLayerID, float ZoomFactor, const FSlateRect& InClippingRect, FSlateWindowElementList& InDrawElements, UEdGraph* InGraphObj)
+	: FKismetConnectionDrawingPolicy(InBackLayerID, InFrontLayerID, ZoomFactor, InClippingRect, InDrawElements, InGraphObj)
+{
+}
+
+bool FAnimGraphConnectionDrawingPolicy::TreatWireAsExecutionPin(UEdGraphPin* InputPin, UEdGraphPin* OutputPin) const
+{
+	const UAnimationGraphSchema* Schema = GetDefault<UAnimationGraphSchema>();
+
+	return (InputPin != NULL) && (Schema->IsPosePin(OutputPin->PinType));
+}
+
+void FAnimGraphConnectionDrawingPolicy::BuildExecutionRoadmap()
+{
+	UAnimBlueprint* TargetBP = CastChecked<UAnimBlueprint>(FBlueprintEditorUtils::FindBlueprintForGraphChecked(GraphObj));
+	UAnimBlueprintGeneratedClass* AnimBlueprintClass = (UAnimBlueprintGeneratedClass*)(*(TargetBP->GeneratedClass));
+
+	if (TargetBP->GetObjectBeingDebugged() == NULL)
+	{
+		return;
+	}
+
+	TMap<UProperty*, UObject*> PropertySourceMap;
+	AnimBlueprintClass->GetDebugData().GenerateReversePropertyMap(/*out*/ PropertySourceMap);
+
+	FAnimBlueprintDebugData& DebugInfo = AnimBlueprintClass->GetAnimBlueprintDebugData();
+	for (auto VisitIt = DebugInfo.UpdatedNodesThisFrame.CreateIterator(); VisitIt; ++VisitIt)
+	{
+		const FAnimBlueprintDebugData::FNodeVisit& VisitRecord = *VisitIt;
+
+		if ((VisitRecord.SourceID >= 0) && (VisitRecord.SourceID < AnimBlueprintClass->AnimNodeProperties.Num()) && (VisitRecord.TargetID >= 0) && (VisitRecord.TargetID < AnimBlueprintClass->AnimNodeProperties.Num()))
+		{
+			if (UAnimGraphNode_Base* SourceNode = Cast<UAnimGraphNode_Base>(PropertySourceMap.FindRef(AnimBlueprintClass->AnimNodeProperties[VisitRecord.SourceID])))
+			{
+				if (UAnimGraphNode_Base* TargetNode = Cast<UAnimGraphNode_Base>(PropertySourceMap.FindRef(AnimBlueprintClass->AnimNodeProperties[VisitRecord.TargetID])))
+				{
+					UEdGraphPin* PoseNet = NULL;
+
+					UAnimationGraphSchema const* AnimSchema = GetDefault<UAnimationGraphSchema>();
+					for (int32 PinIndex = 0; PinIndex < TargetNode->Pins.Num(); ++PinIndex)
+					{
+						UEdGraphPin* Pin = TargetNode->Pins[PinIndex];
+						if (AnimSchema->IsPosePin(Pin->PinType) && (Pin->Direction == EGPD_Output))
+						{
+							PoseNet = Pin;
+							break;
+						}
+					}
+
+					if (PoseNet != NULL)
+					{
+						//@TODO: Extend the rendering code to allow using the recorded blend weight instead of faked exec times
+						FExecPairingMap& Predecessors = PredecessorPins.FindOrAdd((UEdGraphNode*)SourceNode);
+						FTimePair& Timings = Predecessors.FindOrAdd(PoseNet);
+						Timings.PredExecTime = 0.0;
+						Timings.ThisExecTime = VisitRecord.Weight;
+					}
+				}
+			}
+		}
+	}
+}
+
+void FAnimGraphConnectionDrawingPolicy::DetermineStyleOfExecWire(float& Thickness, FLinearColor& WireColor, bool& bDrawBubbles, const FTimePair& Times)
+{
+	// It's a followed link, make it strong and yellowish but fading over time
+	const double BlendWeight = Times.ThisExecTime;
+
+	const float HeavyBlendThickness = AttackWireThickness;
+	const float LightBlendThickness = SustainWireThickness;
+
+	Thickness = FMath::Lerp<float>(LightBlendThickness, HeavyBlendThickness, BlendWeight);
+	WireColor = WireColor * (BlendWeight * 0.5f + 0.5f);//FMath::Lerp<FLinearColor>(SustainColor, AttackColor, BlendWeight);
+
+	bDrawBubbles = true;
+}
+
diff --git a/Engine/Source/Editor/GraphEditor/Private/AnimGraphConnectionDrawingPolicy.h b/Engine/Source/Editor/GraphEditor/Private/AnimGraphConnectionDrawingPolicy.h
new file mode 100644
--- /dev/null
+++ b/Engine/Source/Editor/GraphEditor/Private/AnimGraphConnectionDrawingPolicy.h
@@ -0,0 +1,23 @@
+// Copyright 1998-2014 Epic Games, Inc. All Rights Reserved.
+
+#pragma once
+
+#include "BlueprintConnectionDrawingPolicy.h"
+
+/////////////////////////////////////////////////////
+// FAnimGraphConnectionDrawingPolicy
+
+// This class draws the connections for an UEdGraph with an animation schema
+class FAnimGraphConnectionDrawingPolicy : public FKismetConnectionDrawingPolicy
+{
+public:
+	// Constructor
+	FAnimGraphConnectionDrawingPolicy(int32 InBackLayerID, int32 InFrontLayerID, float ZoomFactor, const FSlateRect& InClippingRect, FSlateWindowElementList& InDrawElements, UEdGraph* InGraphObj);
+
+	// FKismetConnectionDrawingPolicy interface
+	virtual bool TreatWireAsExecutionPin(UEdGraphPin* InputPin, UEdGraphPin* OutputPin) const override;
+	virtual void BuildExecutionRoadmap() override;
+	virtual void DetermineStyleOfExecWire(float& Thickness, FLinearColor& WireColor, bool& bDrawBubbles, const FTimePair& Times) override;
+	// End of FKismetConnectionDrawingPolicy interface
+};
+
diff --git a/Engine/Source/Editor/GraphEditor/Private/BlueprintConnectionDrawingPolicy.cpp b/Engine/Source/Editor/GraphEditor/Private/BlueprintConnectionDrawingPolicy.cpp
new file mode 100644
--- /dev/null
+++ b/Engine/Source/Editor/GraphEditor/Private/BlueprintConnectionDrawingPolicy.cpp
@@ -0,0 +1,446 @@
+// Copyright 1998-2014 Epic Games, Inc. All Rights Reserved.
+
+#include "GraphEditorCommon.h"
+#include "BlueprintConnectionDrawingPolicy.h"
+#include "Editor/UnrealEd/Public/Kismet2/KismetDebugUtilities.h"
+#include "Editor/UnrealEd/Public/Kismet2/BlueprintEditorUtils.h"
+#include "K2Node_Knot.h"
+#include "BlueprintEditorSettings.h"
+
+/////////////////////////////////////////////////////
+// FKismetConnectionDrawingPolicy
+
+FKismetConnectionDrawingPolicy::FKismetConnectionDrawingPolicy(int32 InBackLayerID, int32 InFrontLayerID, float ZoomFactor, const FSlateRect& InClippingRect, FSlateWindowElementList& InDrawElements, UEdGraph* InGraphObj)
+	: FConnectionDrawingPolicy(InBackLayerID, InFrontLayerID, ZoomFactor, InClippingRect, InDrawElements)
+	, GraphObj(InGraphObj)
+{
+	// Don't want to draw ending arrowheads
+	ArrowImage = nullptr;
+	ArrowRadius = FVector2D::ZeroVector;
+
+	// But we do want to draw midpoint arrowheads
+	if (GetDefault<UBlueprintEditorSettings>()->bDrawMidpointArrowsInBlueprints)
+	{
+		MidpointImage = FEditorStyle::GetBrush( TEXT("Graph.Arrow") );
+		MidpointRadius = MidpointImage->ImageSize * ZoomFactor * 0.5f;
+	}
+
+	// Cache off the editor options
+	AttackColor = Settings->TraceAttackColor;
+	SustainColor = Settings->TraceSustainColor;
+	ReleaseColor = Settings->TraceReleaseColor;
+
+	AttackWireThickness = Settings->TraceAttackWireThickness;
+	SustainWireThickness = Settings->TraceSustainWireThickness;
+	ReleaseWireThickness = Settings->TraceReleaseWireThickness;
+
+	TracePositionBonusPeriod = Settings->TracePositionBonusPeriod;
+	TracePositionExponent = Settings->TracePositionExponent;
+	AttackHoldPeriod = Settings->TraceAttackHoldPeriod;
+	DecayPeriod = Settings->TraceDecayPeriod;
+	DecayExponent = Settings->TraceDecayExponent;
+	SustainHoldPeriod = Settings->TraceSustainHoldPeriod;
+	ReleasePeriod = Settings->TraceReleasePeriod;
+	ReleaseExponent = Settings->TraceReleaseExponent;
+
+	CurrentTime = 0.0;
+	LatestTimeDiscovered = 0.0;
+}
+
+void FKismetConnectionDrawingPolicy::Draw(TMap<TSharedRef<SWidget>, FArrangedWidget>& PinGeometries, FArrangedChildren& ArrangedNodes)
+{
+	// Build the execution roadmap (also populates execution times)
+	BuildExecutionRoadmap();
+
+	// Draw everything
+	FConnectionDrawingPolicy::Draw(PinGeometries, ArrangedNodes);
+}
+
+bool FKismetConnectionDrawingPolicy::CanBuildRoadmap() const
+{
+	UBlueprint* TargetBP = FBlueprintEditorUtils::FindBlueprintForGraphChecked(GraphObj);
+	UObject* ActiveObject = TargetBP->GetObjectBeingDebugged();
+
+	return ActiveObject != NULL;
+}
+
+void FKismetConnectionDrawingPolicy::BuildExecutionRoadmap()
+{
+	LatestTimeDiscovered = 0.0;
+
+	// Only do highlighting in PIE or SIE
+	if (!CanBuildRoadmap())
+	{
+		return;
+	}
+
+	UBlueprint* TargetBP = FBlueprintEditorUtils::FindBlueprintForGraphChecked(GraphObj);
+	UObject* ActiveObject = TargetBP->GetObjectBeingDebugged();
+	check(ActiveObject); // Due to CanBuildRoadmap
+
+	// Redirect the target Blueprint when debugging with a macro graph visible
+	if (TargetBP->BlueprintType == BPTYPE_MacroLibrary)
+	{
+		TargetBP = Cast<UBlueprint>(ActiveObject->GetClass()->ClassGeneratedBy);
+	}
+
+	TArray<UEdGraphNode*> SequentialNodesInGraph;
+	TArray<double> SequentialNodeTimes;
+	TArray<UEdGraphPin*> SequentialExecPinsInGraph;
+
+	{
+		const TSimpleRingBuffer<FKismetTraceSample>& TraceStack = FKismetDebugUtilities::GetTraceStack();
+
+		UBlueprintGeneratedClass* TargetClass = Cast<UBlueprintGeneratedClass>(TargetBP->GeneratedClass);
+		FBlueprintDebugData& DebugData = TargetClass->GetDebugData();
+
+		for (int32 i = 0; i < TraceStack.Num(); ++i)
+		{
+			const FKismetTraceSample& Sample = TraceStack(i);
+
+			if (UObject* TestObject = Sample.Context.Get())
+			{
+				if (TestObject == ActiveObject)
+				{
+					UEdGraphPin* AssociatedPin = DebugData.FindExecPinFromCodeLocation(Sample.Function.Get(), Sample.Offset);
+
+					if (UEdGraphNode* Node = DebugData.FindSourceNodeFromCodeLocation(Sample.Function.Get(), Sample.Offset, /*bAllowImpreciseHit=*/ false))
+					{
+						if (GraphObj == Node->GetGraph())
+						{
+							SequentialNodesInGraph.Add(Node);
+							SequentialNodeTimes.Add(Sample.ObservationTime);
+							SequentialExecPinsInGraph.Add(AssociatedPin);
+						}
+						else
+						{
+							// If the top-level source node is a macro instance node
+							UK2Node_MacroInstance* MacroInstanceNode = Cast<UK2Node_MacroInstance>(Node);
+							if (MacroInstanceNode)
+							{
+								// Attempt to locate the macro source node through the code mapping
+								UEdGraphNode* MacroSourceNode = DebugData.FindMacroSourceNodeFromCodeLocation(Sample.Function.Get(), Sample.Offset);
+								if (MacroSourceNode)
+								{
+									// If the macro source node is located in the current graph context
+									if (GraphObj == MacroSourceNode->GetGraph())
+									{
+										// Add it to the sequential node list
+										SequentialNodesInGraph.Add(MacroSourceNode);
+										SequentialNodeTimes.Add(Sample.ObservationTime);
+										SequentialExecPinsInGraph.Add(AssociatedPin);
+									}
+									else
+									{
+										// The macro source node isn't in the current graph context, but we might have a macro instance node that is
+										// in the current graph context, so obtain the set of macro instance nodes that are mapped to the code here.
+										TArray<UEdGraphNode*> MacroInstanceNodes;
+										DebugData.FindMacroInstanceNodesFromCodeLocation(Sample.Function.Get(), Sample.Offset, MacroInstanceNodes);
+
+										// For each macro instance node in the set
+										for (auto MacroInstanceNodeIt = MacroInstanceNodes.CreateConstIterator(); MacroInstanceNodeIt; ++MacroInstanceNodeIt)
+										{
+											// If the macro instance node is located in the current graph context
+											MacroInstanceNode = Cast<UK2Node_MacroInstance>(*MacroInstanceNodeIt);
+											if (MacroInstanceNode && GraphObj == MacroInstanceNode->GetGraph())
+											{
+												// Add it to the sequential node list
+												SequentialNodesInGraph.Add(MacroInstanceNode);
+												SequentialNodeTimes.Add(Sample.ObservationTime);
+												SequentialExecPinsInGraph.Add(AssociatedPin);
+
+												// Exit the loop; we're done
+												break;
+											}
+										}
+									}
+								}
+							}
+						}
+					}
+				}
+			}
+		}
+	}
+
+	// Run thru and apply bonus time
+	const float InvNumNodes = 1.0f / (float)SequentialNodeTimes.Num();
+	for (int32 i = 0; i < SequentialNodesInGraph.Num(); ++i)
+	{
+		double& ObservationTime = SequentialNodeTimes[i];
+
+		const float PositionRatio = (SequentialNodeTimes.Num() - i) * InvNumNodes;
+		const float PositionBonus = FMath::Pow(PositionRatio, TracePositionExponent) * TracePositionBonusPeriod;
+		ObservationTime += PositionBonus;
+
+		LatestTimeDiscovered = FMath::Max<double>(LatestTimeDiscovered, ObservationTime);
+	}
+
+	UEdGraphPin* LastExecPin = NULL;
+	// Record the unique exec-pin to time pairings, keeping only the most recent 
+	// times for each pairing... reverse the "SequentialNodes" because right now
+	// it is in stack order (with the last executed node first)
+	for (int32 i = SequentialNodesInGraph.Num() - 1; i >= 1; --i)
+	{
+		UEdGraphNode* CurNode  = SequentialNodesInGraph[i];
+		UEdGraphNode* NextNode = SequentialNodesInGraph[i-1];
+
+		// keep track of the last exec-pin executed by CurNode (these tracked 
+		// pins coincide with "WireTraceSite" op-codes that have been injected 
+		// prior to every "goto" statement... this way we have context for which
+		// pin executed the jump)
+		if (UEdGraphPin* AssociatedPin = SequentialExecPinsInGraph[i])
+		{
+			LastExecPin = AssociatedPin;
+		}
+		
+		// if this statement is a jump (from one node to another)
+		if (CurNode != NextNode)
+		{
+			// if there was a wire-trace op-code inserted before this jump
+			if (LastExecPin != NULL)
+			{
+				//ensure(LastExecPin->GetOwningNode() == CurNode);
+				double NextNodeTime = SequentialNodeTimes[i-1];
+
+				FExecPairingMap& ExecPaths  = PredecessorPins.FindOrAdd(NextNode);
+				FTimePair&       ExecTiming = ExecPaths.FindOrAdd(LastExecPin);
+				// make sure that if we've already visited this exec-pin (like 
+				// in a for-loop or something), that we're replacing it with a 
+				// more recent execution time
+				//
+				// @TODO I don't see when this wouldn't be the case
+				if (ExecTiming.ThisExecTime < NextNodeTime)
+				{
+					double CurNodeTime = SequentialNodeTimes[i];
+					ExecTiming.ThisExecTime = NextNodeTime;
+					ExecTiming.PredExecTime = CurNodeTime;
+				}
+			}
+			// if the nodes aren't graphically connected how could they be 
+			// executed back-to-back? well, this could be a pop back to a 
+			// sequence node from the end of one thread of execution, etc.
+			else if (AreNodesGraphicallySequential(CurNode, NextNode))
+			{
+				// only warn when the nodes are directly connected (this is all
+				// for execution flow visualization after all)
+				UE_LOG(LogConnectionDrawingPolicy, Warning, TEXT("Looks like a wire-trace was not injected before the jump from '%s' to '%s'."), 
+					*CurNode->GetNodeTitle(ENodeTitleType::FullTitle).ToString(), *NextNode->GetNodeTitle(ENodeTitleType::FullTitle).ToString());
+			}
+
+			// clear the exec-pin (we're moving to a new node and want to find 
+			// it's executed out pin)
+			LastExecPin = NULL;
+		}
+		// else, we're only collecting this data for tracing node-to-node
+		// executions (so we don't care about this sequence of statements)
+	}
+
+	// Fade only when free-running (since we're using FApp::GetCurrentTime(), instead of FPlatformTime::Seconds)
+	const double MaxTimeAhead = FMath::Min(FApp::GetCurrentTime() + 2*TracePositionBonusPeriod, LatestTimeDiscovered); //@TODO: Rough clamping; should be exposed as a parameter
+	CurrentTime = FMath::Max(FApp::GetCurrentTime(), MaxTimeAhead);
+}
+
+void FKismetConnectionDrawingPolicy::CalculateEnvelopeAlphas(double ExecutionTime, /*out*/ float& AttackAlpha, /*out*/ float& SustainAlpha) const
+{
+	const float DeltaTime = (float)(CurrentTime - ExecutionTime);
+
+	{
+		const float UnclampedDecayRatio = 1.0f - ((DeltaTime - AttackHoldPeriod) / DecayPeriod);
+		const float ClampedDecayRatio = FMath::Clamp<float>(UnclampedDecayRatio, 0.0f, 1.0f);
+		AttackAlpha = FMath::Pow(ClampedDecayRatio, DecayExponent);
+	}
+
+	{
+		const float SustainEndTime = AttackHoldPeriod + DecayPeriod + SustainHoldPeriod;
+
+		const float UnclampedReleaseRatio = 1.0f - ((DeltaTime - SustainEndTime) / ReleasePeriod);
+		const float ClampedReleaseRatio = FMath::Clamp<float>(UnclampedReleaseRatio, 0.0f, 1.0f);
+		SustainAlpha = FMath::Pow(ClampedReleaseRatio, ReleaseExponent);
+	}
+}
+
+bool FKismetConnectionDrawingPolicy::TreatWireAsExecutionPin(UEdGraphPin* InputPin, UEdGraphPin* OutputPin) const
+{
+	const UEdGraphSchema_K2* Schema = GetDefault<UEdGraphSchema_K2>();
+
+	return (InputPin != NULL) && (Schema->IsExecPin(*OutputPin));
+}
+
+bool FKismetConnectionDrawingPolicy::AreNodesGraphicallySequential(UEdGraphNode* InputNode, UEdGraphNode* OutputNode) const
+{
+	for (UEdGraphPin* Pin : InputNode->Pins)
+	{
+		if (Pin->Direction != EGPD_Output)
+		{
+			continue;
+		}
+
+		for (UEdGraphPin* Connection : Pin->LinkedTo)
+		{
+			if (!TreatWireAsExecutionPin(Pin, Connection))
+			{
+				continue;
+			}
+
+			if (Connection->GetOwningNode() == OutputNode)
+			{
+				return true;
+			}
+		}
+	}
+
+	return false;
+}
+
+void FKismetConnectionDrawingPolicy::DetermineStyleOfExecWire(float& Thickness, FLinearColor& WireColor, bool& bDrawBubbles, const FTimePair& Times)
+{
+	// It's a followed link, make it strong and yellowish but fading over time
+	const double ExecTime = Times.ThisExecTime;
+
+	float AttackAlpha;
+	float SustainAlpha;
+	CalculateEnvelopeAlphas(ExecTime, /*out*/ AttackAlpha, /*out*/ SustainAlpha);
+
+	const float DecayedAttackThickness = FMath::Lerp<float>(SustainWireThickness, AttackWireThickness, AttackAlpha);
+	Thickness = FMath::Lerp<float>(ReleaseWireThickness, DecayedAttackThickness, SustainAlpha);
+
+	const FLinearColor DecayedAttackColor = FMath::Lerp<FLinearColor>(SustainColor, AttackColor, AttackAlpha);
+	WireColor = WireColor * FMath::Lerp<FLinearColor>(ReleaseColor, DecayedAttackColor, SustainAlpha);
+
+	if (SustainAlpha > KINDA_SMALL_NUMBER)
+	{
+		bDrawBubbles = true;
+	}
+}
+
+FKismetConnectionDrawingPolicy::FTimePair const* FKismetConnectionDrawingPolicy::BackTraceExecPath(UEdGraphPin const* const OutputPin, FExecPairingMap const* const NodeExecutionList)
+{
+	FTimePair const* FoundExecPath = nullptr;
+
+	UEdGraphNode const* const OwningNode = OutputPin->GetOwningNode();
+	if (UK2Node_Knot const* const KnotNode = Cast<UK2Node_Knot>(OwningNode))
+	{
+		UEdGraphPin const* const KnotInputPin = KnotNode->GetInputPin();
+		for (UEdGraphPin const* KnotInput : KnotInputPin->LinkedTo)
+		{
+			FoundExecPath = BackTraceExecPath(KnotInput, NodeExecutionList);
+			if (FoundExecPath != nullptr)
+			{
+				break;
+			}
+		}
+	}
+	else
+	{
+		FoundExecPath = NodeExecutionList->Find(OutputPin);
+	}
+
+	return FoundExecPath;
+}
+
+// Give specific editor modes a chance to highlight this connection or darken non-interesting connections
+void FKismetConnectionDrawingPolicy::DetermineWiringStyle(UEdGraphPin* OutputPin, UEdGraphPin* InputPin, /*inout*/ float& Thickness, /*inout*/ FLinearColor& WireColor, /*inout*/bool& bDrawBubbles, /*inout*/bool& bBidirectional)
+{
+	// Get the schema and grab the default color from it
+	check(OutputPin);
+	check(GraphObj);
+	const UEdGraphSchema* Schema = GraphObj->GetSchema();
+
+	WireColor = Schema->GetPinTypeColor(OutputPin->PinType);
+
+	const bool bDeemphasizeUnhoveredPins = HoveredPins.Num() > 0;
+
+	// If this is a K2 graph, try to be a little more specific
+	const UEdGraphSchema_K2* K2Schema = Cast<const UEdGraphSchema_K2>(Schema);
+	if (K2Schema != NULL)
+	{
+		if (TreatWireAsExecutionPin(InputPin, OutputPin))
+		{
+			if (CanBuildRoadmap())
+			{
+				UEdGraphNode* InputNode = InputPin->GetOwningNode();
+				// knot nodes are removed from the graph at compile time, so we 
+				// have to follow them until we find something that would have 
+				// actually executed
+				while (UK2Node_Knot* KnotNode = Cast<UK2Node_Knot>(InputNode))
+				{
+					InputNode = nullptr;
+
+					UEdGraphPin* OutPin = KnotNode->GetOutputPin();
+					if (OutPin->LinkedTo.Num() > 0)
+					{
+						check(OutPin->LinkedTo.Num() == 1);
+						InputNode = OutPin->LinkedTo[0]->GetOwningNode();
+					}	
+				}
+
+				// track if this node connection was ran or not
+				bool bExecuted = false;
+
+				// if the node belonging to InputPin was actually executed
+				if (FExecPairingMap* ExecPaths = PredecessorPins.Find(InputNode))
+				{
+					// if the output pin is one of the pins that lead to InputNode being ran
+					if (FTimePair const* ExecTiming = BackTraceExecPath(OutputPin, ExecPaths))
+					{
+						bExecuted = true;
+						DetermineStyleOfExecWire(/*inout*/ Thickness, /*inout*/ WireColor, /*inout*/ bDrawBubbles, *ExecTiming);
+					}
+				}
+				
+				if (!bExecuted)
+				{
+					// It's not followed, fade it and keep it thin
+					WireColor = ReleaseColor;
+					Thickness = ReleaseWireThickness;
+				}
+			}
+			else
+			{
+				// Make exec wires slightly thicker even outside of debug
+				Thickness = 3.0f;
+			}
+		}
+		else
+		{
+			// Array types should draw thicker
+			if( (InputPin && InputPin->PinType.bIsArray) || (OutputPin && OutputPin->PinType.bIsArray) )
+			{
+				Thickness = 3.0f;
+			}
+		}
+	}
+
+	if (bDeemphasizeUnhoveredPins)
+	{
+		ApplyHoverDeemphasis(OutputPin, InputPin, /*inout*/ Thickness, /*inout*/ WireColor);
+	}
+}
+
+void FKismetConnectionDrawingPolicy::SetIncompatiblePinDrawState(const TSharedPtr<SGraphPin>& StartPin, const TSet< TSharedRef<SWidget> >& VisiblePins)
+{
+	ResetIncompatiblePinDrawState(VisiblePins);
+
+	for(auto VisiblePinIterator = VisiblePins.CreateConstIterator(); VisiblePinIterator; ++VisiblePinIterator)
+	{
+		TSharedPtr<SGraphPin> CheckPin = StaticCastSharedRef<SGraphPin>(*VisiblePinIterator);
+		if(CheckPin != StartPin)
+		{
+			const FPinConnectionResponse Response = StartPin->GetPinObj()->GetSchema()->CanCreateConnection(StartPin->GetPinObj(), CheckPin->GetPinObj());
+			if(Response.Response == CONNECT_RESPONSE_DISALLOW)
+			{
+				CheckPin->SetPinColorModifier(FLinearColor(0.25f, 0.25f, 0.25f, 0.5f));
+			}
+		}
+	}
+}
+
+void FKismetConnectionDrawingPolicy::ResetIncompatiblePinDrawState(const TSet< TSharedRef<SWidget> >& VisiblePins)
+{
+	for(auto VisiblePinIterator = VisiblePins.CreateConstIterator(); VisiblePinIterator; ++VisiblePinIterator)
+	{
+		TSharedPtr<SGraphPin> VisiblePin = StaticCastSharedRef<SGraphPin>(*VisiblePinIterator);
+		VisiblePin->SetPinColorModifier(FLinearColor::White);
+	}
+}
diff --git a/Engine/Source/Editor/GraphEditor/Private/BlueprintConnectionDrawingPolicy.h b/Engine/Source/Editor/GraphEditor/Private/BlueprintConnectionDrawingPolicy.h
new file mode 100644
--- /dev/null
+++ b/Engine/Source/Editor/GraphEditor/Private/BlueprintConnectionDrawingPolicy.h
@@ -0,0 +1,105 @@
+// Copyright 1998-2014 Epic Games, Inc. All Rights Reserved.
+
+#pragma once
+
+#include "ConnectionDrawingPolicy.h"
+
+/////////////////////////////////////////////////////
+// FKismetConnectionDrawingPolicy
+
+// This class draws the connections for an UEdGraph using a K2-based schema
+class FKismetConnectionDrawingPolicy : public FConnectionDrawingPolicy
+{
+protected:
+	// Times for one execution pair within the current graph
+	struct FTimePair
+	{
+		double PredExecTime; // time that the node with the associated exec-out was executed
+		double ThisExecTime; // time that the following node was executed
+
+		FTimePair()
+			: PredExecTime(0.0)
+			, ThisExecTime(0.0)
+		{
+		}
+	};
+	/** 
+	 * Map of exec-pins to a pair of execution times (holds timing info for the 
+	 * owning-node, paired with a time of execution for the linked node)... only 
+	 * one timing-pair per pin because exec outputs can only connect to one node
+	 */
+	typedef TMap<UEdGraphPin*, FTimePair> FExecPairingMap;
+
+	// Map of executed nodes, and the execution pins that lead to them being ran
+	TMap<UEdGraphNode*, FExecPairingMap> PredecessorPins;
+
+	UEdGraph* GraphObj;
+	double CurrentTime;
+	double LatestTimeDiscovered;
+
+	FLinearColor AttackColor;
+	FLinearColor SustainColor;
+	FLinearColor ReleaseColor;
+
+	float AttackWireThickness;
+	float SustainWireThickness;
+	float ReleaseWireThickness;
+
+	float TracePositionBonusPeriod;
+	float TracePositionExponent;
+	float AttackHoldPeriod;
+	float DecayPeriod;
+	float DecayExponent;
+	float SustainHoldPeriod;
+	float ReleasePeriod;
+	float ReleaseExponent;
+
+protected:
+	// Should this wire be treated as an execution pin, querying the execution trace for connected nodes to draw it differently?
+	virtual bool TreatWireAsExecutionPin(UEdGraphPin* InputPin, UEdGraphPin* OutputPin) const;
+public:
+	FKismetConnectionDrawingPolicy(int32 InBackLayerID, int32 InFrontLayerID, float ZoomFactor, const FSlateRect& InClippingRect, FSlateWindowElementList& InDrawElements, UEdGraph* InGraphObj);
+
+	virtual void BuildExecutionRoadmap();
+
+	virtual bool CanBuildRoadmap() const;
+
+	void CalculateEnvelopeAlphas(double ExecutionTime, /*out*/ float& AttackAlpha, /*out*/ float& SustainAlpha) const;
+
+	// FConnectionDrawingPolicy interface
+	virtual void DetermineWiringStyle(UEdGraphPin* OutputPin, UEdGraphPin* InputPin, /*inout*/ float& Thickness, /*inout*/ FLinearColor& WireColor, /*inout*/bool& bDrawBubbles, /*inout*/ bool& bBidirectional) override;
+	virtual void Draw(TMap<TSharedRef<SWidget>, FArrangedWidget>& PinGeometries, FArrangedChildren& ArrangedNodes) override;
+	// End of FConnectionDrawingPolicy interface
+
+	virtual void DetermineStyleOfExecWire(float& Thickness, FLinearColor& WireColor, bool& bDrawBubbles, const FTimePair& Times);
+
+	virtual void SetIncompatiblePinDrawState(const TSharedPtr<SGraphPin>& StartPin, const TSet< TSharedRef<SWidget> >& VisiblePins);
+	virtual void ResetIncompatiblePinDrawState(const TSet< TSharedRef<SWidget> >& VisiblePins);
+
+private: 
+	/**
+	* Checks that the specified nodes are connected in the graph directly by an 
+	* exec pin.
+	* 
+	* @param  InputNode	The node whose exec output(s) should be connected to OutputNode.
+	* @param  OutputNode	The node whose exec input should be coming from the InputNode.
+	* @return True if the two nodes are connected in the graph, otherwise false.
+	*/
+	bool AreNodesGraphicallySequential(UEdGraphNode* InputNode, UEdGraphNode* OutputNode) const;
+
+	/**
+	 * With the introduction of "reroute" (UK2Node_Knot) nodes, we have to 
+	 * backtrace to find a node that was actually executed before the one in 
+	 * question (knot nodes are removed at compile time, and therefore never
+	 * executed).
+	 *
+	 * This function recursively searches for any pins feeding the OutputPin (or 
+	 * the OutputPin itself); it looks for a pin on a non-reroute node (one that
+	 * is listed in NodeExecutionList).
+	 * 
+	 * @param  OutputPin			An exec pin, leading into the node in question (an output pin on some other node).
+	 * @param  NodeExecutionList	A list of pins that feed the execution of a specific node.
+	 * @return Null if OutputPin did not cause the invocation of the node in question, otherwise a valid time pair denoting when the node was executed.
+	 */
+	FTimePair const* BackTraceExecPath(UEdGraphPin const* const OutputPin, FExecPairingMap const* const NodeExecutionList);
+};
diff --git a/Engine/Source/Editor/GraphEditor/Private/ConnectionDrawingPolicy.cpp b/Engine/Source/Editor/GraphEditor/Private/ConnectionDrawingPolicy.cpp
--- a/Engine/Source/Editor/GraphEditor/Private/ConnectionDrawingPolicy.cpp
+++ b/Engine/Source/Editor/GraphEditor/Private/ConnectionDrawingPolicy.cpp
@@ -1,23 +1,9 @@
 // Copyright 1998-2014 Epic Games, Inc. All Rights Reserved.
 
-
 #include "GraphEditorCommon.h"
 #include "ConnectionDrawingPolicy.h"
-#include "Editor/UnrealEd/Public/Kismet2/KismetDebugUtilities.h"
-#include "Editor/UnrealEd/Public/Kismet2/BlueprintEditorUtils.h"
-#include "KismetNodes/KismetNodeInfoContext.h"
-#include "SoundDefinitions.h"
-#include "AnimationStateNodes/SGraphNodeAnimTransition.h"
-#include "AnimStateTransitionNode.h"
-#include "AnimStateEntryNode.h"
-#include "AnimationGraphSchema.h"
-#include "AnimGraphNode_Base.h"
-#include "Sound/SoundNode.h"
-#include "K2Node_Knot.h"
-#include "BlueprintEditorSettings.h"
-#include "Sound/SoundCue.h"
-
-DEFINE_LOG_CATEGORY_STATIC(LogConnectionDrawingPolicy, Log, All);
+
+DEFINE_LOG_CATEGORY(LogConnectionDrawingPolicy);
 
 //////////////////////////////////////////////////////////////////////////
 // FGeometryHelper
@@ -90,101 +76,6 @@ FVector2D FGeometryHelper::FindClosestPointOnGeom(const FGeometry& Geom, const F
 	return BestPoint;
 }
 
-//////////////////////////////////////////////////////////////////////////
-// FKismetNodeInfoContext
-
-// Context used to aid debugging displays for nodes
-FKismetNodeInfoContext::FKismetNodeInfoContext(UEdGraph* SourceGraph)
-	: ActiveObjectBeingDebugged(NULL)
-{
-	// Only show pending latent actions in PIE/SIE mode
-	SourceBlueprint = FBlueprintEditorUtils::FindBlueprintForGraph(SourceGraph);
-
-	if (SourceBlueprint != NULL)
-	{
-		ActiveObjectBeingDebugged = SourceBlueprint->GetObjectBeingDebugged();
-
-		// Run thru debugged objects to see if any are objects with pending latent actions
-		if (ActiveObjectBeingDebugged != NULL)
-		{
-			UBlueprintGeneratedClass* Class = CastChecked<UBlueprintGeneratedClass>((UObject*)(ActiveObjectBeingDebugged->GetClass()));
-			FBlueprintDebugData const& ClassDebugData = Class->GetDebugData();
-
-			TSet<UObject*> LatentContextObjects;
-
-			TArray<UK2Node_CallFunction*> FunctionNodes;
-			SourceGraph->GetNodesOfClass<UK2Node_CallFunction>(FunctionNodes);
-			// collect all the world context objects for all of the graph's latent nodes
-			for (UK2Node_CallFunction const* FunctionNode : FunctionNodes)
-			{
-				UFunction* Function = FunctionNode->GetTargetFunction();
-				if ((Function == NULL) || !Function->HasMetaData(FBlueprintMetadata::MD_Latent))
-				{
-					continue;
-				}
-
-				UObject* NodeWorldContext = ActiveObjectBeingDebugged;
-				// if the node has a specific "world context" pin, attempt to get the value set for that first
-				if (Function->HasMetaData(FBlueprintMetadata::MD_WorldContext))
-				{
-					FString const WorldContextPinName = Function->GetMetaData(FBlueprintMetadata::MD_WorldContext);
-					if (UEdGraphPin* ContextPin = FunctionNode->FindPin(WorldContextPinName))
-					{
-						if (UObjectPropertyBase* ContextProperty = Cast<UObjectPropertyBase>(ClassDebugData.FindClassPropertyForPin(ContextPin)))
-						{
-							UObject* PropertyValue = ContextProperty->GetObjectPropertyValue_InContainer(ActiveObjectBeingDebugged);
-							if (PropertyValue != NULL)
-							{
-								NodeWorldContext = PropertyValue;
-							}
-						}
-					}
-				}
-				
-				LatentContextObjects.Add(NodeWorldContext);
-			}
-
-			for (UObject* ContextObject : LatentContextObjects)
-			{
-				if (UWorld* World = GEngine->GetWorldFromContextObject(ContextObject, /*bChecked =*/false))
-				{
-					FLatentActionManager& Manager = World->GetLatentActionManager();
-
-					TSet<int32> UUIDSet;
-					Manager.GetActiveUUIDs(ActiveObjectBeingDebugged, /*out*/ UUIDSet);
-
-					for (TSet<int32>::TConstIterator IterUUID(UUIDSet); IterUUID; ++IterUUID)
-					{
-						const int32 UUID = *IterUUID;
-
-						if (UEdGraphNode* ParentNode = ClassDebugData.FindNodeFromUUID(UUID))
-						{
- 							TArray<FObjectUUIDPair>& Pairs = NodesWithActiveLatentActions.FindOrAdd(ParentNode);
- 							new (Pairs) FObjectUUIDPair(ContextObject, UUID);
-						}
-					}
-				}
-			}
-		}
-
-		// Covert the watched pin array into a set
-		for (auto WatchedPinIt = SourceBlueprint->PinWatches.CreateConstIterator(); WatchedPinIt; ++WatchedPinIt)
-		{
-			UEdGraphPin* WatchedPin = *WatchedPinIt;
-
-			UEdGraphNode* OwningNode = Cast<UEdGraphNode>(WatchedPin->GetOuter());
-			if (!ensure(OwningNode != NULL)) // shouldn't happen, but just in case a dead pin was added to the PinWatches array
-			{
-				continue;
-			}
-			check(OwningNode == WatchedPin->GetOwningNode());
-
-			WatchedPinSet.Add(WatchedPin);
-			WatchedNodeSet.Add(OwningNode);
-		}
-	}
-}
-
 /////////////////////////////////////////////////////
 // FConnectionDrawingPolicy
 
@@ -544,901 +435,3 @@ void FConnectionDrawingPolicy::ApplyHoverDeemphasis(UEdGraphPin* OutputPin, UEdG
 		WireColor = FMath::Lerp<FLinearColor>(WireColor, DarkenedColor, HoverDeemphasisDarkFraction * TimeFraction);
 	}
 }
-
-/////////////////////////////////////////////////////
-// FKismetConnectionDrawingPolicy
-
-FKismetConnectionDrawingPolicy::FKismetConnectionDrawingPolicy(int32 InBackLayerID, int32 InFrontLayerID, float ZoomFactor, const FSlateRect& InClippingRect, FSlateWindowElementList& InDrawElements, UEdGraph* InGraphObj)
-	: FConnectionDrawingPolicy(InBackLayerID, InFrontLayerID, ZoomFactor, InClippingRect, InDrawElements)
-	, GraphObj(InGraphObj)
-{
-	// Don't want to draw ending arrowheads
-	ArrowImage = nullptr;
-	ArrowRadius = FVector2D::ZeroVector;
-
-	// But we do want to draw midpoint arrowheads
-	if (GetDefault<UBlueprintEditorSettings>()->bDrawMidpointArrowsInBlueprints)
-	{
-		MidpointImage = FEditorStyle::GetBrush( TEXT("Graph.Arrow") );
-		MidpointRadius = MidpointImage->ImageSize * ZoomFactor * 0.5f;
-	}
-
-	// Cache off the editor options
-	AttackColor = Settings->TraceAttackColor;
-	SustainColor = Settings->TraceSustainColor;
-	ReleaseColor = Settings->TraceReleaseColor;
-
-	AttackWireThickness = Settings->TraceAttackWireThickness;
-	SustainWireThickness = Settings->TraceSustainWireThickness;
-	ReleaseWireThickness = Settings->TraceReleaseWireThickness;
-
-	TracePositionBonusPeriod = Settings->TracePositionBonusPeriod;
-	TracePositionExponent = Settings->TracePositionExponent;
-	AttackHoldPeriod = Settings->TraceAttackHoldPeriod;
-	DecayPeriod = Settings->TraceDecayPeriod;
-	DecayExponent = Settings->TraceDecayExponent;
-	SustainHoldPeriod = Settings->TraceSustainHoldPeriod;
-	ReleasePeriod = Settings->TraceReleasePeriod;
-	ReleaseExponent = Settings->TraceReleaseExponent;
-
-	CurrentTime = 0.0;
-	LatestTimeDiscovered = 0.0;
-}
-
-void FKismetConnectionDrawingPolicy::Draw(TMap<TSharedRef<SWidget>, FArrangedWidget>& PinGeometries, FArrangedChildren& ArrangedNodes)
-{
-	// Build the execution roadmap (also populates execution times)
-	BuildExecutionRoadmap();
-
-	// Draw everything
-	FConnectionDrawingPolicy::Draw(PinGeometries, ArrangedNodes);
-}
-
-bool FKismetConnectionDrawingPolicy::CanBuildRoadmap() const
-{
-	UBlueprint* TargetBP = FBlueprintEditorUtils::FindBlueprintForGraphChecked(GraphObj);
-	UObject* ActiveObject = TargetBP->GetObjectBeingDebugged();
-
-	return ActiveObject != NULL;
-}
-
-void FKismetConnectionDrawingPolicy::BuildExecutionRoadmap()
-{
-	LatestTimeDiscovered = 0.0;
-
-	// Only do highlighting in PIE or SIE
-	if (!CanBuildRoadmap())
-	{
-		return;
-	}
-
-	UBlueprint* TargetBP = FBlueprintEditorUtils::FindBlueprintForGraphChecked(GraphObj);
-	UObject* ActiveObject = TargetBP->GetObjectBeingDebugged();
-	check(ActiveObject); // Due to CanBuildRoadmap
-
-	// Redirect the target Blueprint when debugging with a macro graph visible
-	if (TargetBP->BlueprintType == BPTYPE_MacroLibrary)
-	{
-		TargetBP = Cast<UBlueprint>(ActiveObject->GetClass()->ClassGeneratedBy);
-	}
-
-	TArray<UEdGraphNode*> SequentialNodesInGraph;
-	TArray<double> SequentialNodeTimes;
-	TArray<UEdGraphPin*> SequentialExecPinsInGraph;
-
-	{
-		const TSimpleRingBuffer<FKismetTraceSample>& TraceStack = FKismetDebugUtilities::GetTraceStack();
-
-		UBlueprintGeneratedClass* TargetClass = Cast<UBlueprintGeneratedClass>(TargetBP->GeneratedClass);
-		FBlueprintDebugData& DebugData = TargetClass->GetDebugData();
-
-		for (int32 i = 0; i < TraceStack.Num(); ++i)
-		{
-			const FKismetTraceSample& Sample = TraceStack(i);
-
-			if (UObject* TestObject = Sample.Context.Get())
-			{
-				if (TestObject == ActiveObject)
-				{
-					UEdGraphPin* AssociatedPin = DebugData.FindExecPinFromCodeLocation(Sample.Function.Get(), Sample.Offset);
-
-					if (UEdGraphNode* Node = DebugData.FindSourceNodeFromCodeLocation(Sample.Function.Get(), Sample.Offset, /*bAllowImpreciseHit=*/ false))
-					{
-						if (GraphObj == Node->GetGraph())
-						{
-							SequentialNodesInGraph.Add(Node);
-							SequentialNodeTimes.Add(Sample.ObservationTime);
-							SequentialExecPinsInGraph.Add(AssociatedPin);
-						}
-						else
-						{
-							// If the top-level source node is a macro instance node
-							UK2Node_MacroInstance* MacroInstanceNode = Cast<UK2Node_MacroInstance>(Node);
-							if (MacroInstanceNode)
-							{
-								// Attempt to locate the macro source node through the code mapping
-								UEdGraphNode* MacroSourceNode = DebugData.FindMacroSourceNodeFromCodeLocation(Sample.Function.Get(), Sample.Offset);
-								if (MacroSourceNode)
-								{
-									// If the macro source node is located in the current graph context
-									if (GraphObj == MacroSourceNode->GetGraph())
-									{
-										// Add it to the sequential node list
-										SequentialNodesInGraph.Add(MacroSourceNode);
-										SequentialNodeTimes.Add(Sample.ObservationTime);
-										SequentialExecPinsInGraph.Add(AssociatedPin);
-									}
-									else
-									{
-										// The macro source node isn't in the current graph context, but we might have a macro instance node that is
-										// in the current graph context, so obtain the set of macro instance nodes that are mapped to the code here.
-										TArray<UEdGraphNode*> MacroInstanceNodes;
-										DebugData.FindMacroInstanceNodesFromCodeLocation(Sample.Function.Get(), Sample.Offset, MacroInstanceNodes);
-
-										// For each macro instance node in the set
-										for (auto MacroInstanceNodeIt = MacroInstanceNodes.CreateConstIterator(); MacroInstanceNodeIt; ++MacroInstanceNodeIt)
-										{
-											// If the macro instance node is located in the current graph context
-											MacroInstanceNode = Cast<UK2Node_MacroInstance>(*MacroInstanceNodeIt);
-											if (MacroInstanceNode && GraphObj == MacroInstanceNode->GetGraph())
-											{
-												// Add it to the sequential node list
-												SequentialNodesInGraph.Add(MacroInstanceNode);
-												SequentialNodeTimes.Add(Sample.ObservationTime);
-												SequentialExecPinsInGraph.Add(AssociatedPin);
-
-												// Exit the loop; we're done
-												break;
-											}
-										}
-									}
-								}
-							}
-						}
-					}
-				}
-			}
-		}
-	}
-
-	// Run thru and apply bonus time
-	const float InvNumNodes = 1.0f / (float)SequentialNodeTimes.Num();
-	for (int32 i = 0; i < SequentialNodesInGraph.Num(); ++i)
-	{
-		double& ObservationTime = SequentialNodeTimes[i];
-
-		const float PositionRatio = (SequentialNodeTimes.Num() - i) * InvNumNodes;
-		const float PositionBonus = FMath::Pow(PositionRatio, TracePositionExponent) * TracePositionBonusPeriod;
-		ObservationTime += PositionBonus;
-
-		LatestTimeDiscovered = FMath::Max<double>(LatestTimeDiscovered, ObservationTime);
-	}
-
-	UEdGraphPin* LastExecPin = NULL;
-	// Record the unique exec-pin to time pairings, keeping only the most recent 
-	// times for each pairing... reverse the "SequentialNodes" because right now
-	// it is in stack order (with the last executed node first)
-	for (int32 i = SequentialNodesInGraph.Num() - 1; i >= 1; --i)
-	{
-		UEdGraphNode* CurNode  = SequentialNodesInGraph[i];
-		UEdGraphNode* NextNode = SequentialNodesInGraph[i-1];
-
-		// keep track of the last exec-pin executed by CurNode (these tracked 
-		// pins coincide with "WireTraceSite" op-codes that have been injected 
-		// prior to every "goto" statement... this way we have context for which
-		// pin executed the jump)
-		if (UEdGraphPin* AssociatedPin = SequentialExecPinsInGraph[i])
-		{
-			LastExecPin = AssociatedPin;
-		}
-		
-		// if this statement is a jump (from one node to another)
-		if (CurNode != NextNode)
-		{
-			// if there was a wire-trace op-code inserted before this jump
-			if (LastExecPin != NULL)
-			{
-				//ensure(LastExecPin->GetOwningNode() == CurNode);
-				double NextNodeTime = SequentialNodeTimes[i-1];
-
-				FExecPairingMap& ExecPaths  = PredecessorPins.FindOrAdd(NextNode);
-				FTimePair&       ExecTiming = ExecPaths.FindOrAdd(LastExecPin);
-				// make sure that if we've already visited this exec-pin (like 
-				// in a for-loop or something), that we're replacing it with a 
-				// more recent execution time
-				//
-				// @TODO I don't see when this wouldn't be the case
-				if (ExecTiming.ThisExecTime < NextNodeTime)
-				{
-					double CurNodeTime = SequentialNodeTimes[i];
-					ExecTiming.ThisExecTime = NextNodeTime;
-					ExecTiming.PredExecTime = CurNodeTime;
-				}
-			}
-			// if the nodes aren't graphically connected how could they be 
-			// executed back-to-back? well, this could be a pop back to a 
-			// sequence node from the end of one thread of execution, etc.
-			else if (AreNodesGraphicallySequential(CurNode, NextNode))
-			{
-				// only warn when the nodes are directly connected (this is all
-				// for execution flow visualization after all)
-				UE_LOG(LogConnectionDrawingPolicy, Warning, TEXT("Looks like a wire-trace was not injected before the jump from '%s' to '%s'."), 
-					*CurNode->GetNodeTitle(ENodeTitleType::FullTitle).ToString(), *NextNode->GetNodeTitle(ENodeTitleType::FullTitle).ToString());
-			}
-
-			// clear the exec-pin (we're moving to a new node and want to find 
-			// it's executed out pin)
-			LastExecPin = NULL;
-		}
-		// else, we're only collecting this data for tracing node-to-node
-		// executions (so we don't care about this sequence of statements)
-	}
-
-	// Fade only when free-running (since we're using FApp::GetCurrentTime(), instead of FPlatformTime::Seconds)
-	const double MaxTimeAhead = FMath::Min(FApp::GetCurrentTime() + 2*TracePositionBonusPeriod, LatestTimeDiscovered); //@TODO: Rough clamping; should be exposed as a parameter
-	CurrentTime = FMath::Max(FApp::GetCurrentTime(), MaxTimeAhead);
-}
-
-void FKismetConnectionDrawingPolicy::CalculateEnvelopeAlphas(double ExecutionTime, /*out*/ float& AttackAlpha, /*out*/ float& SustainAlpha) const
-{
-	const float DeltaTime = (float)(CurrentTime - ExecutionTime);
-
-	{
-		const float UnclampedDecayRatio = 1.0f - ((DeltaTime - AttackHoldPeriod) / DecayPeriod);
-		const float ClampedDecayRatio = FMath::Clamp<float>(UnclampedDecayRatio, 0.0f, 1.0f);
-		AttackAlpha = FMath::Pow(ClampedDecayRatio, DecayExponent);
-	}
-
-	{
-		const float SustainEndTime = AttackHoldPeriod + DecayPeriod + SustainHoldPeriod;
-
-		const float UnclampedReleaseRatio = 1.0f - ((DeltaTime - SustainEndTime) / ReleasePeriod);
-		const float ClampedReleaseRatio = FMath::Clamp<float>(UnclampedReleaseRatio, 0.0f, 1.0f);
-		SustainAlpha = FMath::Pow(ClampedReleaseRatio, ReleaseExponent);
-	}
-}
-
-bool FKismetConnectionDrawingPolicy::TreatWireAsExecutionPin(UEdGraphPin* InputPin, UEdGraphPin* OutputPin) const
-{
-	const UEdGraphSchema_K2* Schema = GetDefault<UEdGraphSchema_K2>();
-
-	return (InputPin != NULL) && (Schema->IsExecPin(*OutputPin));
-}
-
-bool FKismetConnectionDrawingPolicy::AreNodesGraphicallySequential(UEdGraphNode* InputNode, UEdGraphNode* OutputNode) const
-{
-	for (UEdGraphPin* Pin : InputNode->Pins)
-	{
-		if (Pin->Direction != EGPD_Output)
-		{
-			continue;
-		}
-
-		for (UEdGraphPin* Connection : Pin->LinkedTo)
-		{
-			if (!TreatWireAsExecutionPin(Pin, Connection))
-			{
-				continue;
-			}
-
-			if (Connection->GetOwningNode() == OutputNode)
-			{
-				return true;
-			}
-		}
-	}
-
-	return false;
-}
-
-void FKismetConnectionDrawingPolicy::DetermineStyleOfExecWire(float& Thickness, FLinearColor& WireColor, bool& bDrawBubbles, const FTimePair& Times)
-{
-	// It's a followed link, make it strong and yellowish but fading over time
-	const double ExecTime = Times.ThisExecTime;
-
-	float AttackAlpha;
-	float SustainAlpha;
-	CalculateEnvelopeAlphas(ExecTime, /*out*/ AttackAlpha, /*out*/ SustainAlpha);
-
-	const float DecayedAttackThickness = FMath::Lerp<float>(SustainWireThickness, AttackWireThickness, AttackAlpha);
-	Thickness = FMath::Lerp<float>(ReleaseWireThickness, DecayedAttackThickness, SustainAlpha);
-
-	const FLinearColor DecayedAttackColor = FMath::Lerp<FLinearColor>(SustainColor, AttackColor, AttackAlpha);
-	WireColor = WireColor * FMath::Lerp<FLinearColor>(ReleaseColor, DecayedAttackColor, SustainAlpha);
-
-	if (SustainAlpha > KINDA_SMALL_NUMBER)
-	{
-		bDrawBubbles = true;
-	}
-}
-
-FKismetConnectionDrawingPolicy::FTimePair const* FKismetConnectionDrawingPolicy::BackTraceExecPath(UEdGraphPin const* const OutputPin, FExecPairingMap const* const NodeExecutionList)
-{
-	FTimePair const* FoundExecPath = nullptr;
-
-	UEdGraphNode const* const OwningNode = OutputPin->GetOwningNode();
-	if (UK2Node_Knot const* const KnotNode = Cast<UK2Node_Knot>(OwningNode))
-	{
-		UEdGraphPin const* const KnotInputPin = KnotNode->GetInputPin();
-		for (UEdGraphPin const* KnotInput : KnotInputPin->LinkedTo)
-		{
-			FoundExecPath = BackTraceExecPath(KnotInput, NodeExecutionList);
-			if (FoundExecPath != nullptr)
-			{
-				break;
-			}
-		}
-	}
-	else
-	{
-		FoundExecPath = NodeExecutionList->Find(OutputPin);
-	}
-
-	return FoundExecPath;
-}
-
-// Give specific editor modes a chance to highlight this connection or darken non-interesting connections
-void FKismetConnectionDrawingPolicy::DetermineWiringStyle(UEdGraphPin* OutputPin, UEdGraphPin* InputPin, /*inout*/ float& Thickness, /*inout*/ FLinearColor& WireColor, /*inout*/bool& bDrawBubbles, /*inout*/bool& bBidirectional)
-{
-	// Get the schema and grab the default color from it
-	check(OutputPin);
-	check(GraphObj);
-	const UEdGraphSchema* Schema = GraphObj->GetSchema();
-
-	WireColor = Schema->GetPinTypeColor(OutputPin->PinType);
-
-	const bool bDeemphasizeUnhoveredPins = HoveredPins.Num() > 0;
-
-	// If this is a K2 graph, try to be a little more specific
-	const UEdGraphSchema_K2* K2Schema = Cast<const UEdGraphSchema_K2>(Schema);
-	if (K2Schema != NULL)
-	{
-		if (TreatWireAsExecutionPin(InputPin, OutputPin))
-		{
-			if (CanBuildRoadmap())
-			{
-				UEdGraphNode* InputNode = InputPin->GetOwningNode();
-				// knot nodes are removed from the graph at compile time, so we 
-				// have to follow them until we find something that would have 
-				// actually executed
-				while (UK2Node_Knot* KnotNode = Cast<UK2Node_Knot>(InputNode))
-				{
-					InputNode = nullptr;
-
-					UEdGraphPin* OutPin = KnotNode->GetOutputPin();
-					if (OutPin->LinkedTo.Num() > 0)
-					{
-						check(OutPin->LinkedTo.Num() == 1);
-						InputNode = OutPin->LinkedTo[0]->GetOwningNode();
-					}	
-				}
-
-				// track if this node connection was ran or not
-				bool bExecuted = false;
-
-				// if the node belonging to InputPin was actually executed
-				if (FExecPairingMap* ExecPaths = PredecessorPins.Find(InputNode))
-				{
-					// if the output pin is one of the pins that lead to InputNode being ran
-					if (FTimePair const* ExecTiming = BackTraceExecPath(OutputPin, ExecPaths))
-					{
-						bExecuted = true;
-						DetermineStyleOfExecWire(/*inout*/ Thickness, /*inout*/ WireColor, /*inout*/ bDrawBubbles, *ExecTiming);
-					}
-				}
-				
-				if (!bExecuted)
-				{
-					// It's not followed, fade it and keep it thin
-					WireColor = ReleaseColor;
-					Thickness = ReleaseWireThickness;
-				}
-			}
-			else
-			{
-				// Make exec wires slightly thicker even outside of debug
-				Thickness = 3.0f;
-			}
-		}
-		else
-		{
-			// Array types should draw thicker
-			if( (InputPin && InputPin->PinType.bIsArray) || (OutputPin && OutputPin->PinType.bIsArray) )
-			{
-				Thickness = 3.0f;
-			}
-		}
-	}
-
-	if (bDeemphasizeUnhoveredPins)
-	{
-		ApplyHoverDeemphasis(OutputPin, InputPin, /*inout*/ Thickness, /*inout*/ WireColor);
-	}
-}
-
-void FKismetConnectionDrawingPolicy::SetIncompatiblePinDrawState(const TSharedPtr<SGraphPin>& StartPin, const TSet< TSharedRef<SWidget> >& VisiblePins)
-{
-	ResetIncompatiblePinDrawState(VisiblePins);
-
-	for(auto VisiblePinIterator = VisiblePins.CreateConstIterator(); VisiblePinIterator; ++VisiblePinIterator)
-	{
-		TSharedPtr<SGraphPin> CheckPin = StaticCastSharedRef<SGraphPin>(*VisiblePinIterator);
-		if(CheckPin != StartPin)
-		{
-			const FPinConnectionResponse Response = StartPin->GetPinObj()->GetSchema()->CanCreateConnection(StartPin->GetPinObj(), CheckPin->GetPinObj());
-			if(Response.Response == CONNECT_RESPONSE_DISALLOW)
-			{
-				CheckPin->SetPinColorModifier(FLinearColor(0.25f, 0.25f, 0.25f, 0.5f));
-			}
-		}
-	}
-}
-
-void FKismetConnectionDrawingPolicy::ResetIncompatiblePinDrawState(const TSet< TSharedRef<SWidget> >& VisiblePins)
-{
-	for(auto VisiblePinIterator = VisiblePins.CreateConstIterator(); VisiblePinIterator; ++VisiblePinIterator)
-	{
-		TSharedPtr<SGraphPin> VisiblePin = StaticCastSharedRef<SGraphPin>(*VisiblePinIterator);
-		VisiblePin->SetPinColorModifier(FLinearColor::White);
-	}
-}
-
-/////////////////////////////////////////////////////
-// FStateMachineConnectionDrawingPolicy
-
-FStateMachineConnectionDrawingPolicy::FStateMachineConnectionDrawingPolicy(int32 InBackLayerID, int32 InFrontLayerID, float ZoomFactor, const FSlateRect& InClippingRect, FSlateWindowElementList& InDrawElements, UEdGraph* InGraphObj)
-	: FConnectionDrawingPolicy(InBackLayerID, InFrontLayerID, ZoomFactor, InClippingRect, InDrawElements)
-	, GraphObj(InGraphObj)
-{
-}
-
-void FStateMachineConnectionDrawingPolicy::DetermineWiringStyle(UEdGraphPin* OutputPin, UEdGraphPin* InputPin, /*inout*/ float& Thickness, /*inout*/ FLinearColor& WireColor, /*inout*/ bool& bDrawBubbles, /*inout*/ bool& bBidirectional)
-{
-	Thickness = 1.5f;
-
-	if (InputPin)
-	{
-		if (UAnimStateTransitionNode* TransNode = Cast<UAnimStateTransitionNode>(InputPin->GetOwningNode()))
-		{
-			WireColor = SGraphNodeAnimTransition::StaticGetTransitionColor(TransNode, HoveredPins.Contains(InputPin));
-
-			bBidirectional = TransNode->Bidirectional;
-		}
-	}
-
-	const bool bDeemphasizeUnhoveredPins = HoveredPins.Num() > 0;
-	if (bDeemphasizeUnhoveredPins)
-	{
-		ApplyHoverDeemphasis(OutputPin, InputPin, /*inout*/ Thickness, /*inout*/ WireColor);
-	}
-}
-
-void FStateMachineConnectionDrawingPolicy::DetermineLinkGeometry(
-	TMap<TSharedRef<SWidget>, FArrangedWidget>& PinGeometries,
-	FArrangedChildren& ArrangedNodes, 
-	TSharedRef<SWidget>& OutputPinWidget,
-	UEdGraphPin* OutputPin,
-	UEdGraphPin* InputPin,
-	/*out*/ FArrangedWidget*& StartWidgetGeometry,
-	/*out*/ FArrangedWidget*& EndWidgetGeometry
-	)
-{
-	if (UAnimStateEntryNode* EntryNode = Cast<UAnimStateEntryNode>(OutputPin->GetOwningNode()))
-	{
-		//FConnectionDrawingPolicy::DetermineLinkGeometry(PinGeometries, ArrangedNodes, OutputPinWidget, OutputPin, InputPin, StartWidgetGeometry, EndWidgetGeometry);
-
-		StartWidgetGeometry = PinGeometries.Find(OutputPinWidget);
-
-		UAnimStateNodeBase* State = CastChecked<UAnimStateNodeBase>(InputPin->GetOwningNode());
-		int32 StateIndex = NodeWidgetMap.FindChecked(State);
-		EndWidgetGeometry = &(ArrangedNodes[StateIndex]);
-	}
-	else if (UAnimStateTransitionNode* TransNode = Cast<UAnimStateTransitionNode>(InputPin->GetOwningNode()))
-	{
-		UAnimStateNodeBase* PrevState = TransNode->GetPreviousState();
-		UAnimStateNodeBase* NextState = TransNode->GetNextState();
-		if ((PrevState != NULL) && (NextState != NULL))
-		{
-			int32* PrevNodeIndex = NodeWidgetMap.Find(PrevState);
-			int32* NextNodeIndex = NodeWidgetMap.Find(NextState);
-			if ((PrevNodeIndex != NULL) && (NextNodeIndex != NULL))
-			{
-				StartWidgetGeometry = &(ArrangedNodes[*PrevNodeIndex]);
-				EndWidgetGeometry = &(ArrangedNodes[*NextNodeIndex]);
-			}
-		}
-	}
-	else
-	{
-		StartWidgetGeometry = PinGeometries.Find(OutputPinWidget);
-
-		if (TSharedRef<SGraphPin>* pTargetWidget = PinToPinWidgetMap.Find(InputPin))
-		{
-			TSharedRef<SGraphPin> InputWidget = *pTargetWidget;
-			EndWidgetGeometry = PinGeometries.Find(InputWidget);
-		}
-	}
-}
-
-void FStateMachineConnectionDrawingPolicy::Draw(TMap<TSharedRef<SWidget>, FArrangedWidget>& PinGeometries, FArrangedChildren& ArrangedNodes)
-{
-	// Build an acceleration structure to quickly find geometry for the nodes
-	NodeWidgetMap.Empty();
-	for (int32 NodeIndex = 0; NodeIndex < ArrangedNodes.Num(); ++NodeIndex)
-	{
-		FArrangedWidget& CurWidget = ArrangedNodes[NodeIndex];
-		TSharedRef<SGraphNode> ChildNode = StaticCastSharedRef<SGraphNode>(CurWidget.Widget);
-		NodeWidgetMap.Add(ChildNode->GetNodeObj(), NodeIndex);
-	}
-
-	// Now draw
-	FConnectionDrawingPolicy::Draw(PinGeometries, ArrangedNodes);
-}
-
-void FStateMachineConnectionDrawingPolicy::DrawPreviewConnector(const FGeometry& PinGeometry, const FVector2D& StartPoint, const FVector2D& EndPoint, UEdGraphPin* Pin)
-{
-	float Thickness = 1.0f;
-	FLinearColor WireColor = FLinearColor::White;
-	bool bDrawBubbles = false;
-	bool bBiDirectional = false;
-	DetermineWiringStyle(Pin, NULL, /*inout*/ Thickness, /*inout*/ WireColor, /*inout*/ bDrawBubbles, /*inout*/ bBiDirectional);
-
-
-	const FVector2D SeedPoint = EndPoint;
-	const FVector2D AdjustedStartPoint = FGeometryHelper::FindClosestPointOnGeom(PinGeometry, SeedPoint);
-
-	DrawSplineWithArrow(AdjustedStartPoint, EndPoint, WireColor, Thickness, bDrawBubbles, bBiDirectional);
-}
-
-
-void FStateMachineConnectionDrawingPolicy::DrawSplineWithArrow(const FVector2D& StartAnchorPoint, const FVector2D& EndAnchorPoint, const FLinearColor& WireColor, float WireThickness, bool bDrawBubbles, bool Bidirectional)
-{
-	Internal_DrawLineWithArrow(StartAnchorPoint, EndAnchorPoint, WireColor, WireThickness, bDrawBubbles);
-	if (Bidirectional)
-	{
-		Internal_DrawLineWithArrow(EndAnchorPoint, StartAnchorPoint, WireColor, WireThickness, bDrawBubbles);
-	}
-}
-
-void FStateMachineConnectionDrawingPolicy::Internal_DrawLineWithArrow(const FVector2D& StartAnchorPoint, const FVector2D& EndAnchorPoint, const FLinearColor& WireColor, float WireThickness, bool bDrawBubbles)
-{
-	//@TODO: Should this be scaled by zoom factor?
-	const float LineSeparationAmount = 4.5f;
-
-	const FVector2D DeltaPos = EndAnchorPoint - StartAnchorPoint;
-	const FVector2D UnitDelta = DeltaPos.SafeNormal();
-	const FVector2D Normal = FVector2D(DeltaPos.Y, -DeltaPos.X).SafeNormal();
-
-	// Come up with the final start/end points
-	const FVector2D DirectionBias = Normal * LineSeparationAmount;
-	const FVector2D LengthBias = ArrowRadius.X * UnitDelta;
-	const FVector2D StartPoint = StartAnchorPoint + DirectionBias + LengthBias;
-	const FVector2D EndPoint = EndAnchorPoint + DirectionBias - LengthBias;
-
-	// Draw a line/spline
-	DrawConnection(WireLayerID, StartPoint, EndPoint, WireColor, WireThickness, bDrawBubbles);
-
-	// Draw the arrow
-	const FVector2D ArrowDrawPos = EndPoint - ArrowRadius;
-	const float AngleInRadians = FMath::Atan2(DeltaPos.Y, DeltaPos.X);
-
-	FSlateDrawElement::MakeRotatedBox(
-		DrawElementsList,
-		ArrowLayerID,
-		FPaintGeometry(ArrowDrawPos, ArrowImage->ImageSize * ZoomFactor, ZoomFactor),
-		ArrowImage,
-		ClippingRect,
-		ESlateDrawEffect::None,
-		AngleInRadians,
-		TOptional<FVector2D>(),
-		FSlateDrawElement::RelativeToElement,
-		WireColor
-		);
-}
-
-void FStateMachineConnectionDrawingPolicy::DrawSplineWithArrow(FGeometry& StartGeom, FGeometry& EndGeom, const FLinearColor& WireColor, float WireThickness, bool bDrawBubbles, bool Bidirectional)
-{
-	// Get a reasonable seed point (halfway between the boxes)
-	const FVector2D StartCenter = FGeometryHelper::CenterOf(StartGeom);
-	const FVector2D EndCenter = FGeometryHelper::CenterOf(EndGeom);
-	const FVector2D SeedPoint = (StartCenter + EndCenter) * 0.5f;
-	
-	// Find the (approximate) closest points between the two boxes
-	const FVector2D StartAnchorPoint = FGeometryHelper::FindClosestPointOnGeom(StartGeom, SeedPoint);
-	const FVector2D EndAnchorPoint = FGeometryHelper::FindClosestPointOnGeom(EndGeom, SeedPoint);
-
-	DrawSplineWithArrow(StartAnchorPoint, EndAnchorPoint, WireColor, WireThickness, bDrawBubbles, Bidirectional);
-}
-
-void FStateMachineConnectionDrawingPolicy::DrawConnection(int32 LayerId, const FVector2D& Start, const FVector2D& End, const FLinearColor& InColor, float Thickness, bool bDrawBubbles)
-{
-	const FVector2D& P0 = Start;
-	const FVector2D& P1 = End;
-
-	const FVector2D Delta = End-Start;
-	const FVector2D NormDelta = Delta.SafeNormal();
-	
-	const FVector2D P0Tangent = NormDelta;
-	const FVector2D P1Tangent = NormDelta;
-
-	// Draw the spline itself
-	FSlateDrawElement::MakeDrawSpaceSpline(
-		DrawElementsList,
-		LayerId,
-		P0, P0Tangent,
-		P1, P1Tangent,
-		ClippingRect,
-		Thickness,
-		ESlateDrawEffect::None,
-		InColor
-		);
-
-	//@TODO: Handle bDrawBubbles
-}
-
-/////////////////////////////////////////////////////
-// FAnimGraphConnectionDrawingPolicy
-
-FAnimGraphConnectionDrawingPolicy::FAnimGraphConnectionDrawingPolicy(int32 InBackLayerID, int32 InFrontLayerID, float ZoomFactor, const FSlateRect& InClippingRect, FSlateWindowElementList& InDrawElements, UEdGraph* InGraphObj)
-	: FKismetConnectionDrawingPolicy(InBackLayerID, InFrontLayerID, ZoomFactor, InClippingRect, InDrawElements, InGraphObj)
-{
-}
-
-bool FAnimGraphConnectionDrawingPolicy::TreatWireAsExecutionPin(UEdGraphPin* InputPin, UEdGraphPin* OutputPin) const
-{
-	const UAnimationGraphSchema* Schema = GetDefault<UAnimationGraphSchema>();
-
-	return (InputPin != NULL) && (Schema->IsPosePin(OutputPin->PinType));
-}
-
-void FAnimGraphConnectionDrawingPolicy::BuildExecutionRoadmap()
-{
-	UAnimBlueprint* TargetBP = CastChecked<UAnimBlueprint>(FBlueprintEditorUtils::FindBlueprintForGraphChecked(GraphObj));
-	UAnimBlueprintGeneratedClass* AnimBlueprintClass = (UAnimBlueprintGeneratedClass*)(*(TargetBP->GeneratedClass));
-
-	if (TargetBP->GetObjectBeingDebugged() == NULL)
-	{
-		return;
-	}
-
-	TMap<UProperty*, UObject*> PropertySourceMap;
-	AnimBlueprintClass->GetDebugData().GenerateReversePropertyMap(/*out*/ PropertySourceMap);
-
-	FAnimBlueprintDebugData& DebugInfo = AnimBlueprintClass->GetAnimBlueprintDebugData();
-	for (auto VisitIt = DebugInfo.UpdatedNodesThisFrame.CreateIterator(); VisitIt; ++VisitIt)
-	{
-		const FAnimBlueprintDebugData::FNodeVisit& VisitRecord = *VisitIt;
-
-		if ((VisitRecord.SourceID >= 0) && (VisitRecord.SourceID < AnimBlueprintClass->AnimNodeProperties.Num()) && (VisitRecord.TargetID >= 0) && (VisitRecord.TargetID < AnimBlueprintClass->AnimNodeProperties.Num()))
-		{
-			if (UAnimGraphNode_Base* SourceNode = Cast<UAnimGraphNode_Base>(PropertySourceMap.FindRef(AnimBlueprintClass->AnimNodeProperties[VisitRecord.SourceID])))
-			{
-				if (UAnimGraphNode_Base* TargetNode = Cast<UAnimGraphNode_Base>(PropertySourceMap.FindRef(AnimBlueprintClass->AnimNodeProperties[VisitRecord.TargetID])))
-				{
-					UEdGraphPin* PoseNet = NULL;
-
-					UAnimationGraphSchema const* AnimSchema = GetDefault<UAnimationGraphSchema>();
-					for (int32 PinIndex = 0; PinIndex < TargetNode->Pins.Num(); ++PinIndex)
-					{
-						UEdGraphPin* Pin = TargetNode->Pins[PinIndex];
-						if (AnimSchema->IsPosePin(Pin->PinType) && (Pin->Direction == EGPD_Output))
-						{
-							PoseNet = Pin;
-							break;
-						}
-					}
-
-					if (PoseNet != NULL)
-					{
-						//@TODO: Extend the rendering code to allow using the recorded blend weight instead of faked exec times
-						FExecPairingMap& Predecessors = PredecessorPins.FindOrAdd((UEdGraphNode*)SourceNode);
-						FTimePair& Timings = Predecessors.FindOrAdd(PoseNet);
-						Timings.PredExecTime = 0.0;
-						Timings.ThisExecTime = VisitRecord.Weight;
-					}
-				}
-			}
-		}
-	}
-}
-
-void FAnimGraphConnectionDrawingPolicy::DetermineStyleOfExecWire(float& Thickness, FLinearColor& WireColor, bool& bDrawBubbles, const FTimePair& Times)
-{
-	// It's a followed link, make it strong and yellowish but fading over time
-	const double BlendWeight = Times.ThisExecTime;
-
-	const float HeavyBlendThickness = AttackWireThickness;
-	const float LightBlendThickness = SustainWireThickness;
-
-	Thickness = FMath::Lerp<float>(LightBlendThickness, HeavyBlendThickness, BlendWeight);
-	WireColor = WireColor * (BlendWeight * 0.5f + 0.5f);//FMath::Lerp<FLinearColor>(SustainColor, AttackColor, BlendWeight);
-
-	bDrawBubbles = true;
-}
-
-/////////////////////////////////////////////////////
-// FSoundCueGraphConnectionDrawingPolicy
-
-FSoundCueGraphConnectionDrawingPolicy::FSoundCueGraphConnectionDrawingPolicy(int32 InBackLayerID, int32 InFrontLayerID, float ZoomFactor, const FSlateRect& InClippingRect, FSlateWindowElementList& InDrawElements, UEdGraph* InGraphObj)
-	: FConnectionDrawingPolicy(InBackLayerID, InFrontLayerID, ZoomFactor, InClippingRect, InDrawElements)
-	, GraphObj(InGraphObj)
-{
-	// Cache off the editor options
-	ActiveColor = Settings->TraceAttackColor;
-	InactiveColor = Settings->TraceReleaseColor;
-
-	ActiveWireThickness = Settings->TraceAttackWireThickness;
-	InactiveWireThickness = Settings->TraceReleaseWireThickness;
-
-	// Don't want to draw ending arrowheads
-	ArrowImage = nullptr;
-	ArrowRadius = FVector2D::ZeroVector;
-}
-
-void FSoundCueGraphConnectionDrawingPolicy::Draw(TMap<TSharedRef<SWidget>, FArrangedWidget>& PinGeometries, FArrangedChildren& ArrangedNodes)
-{
-	// Build the execution roadmap (also populates execution times)
-	BuildAudioFlowRoadmap();
-
-	// Draw everything
-	FConnectionDrawingPolicy::Draw(PinGeometries, ArrangedNodes);
-}
-
-void FSoundCueGraphConnectionDrawingPolicy::BuildAudioFlowRoadmap()
-{
-	FAudioDevice* AudioDevice = GEngine->GetAudioDevice();
-
-	if (AudioDevice)
-	{
-		USoundCueGraph* SoundCueGraph = CastChecked<USoundCueGraph>(GraphObj);
-		USoundCue* SoundCue = SoundCueGraph->GetSoundCue();
-
-		UAudioComponent* PreviewAudioComponent = GEditor->GetPreviewAudioComponent();
-
-		if (PreviewAudioComponent && PreviewAudioComponent->IsPlaying() && PreviewAudioComponent->Sound == SoundCue)
-		{
-			TArray<FWaveInstance*> WaveInstances;
-			const int32 FirstActiveIndex = AudioDevice->GetSortedActiveWaveInstances(WaveInstances, ESortedActiveWaveGetType::QueryOnly);
-
-			// Run through the active instances and cull out anything that isn't related to this graph
-			if (FirstActiveIndex > 0)
-			{
-				WaveInstances.RemoveAt(0, FirstActiveIndex + 1);
-			}
-
-			for (int32 WaveIndex = WaveInstances.Num() - 1; WaveIndex >= 0 ; --WaveIndex)
-			{
-				UAudioComponent* WaveInstanceAudioComponent = WaveInstances[WaveIndex]->ActiveSound->AudioComponent.Get();
-				if (WaveInstanceAudioComponent != PreviewAudioComponent)
-				{
-					WaveInstances.RemoveAtSwap(WaveIndex);
-				}
-			}
-
-			for (int32 WaveIndex = 0; WaveIndex < WaveInstances.Num(); ++WaveIndex)
-			{
-				TArray<USoundNode*> PathToWaveInstance;
-				if (SoundCue->FindPathToNode(WaveInstances[WaveIndex]->WaveInstanceHash, PathToWaveInstance))
-				{
-					TArray<USoundCueGraphNode_Root*> RootNode;
-					TArray<UEdGraphNode*> GraphNodes;
-					SoundCueGraph->GetNodesOfClass<USoundCueGraphNode_Root>(RootNode);
-					check(RootNode.Num() == 1);
-					GraphNodes.Add(RootNode[0]);
-
-					TArray<double> NodeTimes;
-					NodeTimes.Add(FApp::GetCurrentTime()); // Time for the root node
-
-					for (int32 i = 0; i < PathToWaveInstance.Num(); ++i)
-					{
-						const double ObservationTime = FApp::GetCurrentTime() + 1.f;
-
-						NodeTimes.Add(ObservationTime);
-						GraphNodes.Add(PathToWaveInstance[i]->GraphNode);
-					}
-
-					// Record the unique node->node pairings, keeping only the most recent times for each pairing
-					for (int32 i = GraphNodes.Num() - 1; i >= 1; --i)
-					{
-						UEdGraphNode* CurNode = GraphNodes[i];
-						double CurNodeTime = NodeTimes[i];
-						UEdGraphNode* NextNode = GraphNodes[i-1];
-						double NextNodeTime = NodeTimes[i-1];
-
-						FExecPairingMap& Predecessors = PredecessorNodes.FindOrAdd(NextNode);
-
-						// Update the timings if this is a more recent pairing
-						FTimePair& Timings = Predecessors.FindOrAdd(CurNode);
-						if (Timings.ThisExecTime < NextNodeTime)
-						{
-							Timings.PredExecTime = CurNodeTime;
-							Timings.ThisExecTime = NextNodeTime;
-						}
-					}
-				}
-			}
-		}
-	}
-}
-
-// Give specific editor modes a chance to highlight this connection or darken non-interesting connections
-void FSoundCueGraphConnectionDrawingPolicy::DetermineWiringStyle(UEdGraphPin* OutputPin, UEdGraphPin* InputPin, /*inout*/ float& Thickness, /*inout*/ FLinearColor& WireColor, /*inout*/bool& bDrawBubbles, /*inout*/bool& bBidirectional)
-{
-	// Get the schema and grab the default color from it
-	check(OutputPin);
-	check(GraphObj);
-	const UEdGraphSchema* Schema = GraphObj->GetSchema();
-
-	WireColor = Schema->GetPinTypeColor(OutputPin->PinType);
-
-	if (InputPin == NULL)
-	{
-		return;
-	}
-	
-	bool bExecuted = false;
-
-	// Run thru the predecessors, and on
-	if (FExecPairingMap* PredecessorMap = PredecessorNodes.Find(InputPin->GetOwningNode()))
-	{
-		if (FTimePair* Times = PredecessorMap->Find(OutputPin->GetOwningNode()))
-		{
-			bExecuted = true;
-
-			Thickness = ActiveWireThickness;
-			WireColor = ActiveColor;
-
-			bDrawBubbles = true;
-		}
-	}
-
-	if (!bExecuted)
-	{
-		// It's not followed, fade it and keep it thin
-		WireColor = InactiveColor;
-		Thickness = InactiveWireThickness;
-	}
-}
-
-/////////////////////////////////////////////////////
-// FMaterialGraphConnectionDrawingPolicy
-
-FMaterialGraphConnectionDrawingPolicy::FMaterialGraphConnectionDrawingPolicy(int32 InBackLayerID, int32 InFrontLayerID, float ZoomFactor, const FSlateRect& InClippingRect, FSlateWindowElementList& InDrawElements, UEdGraph* InGraphObj)
-	: FConnectionDrawingPolicy(InBackLayerID, InFrontLayerID, ZoomFactor, InClippingRect, InDrawElements)
-	, MaterialGraph(CastChecked<UMaterialGraph>(InGraphObj))
-	, MaterialGraphSchema(CastChecked<UMaterialGraphSchema>(InGraphObj->GetSchema()))
-{
-	// Don't want to draw ending arrowheads
-	ArrowImage = nullptr;
-	ArrowRadius = FVector2D::ZeroVector;
-
-	// Still need to be able to perceive the graph while dragging connectors, esp over comment boxes
-	HoverDeemphasisDarkFraction = 0.4f;
-}
-
-void FMaterialGraphConnectionDrawingPolicy::DetermineWiringStyle(UEdGraphPin* OutputPin, UEdGraphPin* InputPin, /*inout*/ float& Thickness, /*inout*/ FLinearColor& WireColor, /*inout*/bool& bDrawBubbles, /*inout*/ bool& bBidirectional)
-{
-	WireColor = MaterialGraphSchema->ActivePinColor;
-
-	// Have to consider both pins as the input will be an 'output' when previewing a connection
-	if (OutputPin)
-	{
-		if (!MaterialGraph->IsInputActive(OutputPin))
-		{
-			WireColor = MaterialGraphSchema->InactivePinColor;
-		}
-	}
-	if (InputPin)
-	{
-		if (!MaterialGraph->IsInputActive(InputPin))
-		{
-			WireColor = MaterialGraphSchema->InactivePinColor;
-		}
-	}
-
-	const bool bDeemphasizeUnhoveredPins = HoveredPins.Num() > 0;
-	if (bDeemphasizeUnhoveredPins)
-	{
-		ApplyHoverDeemphasis(OutputPin, InputPin, /*inout*/ Thickness, /*inout*/ WireColor);
-	}
-}
diff --git a/Engine/Source/Editor/GraphEditor/Private/ConnectionDrawingPolicy.h b/Engine/Source/Editor/GraphEditor/Private/ConnectionDrawingPolicy.h
--- a/Engine/Source/Editor/GraphEditor/Private/ConnectionDrawingPolicy.h
+++ b/Engine/Source/Editor/GraphEditor/Private/ConnectionDrawingPolicy.h
@@ -3,7 +3,11 @@
 #pragma once
 
 /////////////////////////////////////////////////////
-// FConnectionDrawingPolicy
+
+GRAPHEDITOR_API DECLARE_LOG_CATEGORY_EXTERN(LogConnectionDrawingPolicy, Log, All);
+
+/////////////////////////////////////////////////////
+// FGeometryHelper
 
 class GRAPHEDITOR_API FGeometryHelper
 {
@@ -90,220 +94,3 @@ public:
 
 	virtual bool IsConnectionCulled( const FArrangedWidget& StartLink, const FArrangedWidget& EndLink ) const;
 };
-
-/////////////////////////////////////////////////////
-// FKismetConnectionDrawingPolicy
-
-// This class draws the connections for an UEdGraph using a K2-based schema
-class FKismetConnectionDrawingPolicy : public FConnectionDrawingPolicy
-{
-protected:
-	// Times for one execution pair within the current graph
-	struct FTimePair
-	{
-		double PredExecTime; // time that the node with the associated exec-out was executed
-		double ThisExecTime; // time that the following node was executed
-
-		FTimePair()
-			: PredExecTime(0.0)
-			, ThisExecTime(0.0)
-		{
-		}
-	};
-	/** 
-	 * Map of exec-pins to a pair of execution times (holds timing info for the 
-	 * owning-node, paired with a time of execution for the linked node)... only 
-	 * one timing-pair per pin because exec outputs can only connect to one node
-	 */
-	typedef TMap<UEdGraphPin*, FTimePair> FExecPairingMap;
-
-	// Map of executed nodes, and the execution pins that lead to them being ran
-	TMap<UEdGraphNode*, FExecPairingMap> PredecessorPins;
-
-	UEdGraph* GraphObj;
-	double CurrentTime;
-	double LatestTimeDiscovered;
-
-	FLinearColor AttackColor;
-	FLinearColor SustainColor;
-	FLinearColor ReleaseColor;
-
-	float AttackWireThickness;
-	float SustainWireThickness;
-	float ReleaseWireThickness;
-
-	float TracePositionBonusPeriod;
-	float TracePositionExponent;
-	float AttackHoldPeriod;
-	float DecayPeriod;
-	float DecayExponent;
-	float SustainHoldPeriod;
-	float ReleasePeriod;
-	float ReleaseExponent;
-
-protected:
-	// Should this wire be treated as an execution pin, querying the execution trace for connected nodes to draw it differently?
-	virtual bool TreatWireAsExecutionPin(UEdGraphPin* InputPin, UEdGraphPin* OutputPin) const;
-public:
-	FKismetConnectionDrawingPolicy(int32 InBackLayerID, int32 InFrontLayerID, float ZoomFactor, const FSlateRect& InClippingRect, FSlateWindowElementList& InDrawElements, UEdGraph* InGraphObj);
-
-	virtual void BuildExecutionRoadmap();
-
-	virtual bool CanBuildRoadmap() const;
-
-	void CalculateEnvelopeAlphas(double ExecutionTime, /*out*/ float& AttackAlpha, /*out*/ float& SustainAlpha) const;
-
-	// FConnectionDrawingPolicy interface
-	virtual void DetermineWiringStyle(UEdGraphPin* OutputPin, UEdGraphPin* InputPin, /*inout*/ float& Thickness, /*inout*/ FLinearColor& WireColor, /*inout*/bool& bDrawBubbles, /*inout*/ bool& bBidirectional) override;
-	virtual void Draw(TMap<TSharedRef<SWidget>, FArrangedWidget>& PinGeometries, FArrangedChildren& ArrangedNodes) override;
-	// End of FConnectionDrawingPolicy interface
-
-	virtual void DetermineStyleOfExecWire(float& Thickness, FLinearColor& WireColor, bool& bDrawBubbles, const FTimePair& Times);
-
-	virtual void SetIncompatiblePinDrawState(const TSharedPtr<SGraphPin>& StartPin, const TSet< TSharedRef<SWidget> >& VisiblePins);
-	virtual void ResetIncompatiblePinDrawState(const TSet< TSharedRef<SWidget> >& VisiblePins);
-
-private: 
-	/**
-	* Checks that the specified nodes are connected in the graph directly by an 
-	* exec pin.
-	* 
-	* @param  InputNode	The node whose exec output(s) should be connected to OutputNode.
-	* @param  OutputNode	The node whose exec input should be coming from the InputNode.
-	* @return True if the two nodes are connected in the graph, otherwise false.
-	*/
-	bool AreNodesGraphicallySequential(UEdGraphNode* InputNode, UEdGraphNode* OutputNode) const;
-
-	/**
-	 * With the introduction of "reroute" (UK2Node_Knot) nodes, we have to 
-	 * backtrace to find a node that was actually executed before the one in 
-	 * question (knot nodes are removed at compile time, and therefore never
-	 * executed).
-	 *
-	 * This function recursively searches for any pins feeding the OutputPin (or 
-	 * the OutputPin itself); it looks for a pin on a non-reroute node (one that
-	 * is listed in NodeExecutionList).
-	 * 
-	 * @param  OutputPin			An exec pin, leading into the node in question (an output pin on some other node).
-	 * @param  NodeExecutionList	A list of pins that feed the execution of a specific node.
-	 * @return Null if OutputPin did not cause the invocation of the node in question, otherwise a valid time pair denoting when the node was executed.
-	 */
-	FTimePair const* BackTraceExecPath(UEdGraphPin const* const OutputPin, FExecPairingMap const* const NodeExecutionList);
-};
-
-/////////////////////////////////////////////////////
-// FStateMachineConnectionDrawingPolicy
-
-// This class draws the connections for an UEdGraph with a state machine schema
-class FStateMachineConnectionDrawingPolicy : public FConnectionDrawingPolicy
-{
-protected:
-	UEdGraph* GraphObj;
-
-	TMap<UEdGraphNode*, int32> NodeWidgetMap;
-public:
-	//
-	FStateMachineConnectionDrawingPolicy(int32 InBackLayerID, int32 InFrontLayerID, float ZoomFactor, const FSlateRect& InClippingRect, FSlateWindowElementList& InDrawElements, UEdGraph* InGraphObj);
-
-	// FConnectionDrawingPolicy interface
-	virtual void DetermineWiringStyle(UEdGraphPin* OutputPin, UEdGraphPin* InputPin, /*inout*/ float& Thickness, /*inout*/ FLinearColor& WireColor, /*inout*/ bool& bDrawBubbles, /*inout*/ bool& bBidirectional) override;
-	virtual void Draw(TMap<TSharedRef<SWidget>, FArrangedWidget>& PinGeometries, FArrangedChildren& ArrangedNodes) override;
-	virtual void DetermineLinkGeometry(
-		TMap<TSharedRef<SWidget>,
-		FArrangedWidget>& PinGeometries,
-		FArrangedChildren& ArrangedNodes, 
-		TSharedRef<SWidget>& OutputPinWidget,
-		UEdGraphPin* OutputPin,
-		UEdGraphPin* InputPin,
-		/*out*/ FArrangedWidget*& StartWidgetGeometry,
-		/*out*/ FArrangedWidget*& EndWidgetGeometry
-		) override;
-	virtual void DrawSplineWithArrow(FGeometry& StartGeom, FGeometry& EndGeom, const FLinearColor& WireColor, float WireThickness, bool bDrawBubbles, bool Bidirectional) override;
-	virtual void DrawSplineWithArrow(const FVector2D& StartPoint, const FVector2D& EndPoint, const FLinearColor& WireColor, float WireThickness, bool bDrawBubbles, bool Bidirectional) override;
-	virtual void DrawPreviewConnector(const FGeometry& PinGeometry, const FVector2D& StartPoint, const FVector2D& EndPoint, UEdGraphPin* Pin) override;
-	virtual void DrawConnection(int32 LayerId, const FVector2D& Start, const FVector2D& End, const FLinearColor& InColor, float Thickness, bool bDrawBubbles) override;
-	// End of FConnectionDrawingPolicy interface
-
-protected:
-	void Internal_DrawLineWithArrow(const FVector2D& StartAnchorPoint, const FVector2D& EndAnchorPoint, const FLinearColor& WireColor, float WireThickness, bool bDrawBubbles);
-};
-
-/////////////////////////////////////////////////////
-// FAnimGraphConnectionDrawingPolicy
-
-// This class draws the connections for an UEdGraph with an animation schema
-class FAnimGraphConnectionDrawingPolicy : public FKismetConnectionDrawingPolicy
-{
-public:
-	// Constructor
-	FAnimGraphConnectionDrawingPolicy(int32 InBackLayerID, int32 InFrontLayerID, float ZoomFactor, const FSlateRect& InClippingRect, FSlateWindowElementList& InDrawElements, UEdGraph* InGraphObj);
-
-	// FKismetConnectionDrawingPolicy interface
-	virtual bool TreatWireAsExecutionPin(UEdGraphPin* InputPin, UEdGraphPin* OutputPin) const override;
-	virtual void BuildExecutionRoadmap() override;
-	virtual void DetermineStyleOfExecWire(float& Thickness, FLinearColor& WireColor, bool& bDrawBubbles, const FTimePair& Times) override;
-	// End of FKismetConnectionDrawingPolicy interface
-};
-
-/////////////////////////////////////////////////////
-// FSoundCueGraphConnectionDrawingPolicy
-
-// This class draws the connections for an UEdGraph using a SoundCue schema
-class FSoundCueGraphConnectionDrawingPolicy : public FConnectionDrawingPolicy
-{
-protected:
-	// Times for one execution pair within the current graph
-	struct FTimePair
-	{
-		double PredExecTime;
-		double ThisExecTime;
-
-		FTimePair()
-			: PredExecTime(0.0)
-			, ThisExecTime(0.0)
-		{
-		}
-	};
-
-	// Map of pairings
-	typedef TMap<UEdGraphNode*, FTimePair> FExecPairingMap;
-
-	// Map of nodes that preceeded before a given node in the execution sequence (one entry for each pairing)
-	TMap<UEdGraphNode*, FExecPairingMap> PredecessorNodes;
-
-	UEdGraph* GraphObj;
-
-	FLinearColor ActiveColor;
-	FLinearColor InactiveColor;
-
-	float ActiveWireThickness;
-	float InactiveWireThickness;
-
-public:
-	FSoundCueGraphConnectionDrawingPolicy(int32 InBackLayerID, int32 InFrontLayerID, float ZoomFactor, const FSlateRect& InClippingRect, FSlateWindowElementList& InDrawElements, UEdGraph* InGraphObj);
-
-	void BuildAudioFlowRoadmap();
-
-	// FConnectionDrawingPolicy interface
-	virtual void DetermineWiringStyle(UEdGraphPin* OutputPin, UEdGraphPin* InputPin, /*inout*/ float& Thickness, /*inout*/ FLinearColor& WireColor, /*inout*/bool& bDrawBubbles, /*inout*/ bool& bBidirectional) override;
-	virtual void Draw(TMap<TSharedRef<SWidget>, FArrangedWidget>& PinGeometries, FArrangedChildren& ArrangedNodes) override;
-	// End of FConnectionDrawingPolicy interface
-};
-
-/////////////////////////////////////////////////////
-// FMaterialGraphConnectionDrawingPolicy
-
-// This class draws the connections for an UEdGraph using a MaterialGraph schema
-class FMaterialGraphConnectionDrawingPolicy : public FConnectionDrawingPolicy
-{
-protected:
-	UMaterialGraph* MaterialGraph;
-	const UMaterialGraphSchema* MaterialGraphSchema;
-
-public:
-	FMaterialGraphConnectionDrawingPolicy(int32 InBackLayerID, int32 InFrontLayerID, float ZoomFactor, const FSlateRect& InClippingRect, FSlateWindowElementList& InDrawElements, UEdGraph* InGraphObj);
-
-	// FConnectionDrawingPolicy interface
-	virtual void DetermineWiringStyle(UEdGraphPin* OutputPin, UEdGraphPin* InputPin, /*inout*/ float& Thickness, /*inout*/ FLinearColor& WireColor, /*inout*/bool& bDrawBubbles, /*inout*/ bool& bBidirectional) override;
-	// End of FConnectionDrawingPolicy interface
-};
diff --git a/Engine/Source/Editor/GraphEditor/Private/KismetNodes/KismetNodeInfoContext.cpp b/Engine/Source/Editor/GraphEditor/Private/KismetNodes/KismetNodeInfoContext.cpp
new file mode 100644
--- /dev/null
+++ b/Engine/Source/Editor/GraphEditor/Private/KismetNodes/KismetNodeInfoContext.cpp
@@ -0,0 +1,100 @@
+// Copyright 1998-2014 Epic Games, Inc. All Rights Reserved.
+
+#include "GraphEditorCommon.h"
+#include "KismetNodes/KismetNodeInfoContext.h"
+#include "Editor/UnrealEd/Public/Kismet2/BlueprintEditorUtils.h"
+
+//////////////////////////////////////////////////////////////////////////
+// FKismetNodeInfoContext
+
+// Context used to aid debugging displays for nodes
+FKismetNodeInfoContext::FKismetNodeInfoContext(UEdGraph* SourceGraph)
+	: ActiveObjectBeingDebugged(NULL)
+{
+	// Only show pending latent actions in PIE/SIE mode
+	SourceBlueprint = FBlueprintEditorUtils::FindBlueprintForGraph(SourceGraph);
+
+	if (SourceBlueprint != NULL)
+	{
+		ActiveObjectBeingDebugged = SourceBlueprint->GetObjectBeingDebugged();
+
+		// Run thru debugged objects to see if any are objects with pending latent actions
+		if (ActiveObjectBeingDebugged != NULL)
+		{
+			UBlueprintGeneratedClass* Class = CastChecked<UBlueprintGeneratedClass>((UObject*)(ActiveObjectBeingDebugged->GetClass()));
+			FBlueprintDebugData const& ClassDebugData = Class->GetDebugData();
+
+			TSet<UObject*> LatentContextObjects;
+
+			TArray<UK2Node_CallFunction*> FunctionNodes;
+			SourceGraph->GetNodesOfClass<UK2Node_CallFunction>(FunctionNodes);
+			// collect all the world context objects for all of the graph's latent nodes
+			for (UK2Node_CallFunction const* FunctionNode : FunctionNodes)
+			{
+				UFunction* Function = FunctionNode->GetTargetFunction();
+				if ((Function == NULL) || !Function->HasMetaData(FBlueprintMetadata::MD_Latent))
+				{
+					continue;
+				}
+
+				UObject* NodeWorldContext = ActiveObjectBeingDebugged;
+				// if the node has a specific "world context" pin, attempt to get the value set for that first
+				if (Function->HasMetaData(FBlueprintMetadata::MD_WorldContext))
+				{
+					FString const WorldContextPinName = Function->GetMetaData(FBlueprintMetadata::MD_WorldContext);
+					if (UEdGraphPin* ContextPin = FunctionNode->FindPin(WorldContextPinName))
+					{
+						if (UObjectPropertyBase* ContextProperty = Cast<UObjectPropertyBase>(ClassDebugData.FindClassPropertyForPin(ContextPin)))
+						{
+							UObject* PropertyValue = ContextProperty->GetObjectPropertyValue_InContainer(ActiveObjectBeingDebugged);
+							if (PropertyValue != NULL)
+							{
+								NodeWorldContext = PropertyValue;
+							}
+						}
+					}
+				}
+				
+				LatentContextObjects.Add(NodeWorldContext);
+			}
+
+			for (UObject* ContextObject : LatentContextObjects)
+			{
+				if (UWorld* World = GEngine->GetWorldFromContextObject(ContextObject, /*bChecked =*/false))
+				{
+					FLatentActionManager& Manager = World->GetLatentActionManager();
+
+					TSet<int32> UUIDSet;
+					Manager.GetActiveUUIDs(ActiveObjectBeingDebugged, /*out*/ UUIDSet);
+
+					for (TSet<int32>::TConstIterator IterUUID(UUIDSet); IterUUID; ++IterUUID)
+					{
+						const int32 UUID = *IterUUID;
+
+						if (UEdGraphNode* ParentNode = ClassDebugData.FindNodeFromUUID(UUID))
+						{
+ 							TArray<FObjectUUIDPair>& Pairs = NodesWithActiveLatentActions.FindOrAdd(ParentNode);
+ 							new (Pairs) FObjectUUIDPair(ContextObject, UUID);
+						}
+					}
+				}
+			}
+		}
+
+		// Covert the watched pin array into a set
+		for (auto WatchedPinIt = SourceBlueprint->PinWatches.CreateConstIterator(); WatchedPinIt; ++WatchedPinIt)
+		{
+			UEdGraphPin* WatchedPin = *WatchedPinIt;
+
+			UEdGraphNode* OwningNode = Cast<UEdGraphNode>(WatchedPin->GetOuter());
+			if (!ensure(OwningNode != NULL)) // shouldn't happen, but just in case a dead pin was added to the PinWatches array
+			{
+				continue;
+			}
+			check(OwningNode == WatchedPin->GetOwningNode());
+
+			WatchedPinSet.Add(WatchedPin);
+			WatchedNodeSet.Add(OwningNode);
+		}
+	}
+}
diff --git a/Engine/Source/Editor/GraphEditor/Private/MaterialGraphConnectionDrawingPolicy.cpp b/Engine/Source/Editor/GraphEditor/Private/MaterialGraphConnectionDrawingPolicy.cpp
new file mode 100644
--- /dev/null
+++ b/Engine/Source/Editor/GraphEditor/Private/MaterialGraphConnectionDrawingPolicy.cpp
@@ -0,0 +1,47 @@
+// Copyright 1998-2014 Epic Games, Inc. All Rights Reserved.
+
+#include "GraphEditorCommon.h"
+#include "MaterialGraphConnectionDrawingPolicy.h"
+
+/////////////////////////////////////////////////////
+// FMaterialGraphConnectionDrawingPolicy
+
+FMaterialGraphConnectionDrawingPolicy::FMaterialGraphConnectionDrawingPolicy(int32 InBackLayerID, int32 InFrontLayerID, float ZoomFactor, const FSlateRect& InClippingRect, FSlateWindowElementList& InDrawElements, UEdGraph* InGraphObj)
+	: FConnectionDrawingPolicy(InBackLayerID, InFrontLayerID, ZoomFactor, InClippingRect, InDrawElements)
+	, MaterialGraph(CastChecked<UMaterialGraph>(InGraphObj))
+	, MaterialGraphSchema(CastChecked<UMaterialGraphSchema>(InGraphObj->GetSchema()))
+{
+	// Don't want to draw ending arrowheads
+	ArrowImage = nullptr;
+	ArrowRadius = FVector2D::ZeroVector;
+
+	// Still need to be able to perceive the graph while dragging connectors, esp over comment boxes
+	HoverDeemphasisDarkFraction = 0.4f;
+}
+
+void FMaterialGraphConnectionDrawingPolicy::DetermineWiringStyle(UEdGraphPin* OutputPin, UEdGraphPin* InputPin, /*inout*/ float& Thickness, /*inout*/ FLinearColor& WireColor, /*inout*/bool& bDrawBubbles, /*inout*/ bool& bBidirectional)
+{
+	WireColor = MaterialGraphSchema->ActivePinColor;
+
+	// Have to consider both pins as the input will be an 'output' when previewing a connection
+	if (OutputPin)
+	{
+		if (!MaterialGraph->IsInputActive(OutputPin))
+		{
+			WireColor = MaterialGraphSchema->InactivePinColor;
+		}
+	}
+	if (InputPin)
+	{
+		if (!MaterialGraph->IsInputActive(InputPin))
+		{
+			WireColor = MaterialGraphSchema->InactivePinColor;
+		}
+	}
+
+	const bool bDeemphasizeUnhoveredPins = HoveredPins.Num() > 0;
+	if (bDeemphasizeUnhoveredPins)
+	{
+		ApplyHoverDeemphasis(OutputPin, InputPin, /*inout*/ Thickness, /*inout*/ WireColor);
+	}
+}
diff --git a/Engine/Source/Editor/GraphEditor/Private/MaterialGraphConnectionDrawingPolicy.h b/Engine/Source/Editor/GraphEditor/Private/MaterialGraphConnectionDrawingPolicy.h
new file mode 100644
--- /dev/null
+++ b/Engine/Source/Editor/GraphEditor/Private/MaterialGraphConnectionDrawingPolicy.h
@@ -0,0 +1,23 @@
+// Copyright 1998-2014 Epic Games, Inc. All Rights Reserved.
+
+#pragma once
+
+#include "ConnectionDrawingPolicy.h"
+
+/////////////////////////////////////////////////////
+// FMaterialGraphConnectionDrawingPolicy
+
+// This class draws the connections for an UEdGraph using a MaterialGraph schema
+class FMaterialGraphConnectionDrawingPolicy : public FConnectionDrawingPolicy
+{
+protected:
+	UMaterialGraph* MaterialGraph;
+	const UMaterialGraphSchema* MaterialGraphSchema;
+
+public:
+	FMaterialGraphConnectionDrawingPolicy(int32 InBackLayerID, int32 InFrontLayerID, float ZoomFactor, const FSlateRect& InClippingRect, FSlateWindowElementList& InDrawElements, UEdGraph* InGraphObj);
+
+	// FConnectionDrawingPolicy interface
+	virtual void DetermineWiringStyle(UEdGraphPin* OutputPin, UEdGraphPin* InputPin, /*inout*/ float& Thickness, /*inout*/ FLinearColor& WireColor, /*inout*/bool& bDrawBubbles, /*inout*/ bool& bBidirectional) override;
+	// End of FConnectionDrawingPolicy interface
+};
diff --git a/Engine/Source/Editor/GraphEditor/Private/SGraphPanel.cpp b/Engine/Source/Editor/GraphEditor/Private/SGraphPanel.cpp
--- a/Engine/Source/Editor/GraphEditor/Private/SGraphPanel.cpp
+++ b/Engine/Source/Editor/GraphEditor/Private/SGraphPanel.cpp
@@ -7,7 +7,13 @@
 #include "Editor/UnrealEd/Public/DragAndDrop/ActorDragDropGraphEdOp.h"
 #include "Editor/UnrealEd/Public/DragAndDrop/AssetDragDropOp.h"
 #include "Editor/UnrealEd/Public/DragAndDrop/LevelDragDropOp.h"
+
 #include "ConnectionDrawingPolicy.h"
+#include "BlueprintConnectionDrawingPolicy.h"
+#include "AnimGraphConnectionDrawingPolicy.h"
+#include "SoundCueGraphConnectionDrawingPolicy.h"
+#include "MaterialGraphConnectionDrawingPolicy.h"
+#include "StateMachineConnectionDrawingPolicy.h"
 
 #include "AssetSelection.h"
 #include "ComponentAssetBroker.h"
diff --git a/Engine/Source/Editor/GraphEditor/Private/SoundCueGraphConnectionDrawingPolicy.cpp b/Engine/Source/Editor/GraphEditor/Private/SoundCueGraphConnectionDrawingPolicy.cpp
new file mode 100644
--- /dev/null
+++ b/Engine/Source/Editor/GraphEditor/Private/SoundCueGraphConnectionDrawingPolicy.cpp
@@ -0,0 +1,151 @@
+// Copyright 1998-2014 Epic Games, Inc. All Rights Reserved.
+
+#include "GraphEditorCommon.h"
+#include "SoundCueGraphConnectionDrawingPolicy.h"
+#include "SoundDefinitions.h"
+#include "Sound/SoundNode.h"
+#include "Sound/SoundCue.h"
+
+/////////////////////////////////////////////////////
+// FSoundCueGraphConnectionDrawingPolicy
+
+FSoundCueGraphConnectionDrawingPolicy::FSoundCueGraphConnectionDrawingPolicy(int32 InBackLayerID, int32 InFrontLayerID, float ZoomFactor, const FSlateRect& InClippingRect, FSlateWindowElementList& InDrawElements, UEdGraph* InGraphObj)
+	: FConnectionDrawingPolicy(InBackLayerID, InFrontLayerID, ZoomFactor, InClippingRect, InDrawElements)
+	, GraphObj(InGraphObj)
+{
+	// Cache off the editor options
+	ActiveColor = Settings->TraceAttackColor;
+	InactiveColor = Settings->TraceReleaseColor;
+
+	ActiveWireThickness = Settings->TraceAttackWireThickness;
+	InactiveWireThickness = Settings->TraceReleaseWireThickness;
+
+	// Don't want to draw ending arrowheads
+	ArrowImage = nullptr;
+	ArrowRadius = FVector2D::ZeroVector;
+}
+
+void FSoundCueGraphConnectionDrawingPolicy::Draw(TMap<TSharedRef<SWidget>, FArrangedWidget>& PinGeometries, FArrangedChildren& ArrangedNodes)
+{
+	// Build the execution roadmap (also populates execution times)
+	BuildAudioFlowRoadmap();
+
+	// Draw everything
+	FConnectionDrawingPolicy::Draw(PinGeometries, ArrangedNodes);
+}
+
+void FSoundCueGraphConnectionDrawingPolicy::BuildAudioFlowRoadmap()
+{
+	FAudioDevice* AudioDevice = GEngine->GetAudioDevice();
+
+	if (AudioDevice)
+	{
+		USoundCueGraph* SoundCueGraph = CastChecked<USoundCueGraph>(GraphObj);
+		USoundCue* SoundCue = SoundCueGraph->GetSoundCue();
+
+		UAudioComponent* PreviewAudioComponent = GEditor->GetPreviewAudioComponent();
+
+		if (PreviewAudioComponent && PreviewAudioComponent->IsPlaying() && PreviewAudioComponent->Sound == SoundCue)
+		{
+			TArray<FWaveInstance*> WaveInstances;
+			const int32 FirstActiveIndex = AudioDevice->GetSortedActiveWaveInstances(WaveInstances, ESortedActiveWaveGetType::QueryOnly);
+
+			// Run through the active instances and cull out anything that isn't related to this graph
+			if (FirstActiveIndex > 0)
+			{
+				WaveInstances.RemoveAt(0, FirstActiveIndex + 1);
+			}
+
+			for (int32 WaveIndex = WaveInstances.Num() - 1; WaveIndex >= 0 ; --WaveIndex)
+			{
+				UAudioComponent* WaveInstanceAudioComponent = WaveInstances[WaveIndex]->ActiveSound->AudioComponent.Get();
+				if (WaveInstanceAudioComponent != PreviewAudioComponent)
+				{
+					WaveInstances.RemoveAtSwap(WaveIndex);
+				}
+			}
+
+			for (int32 WaveIndex = 0; WaveIndex < WaveInstances.Num(); ++WaveIndex)
+			{
+				TArray<USoundNode*> PathToWaveInstance;
+				if (SoundCue->FindPathToNode(WaveInstances[WaveIndex]->WaveInstanceHash, PathToWaveInstance))
+				{
+					TArray<USoundCueGraphNode_Root*> RootNode;
+					TArray<UEdGraphNode*> GraphNodes;
+					SoundCueGraph->GetNodesOfClass<USoundCueGraphNode_Root>(RootNode);
+					check(RootNode.Num() == 1);
+					GraphNodes.Add(RootNode[0]);
+
+					TArray<double> NodeTimes;
+					NodeTimes.Add(FApp::GetCurrentTime()); // Time for the root node
+
+					for (int32 i = 0; i < PathToWaveInstance.Num(); ++i)
+					{
+						const double ObservationTime = FApp::GetCurrentTime() + 1.f;
+
+						NodeTimes.Add(ObservationTime);
+						GraphNodes.Add(PathToWaveInstance[i]->GraphNode);
+					}
+
+					// Record the unique node->node pairings, keeping only the most recent times for each pairing
+					for (int32 i = GraphNodes.Num() - 1; i >= 1; --i)
+					{
+						UEdGraphNode* CurNode = GraphNodes[i];
+						double CurNodeTime = NodeTimes[i];
+						UEdGraphNode* NextNode = GraphNodes[i-1];
+						double NextNodeTime = NodeTimes[i-1];
+
+						FExecPairingMap& Predecessors = PredecessorNodes.FindOrAdd(NextNode);
+
+						// Update the timings if this is a more recent pairing
+						FTimePair& Timings = Predecessors.FindOrAdd(CurNode);
+						if (Timings.ThisExecTime < NextNodeTime)
+						{
+							Timings.PredExecTime = CurNodeTime;
+							Timings.ThisExecTime = NextNodeTime;
+						}
+					}
+				}
+			}
+		}
+	}
+}
+
+// Give specific editor modes a chance to highlight this connection or darken non-interesting connections
+void FSoundCueGraphConnectionDrawingPolicy::DetermineWiringStyle(UEdGraphPin* OutputPin, UEdGraphPin* InputPin, /*inout*/ float& Thickness, /*inout*/ FLinearColor& WireColor, /*inout*/bool& bDrawBubbles, /*inout*/bool& bBidirectional)
+{
+	// Get the schema and grab the default color from it
+	check(OutputPin);
+	check(GraphObj);
+	const UEdGraphSchema* Schema = GraphObj->GetSchema();
+
+	WireColor = Schema->GetPinTypeColor(OutputPin->PinType);
+
+	if (InputPin == NULL)
+	{
+		return;
+	}
+	
+	bool bExecuted = false;
+
+	// Run thru the predecessors, and on
+	if (FExecPairingMap* PredecessorMap = PredecessorNodes.Find(InputPin->GetOwningNode()))
+	{
+		if (FTimePair* Times = PredecessorMap->Find(OutputPin->GetOwningNode()))
+		{
+			bExecuted = true;
+
+			Thickness = ActiveWireThickness;
+			WireColor = ActiveColor;
+
+			bDrawBubbles = true;
+		}
+	}
+
+	if (!bExecuted)
+	{
+		// It's not followed, fade it and keep it thin
+		WireColor = InactiveColor;
+		Thickness = InactiveWireThickness;
+	}
+}
diff --git a/Engine/Source/Editor/GraphEditor/Private/SoundCueGraphConnectionDrawingPolicy.h b/Engine/Source/Editor/GraphEditor/Private/SoundCueGraphConnectionDrawingPolicy.h
new file mode 100644
--- /dev/null
+++ b/Engine/Source/Editor/GraphEditor/Private/SoundCueGraphConnectionDrawingPolicy.h
@@ -0,0 +1,50 @@
+// Copyright 1998-2014 Epic Games, Inc. All Rights Reserved.
+
+#pragma once
+
+#include "ConnectionDrawingPolicy.h"
+
+/////////////////////////////////////////////////////
+// FSoundCueGraphConnectionDrawingPolicy
+
+// This class draws the connections for an UEdGraph using a SoundCue schema
+class FSoundCueGraphConnectionDrawingPolicy : public FConnectionDrawingPolicy
+{
+protected:
+	// Times for one execution pair within the current graph
+	struct FTimePair
+	{
+		double PredExecTime;
+		double ThisExecTime;
+
+		FTimePair()
+			: PredExecTime(0.0)
+			, ThisExecTime(0.0)
+		{
+		}
+	};
+
+	// Map of pairings
+	typedef TMap<UEdGraphNode*, FTimePair> FExecPairingMap;
+
+	// Map of nodes that preceeded before a given node in the execution sequence (one entry for each pairing)
+	TMap<UEdGraphNode*, FExecPairingMap> PredecessorNodes;
+
+	UEdGraph* GraphObj;
+
+	FLinearColor ActiveColor;
+	FLinearColor InactiveColor;
+
+	float ActiveWireThickness;
+	float InactiveWireThickness;
+
+public:
+	FSoundCueGraphConnectionDrawingPolicy(int32 InBackLayerID, int32 InFrontLayerID, float ZoomFactor, const FSlateRect& InClippingRect, FSlateWindowElementList& InDrawElements, UEdGraph* InGraphObj);
+
+	void BuildAudioFlowRoadmap();
+
+	// FConnectionDrawingPolicy interface
+	virtual void DetermineWiringStyle(UEdGraphPin* OutputPin, UEdGraphPin* InputPin, /*inout*/ float& Thickness, /*inout*/ FLinearColor& WireColor, /*inout*/bool& bDrawBubbles, /*inout*/ bool& bBidirectional) override;
+	virtual void Draw(TMap<TSharedRef<SWidget>, FArrangedWidget>& PinGeometries, FArrangedChildren& ArrangedNodes) override;
+	// End of FConnectionDrawingPolicy interface
+};
diff --git a/Engine/Source/Editor/GraphEditor/Private/StateMachineConnectionDrawingPolicy.cpp b/Engine/Source/Editor/GraphEditor/Private/StateMachineConnectionDrawingPolicy.cpp
new file mode 100644
--- /dev/null
+++ b/Engine/Source/Editor/GraphEditor/Private/StateMachineConnectionDrawingPolicy.cpp
@@ -0,0 +1,200 @@
+// Copyright 1998-2014 Epic Games, Inc. All Rights Reserved.
+
+#include "GraphEditorCommon.h"
+#include "StateMachineConnectionDrawingPolicy.h"
+#include "AnimationStateNodes/SGraphNodeAnimTransition.h"
+#include "AnimStateTransitionNode.h"
+#include "AnimStateEntryNode.h"
+
+/////////////////////////////////////////////////////
+// FStateMachineConnectionDrawingPolicy
+
+FStateMachineConnectionDrawingPolicy::FStateMachineConnectionDrawingPolicy(int32 InBackLayerID, int32 InFrontLayerID, float ZoomFactor, const FSlateRect& InClippingRect, FSlateWindowElementList& InDrawElements, UEdGraph* InGraphObj)
+	: FConnectionDrawingPolicy(InBackLayerID, InFrontLayerID, ZoomFactor, InClippingRect, InDrawElements)
+	, GraphObj(InGraphObj)
+{
+}
+
+void FStateMachineConnectionDrawingPolicy::DetermineWiringStyle(UEdGraphPin* OutputPin, UEdGraphPin* InputPin, /*inout*/ float& Thickness, /*inout*/ FLinearColor& WireColor, /*inout*/ bool& bDrawBubbles, /*inout*/ bool& bBidirectional)
+{
+	Thickness = 1.5f;
+
+	if (InputPin)
+	{
+		if (UAnimStateTransitionNode* TransNode = Cast<UAnimStateTransitionNode>(InputPin->GetOwningNode()))
+		{
+			WireColor = SGraphNodeAnimTransition::StaticGetTransitionColor(TransNode, HoveredPins.Contains(InputPin));
+
+			bBidirectional = TransNode->Bidirectional;
+		}
+	}
+
+	const bool bDeemphasizeUnhoveredPins = HoveredPins.Num() > 0;
+	if (bDeemphasizeUnhoveredPins)
+	{
+		ApplyHoverDeemphasis(OutputPin, InputPin, /*inout*/ Thickness, /*inout*/ WireColor);
+	}
+}
+
+void FStateMachineConnectionDrawingPolicy::DetermineLinkGeometry(
+	TMap<TSharedRef<SWidget>, FArrangedWidget>& PinGeometries,
+	FArrangedChildren& ArrangedNodes, 
+	TSharedRef<SWidget>& OutputPinWidget,
+	UEdGraphPin* OutputPin,
+	UEdGraphPin* InputPin,
+	/*out*/ FArrangedWidget*& StartWidgetGeometry,
+	/*out*/ FArrangedWidget*& EndWidgetGeometry
+	)
+{
+	if (UAnimStateEntryNode* EntryNode = Cast<UAnimStateEntryNode>(OutputPin->GetOwningNode()))
+	{
+		//FConnectionDrawingPolicy::DetermineLinkGeometry(PinGeometries, ArrangedNodes, OutputPinWidget, OutputPin, InputPin, StartWidgetGeometry, EndWidgetGeometry);
+
+		StartWidgetGeometry = PinGeometries.Find(OutputPinWidget);
+
+		UAnimStateNodeBase* State = CastChecked<UAnimStateNodeBase>(InputPin->GetOwningNode());
+		int32 StateIndex = NodeWidgetMap.FindChecked(State);
+		EndWidgetGeometry = &(ArrangedNodes[StateIndex]);
+	}
+	else if (UAnimStateTransitionNode* TransNode = Cast<UAnimStateTransitionNode>(InputPin->GetOwningNode()))
+	{
+		UAnimStateNodeBase* PrevState = TransNode->GetPreviousState();
+		UAnimStateNodeBase* NextState = TransNode->GetNextState();
+		if ((PrevState != NULL) && (NextState != NULL))
+		{
+			int32* PrevNodeIndex = NodeWidgetMap.Find(PrevState);
+			int32* NextNodeIndex = NodeWidgetMap.Find(NextState);
+			if ((PrevNodeIndex != NULL) && (NextNodeIndex != NULL))
+			{
+				StartWidgetGeometry = &(ArrangedNodes[*PrevNodeIndex]);
+				EndWidgetGeometry = &(ArrangedNodes[*NextNodeIndex]);
+			}
+		}
+	}
+	else
+	{
+		StartWidgetGeometry = PinGeometries.Find(OutputPinWidget);
+
+		if (TSharedRef<SGraphPin>* pTargetWidget = PinToPinWidgetMap.Find(InputPin))
+		{
+			TSharedRef<SGraphPin> InputWidget = *pTargetWidget;
+			EndWidgetGeometry = PinGeometries.Find(InputWidget);
+		}
+	}
+}
+
+void FStateMachineConnectionDrawingPolicy::Draw(TMap<TSharedRef<SWidget>, FArrangedWidget>& PinGeometries, FArrangedChildren& ArrangedNodes)
+{
+	// Build an acceleration structure to quickly find geometry for the nodes
+	NodeWidgetMap.Empty();
+	for (int32 NodeIndex = 0; NodeIndex < ArrangedNodes.Num(); ++NodeIndex)
+	{
+		FArrangedWidget& CurWidget = ArrangedNodes[NodeIndex];
+		TSharedRef<SGraphNode> ChildNode = StaticCastSharedRef<SGraphNode>(CurWidget.Widget);
+		NodeWidgetMap.Add(ChildNode->GetNodeObj(), NodeIndex);
+	}
+
+	// Now draw
+	FConnectionDrawingPolicy::Draw(PinGeometries, ArrangedNodes);
+}
+
+void FStateMachineConnectionDrawingPolicy::DrawPreviewConnector(const FGeometry& PinGeometry, const FVector2D& StartPoint, const FVector2D& EndPoint, UEdGraphPin* Pin)
+{
+	float Thickness = 1.0f;
+	FLinearColor WireColor = FLinearColor::White;
+	bool bDrawBubbles = false;
+	bool bBiDirectional = false;
+	DetermineWiringStyle(Pin, NULL, /*inout*/ Thickness, /*inout*/ WireColor, /*inout*/ bDrawBubbles, /*inout*/ bBiDirectional);
+
+
+	const FVector2D SeedPoint = EndPoint;
+	const FVector2D AdjustedStartPoint = FGeometryHelper::FindClosestPointOnGeom(PinGeometry, SeedPoint);
+
+	DrawSplineWithArrow(AdjustedStartPoint, EndPoint, WireColor, Thickness, bDrawBubbles, bBiDirectional);
+}
+
+
+void FStateMachineConnectionDrawingPolicy::DrawSplineWithArrow(const FVector2D& StartAnchorPoint, const FVector2D& EndAnchorPoint, const FLinearColor& WireColor, float WireThickness, bool bDrawBubbles, bool Bidirectional)
+{
+	Internal_DrawLineWithArrow(StartAnchorPoint, EndAnchorPoint, WireColor, WireThickness, bDrawBubbles);
+	if (Bidirectional)
+	{
+		Internal_DrawLineWithArrow(EndAnchorPoint, StartAnchorPoint, WireColor, WireThickness, bDrawBubbles);
+	}
+}
+
+void FStateMachineConnectionDrawingPolicy::Internal_DrawLineWithArrow(const FVector2D& StartAnchorPoint, const FVector2D& EndAnchorPoint, const FLinearColor& WireColor, float WireThickness, bool bDrawBubbles)
+{
+	//@TODO: Should this be scaled by zoom factor?
+	const float LineSeparationAmount = 4.5f;
+
+	const FVector2D DeltaPos = EndAnchorPoint - StartAnchorPoint;
+	const FVector2D UnitDelta = DeltaPos.SafeNormal();
+	const FVector2D Normal = FVector2D(DeltaPos.Y, -DeltaPos.X).SafeNormal();
+
+	// Come up with the final start/end points
+	const FVector2D DirectionBias = Normal * LineSeparationAmount;
+	const FVector2D LengthBias = ArrowRadius.X * UnitDelta;
+	const FVector2D StartPoint = StartAnchorPoint + DirectionBias + LengthBias;
+	const FVector2D EndPoint = EndAnchorPoint + DirectionBias - LengthBias;
+
+	// Draw a line/spline
+	DrawConnection(WireLayerID, StartPoint, EndPoint, WireColor, WireThickness, bDrawBubbles);
+
+	// Draw the arrow
+	const FVector2D ArrowDrawPos = EndPoint - ArrowRadius;
+	const float AngleInRadians = FMath::Atan2(DeltaPos.Y, DeltaPos.X);
+
+	FSlateDrawElement::MakeRotatedBox(
+		DrawElementsList,
+		ArrowLayerID,
+		FPaintGeometry(ArrowDrawPos, ArrowImage->ImageSize * ZoomFactor, ZoomFactor),
+		ArrowImage,
+		ClippingRect,
+		ESlateDrawEffect::None,
+		AngleInRadians,
+		TOptional<FVector2D>(),
+		FSlateDrawElement::RelativeToElement,
+		WireColor
+		);
+}
+
+void FStateMachineConnectionDrawingPolicy::DrawSplineWithArrow(FGeometry& StartGeom, FGeometry& EndGeom, const FLinearColor& WireColor, float WireThickness, bool bDrawBubbles, bool Bidirectional)
+{
+	// Get a reasonable seed point (halfway between the boxes)
+	const FVector2D StartCenter = FGeometryHelper::CenterOf(StartGeom);
+	const FVector2D EndCenter = FGeometryHelper::CenterOf(EndGeom);
+	const FVector2D SeedPoint = (StartCenter + EndCenter) * 0.5f;
+	
+	// Find the (approximate) closest points between the two boxes
+	const FVector2D StartAnchorPoint = FGeometryHelper::FindClosestPointOnGeom(StartGeom, SeedPoint);
+	const FVector2D EndAnchorPoint = FGeometryHelper::FindClosestPointOnGeom(EndGeom, SeedPoint);
+
+	DrawSplineWithArrow(StartAnchorPoint, EndAnchorPoint, WireColor, WireThickness, bDrawBubbles, Bidirectional);
+}
+
+void FStateMachineConnectionDrawingPolicy::DrawConnection(int32 LayerId, const FVector2D& Start, const FVector2D& End, const FLinearColor& InColor, float Thickness, bool bDrawBubbles)
+{
+	const FVector2D& P0 = Start;
+	const FVector2D& P1 = End;
+
+	const FVector2D Delta = End-Start;
+	const FVector2D NormDelta = Delta.SafeNormal();
+	
+	const FVector2D P0Tangent = NormDelta;
+	const FVector2D P1Tangent = NormDelta;
+
+	// Draw the spline itself
+	FSlateDrawElement::MakeDrawSpaceSpline(
+		DrawElementsList,
+		LayerId,
+		P0, P0Tangent,
+		P1, P1Tangent,
+		ClippingRect,
+		Thickness,
+		ESlateDrawEffect::None,
+		InColor
+		);
+
+	//@TODO: Handle bDrawBubbles
+}
diff --git a/Engine/Source/Editor/GraphEditor/Private/StateMachineConnectionDrawingPolicy.h b/Engine/Source/Editor/GraphEditor/Private/StateMachineConnectionDrawingPolicy.h
new file mode 100644
--- /dev/null
+++ b/Engine/Source/Editor/GraphEditor/Private/StateMachineConnectionDrawingPolicy.h
@@ -0,0 +1,42 @@
+// Copyright 1998-2014 Epic Games, Inc. All Rights Reserved.
+
+#pragma once
+
+#include "ConnectionDrawingPolicy.h"
+
+/////////////////////////////////////////////////////
+// FStateMachineConnectionDrawingPolicy
+
+// This class draws the connections for an UEdGraph with a state machine schema
+class FStateMachineConnectionDrawingPolicy : public FConnectionDrawingPolicy
+{
+protected:
+	UEdGraph* GraphObj;
+
+	TMap<UEdGraphNode*, int32> NodeWidgetMap;
+public:
+	//
+	FStateMachineConnectionDrawingPolicy(int32 InBackLayerID, int32 InFrontLayerID, float ZoomFactor, const FSlateRect& InClippingRect, FSlateWindowElementList& InDrawElements, UEdGraph* InGraphObj);
+
+	// FConnectionDrawingPolicy interface
+	virtual void DetermineWiringStyle(UEdGraphPin* OutputPin, UEdGraphPin* InputPin, /*inout*/ float& Thickness, /*inout*/ FLinearColor& WireColor, /*inout*/ bool& bDrawBubbles, /*inout*/ bool& bBidirectional) override;
+	virtual void Draw(TMap<TSharedRef<SWidget>, FArrangedWidget>& PinGeometries, FArrangedChildren& ArrangedNodes) override;
+	virtual void DetermineLinkGeometry(
+		TMap<TSharedRef<SWidget>,
+		FArrangedWidget>& PinGeometries,
+		FArrangedChildren& ArrangedNodes, 
+		TSharedRef<SWidget>& OutputPinWidget,
+		UEdGraphPin* OutputPin,
+		UEdGraphPin* InputPin,
+		/*out*/ FArrangedWidget*& StartWidgetGeometry,
+		/*out*/ FArrangedWidget*& EndWidgetGeometry
+		) override;
+	virtual void DrawSplineWithArrow(FGeometry& StartGeom, FGeometry& EndGeom, const FLinearColor& WireColor, float WireThickness, bool bDrawBubbles, bool Bidirectional) override;
+	virtual void DrawSplineWithArrow(const FVector2D& StartPoint, const FVector2D& EndPoint, const FLinearColor& WireColor, float WireThickness, bool bDrawBubbles, bool Bidirectional) override;
+	virtual void DrawPreviewConnector(const FGeometry& PinGeometry, const FVector2D& StartPoint, const FVector2D& EndPoint, UEdGraphPin* Pin) override;
+	virtual void DrawConnection(int32 LayerId, const FVector2D& Start, const FVector2D& End, const FLinearColor& InColor, float Thickness, bool bDrawBubbles) override;
+	// End of FConnectionDrawingPolicy interface
+
+protected:
+	void Internal_DrawLineWithArrow(const FVector2D& StartAnchorPoint, const FVector2D& EndAnchorPoint, const FLinearColor& WireColor, float WireThickness, bool bDrawBubbles);
+};
