GraphEditor: Move ConnectionDrawingPolicy.h to public folder, as it was already being used externally

diff --git a/Engine/Source/Editor/GraphEditor/Private/ConnectionDrawingPolicy.h b/Engine/Source/Editor/GraphEditor/Private/ConnectionDrawingPolicy.h
deleted file mode 100644
diff --git a/Engine/Source/Editor/GraphEditor/Public/ConnectionDrawingPolicy.h b/Engine/Source/Editor/GraphEditor/Public/ConnectionDrawingPolicy.h
new file mode 100644
--- /dev/null
+++ b/Engine/Source/Editor/GraphEditor/Public/ConnectionDrawingPolicy.h
@@ -0,0 +1,96 @@
+// Copyright 1998-2014 Epic Games, Inc. All Rights Reserved.
+
+#pragma once
+
+/////////////////////////////////////////////////////
+
+GRAPHEDITOR_API DECLARE_LOG_CATEGORY_EXTERN(LogConnectionDrawingPolicy, Log, All);
+
+/////////////////////////////////////////////////////
+// FGeometryHelper
+
+class GRAPHEDITOR_API FGeometryHelper
+{
+public:
+	static FVector2D VerticalMiddleLeftOf(const FGeometry& SomeGeometry);
+	static FVector2D VerticalMiddleRightOf(const FGeometry& SomeGeometry);
+	static FVector2D CenterOf(const FGeometry& SomeGeometry);
+	static void ConvertToPoints(const FGeometry& Geom, TArray<FVector2D>& Points);
+
+	/** Find the point on line segment from LineStart to LineEnd which is closest to Point */
+	static FVector2D FindClosestPointOnLine(const FVector2D& LineStart, const FVector2D& LineEnd, const FVector2D& TestPoint);
+
+	static FVector2D FindClosestPointOnGeom(const FGeometry& Geom, const FVector2D& TestPoint);
+};
+
+/////////////////////////////////////////////////////
+// FConnectionDrawingPolicy
+
+// This class draws the connections for an UEdGraph composed of pins and nodes
+class GRAPHEDITOR_API FConnectionDrawingPolicy
+{
+protected:
+	int32 WireLayerID;
+	int32 ArrowLayerID;
+
+	const FSlateBrush* ArrowImage;
+	const FSlateBrush* MidpointImage;
+	const FSlateBrush* BubbleImage;
+
+	const UGraphEditorSettings* Settings;
+
+public:
+	FVector2D ArrowRadius;
+	FVector2D MidpointRadius;
+protected:
+	float ZoomFactor; 
+	float HoverDeemphasisDarkFraction;
+	const FSlateRect& ClippingRect;
+	FSlateWindowElementList& DrawElementsList;
+	TMap< UEdGraphPin*, TSharedRef<SGraphPin> > PinToPinWidgetMap;
+	TSet< UEdGraphPin* > HoveredPins;
+	double LastHoverTimeEvent;
+public:
+	virtual ~FConnectionDrawingPolicy() {}
+
+	FConnectionDrawingPolicy(int32 InBackLayerID, int32 InFrontLayerID, float InZoomFactor, const FSlateRect& InClippingRect, FSlateWindowElementList& InDrawElements);
+
+	// Update the drawing policy with the set of hovered pins (which can be empty)
+	void SetHoveredPins(const TSet< TWeakObjectPtr<UEdGraphPin> >& InHoveredPins, const TArray< TSharedPtr<SGraphPin> >& OverridePins, double HoverTime);
+
+	// Update the drawing policy with the marked pin (which may not be valid)
+	void SetMarkedPin(TWeakPtr<SGraphPin> InMarkedPin);
+
+	static float MakeSplineReparamTable(const FVector2D& P0, const FVector2D& P0Tangent, const FVector2D& P1, const FVector2D& P1Tangent, FInterpCurve<float>& OutReparamTable);
+
+	virtual void DrawSplineWithArrow(const FVector2D& StartPoint, const FVector2D& EndPoint, const FLinearColor& WireColor, float WireThickness, bool bDrawBubbles, bool Bidirectional);
+	
+	virtual void DrawSplineWithArrow(FGeometry& StartGeom, FGeometry& EndGeom, const FLinearColor& WireColor, float WireThickness, bool bDrawBubbles, bool Bidirectional);
+
+	virtual FVector2D ComputeSplineTangent(const FVector2D& Start, const FVector2D& End) const;
+	virtual void DrawConnection(int32 LayerId, const FVector2D& Start, const FVector2D& End, const FLinearColor& InColor, float Thickness, bool bDrawBubbles);
+	virtual void DrawPreviewConnector(const FGeometry& PinGeometry, const FVector2D& StartPoint, const FVector2D& EndPoint, UEdGraphPin* Pin);
+
+	// Give specific editor modes a chance to highlight this connection or darken non-interesting connections
+	virtual void DetermineWiringStyle(UEdGraphPin* OutputPin, UEdGraphPin* InputPin, /*inout*/ float& Thickness, /*inout*/ FLinearColor& WireColor, /*inout*/bool& bDrawBubbles, /*inout*/ bool &bBidirectional);
+
+	virtual void Draw(TMap<TSharedRef<SWidget>, FArrangedWidget>& PinGeometries, FArrangedChildren& ArrangedNodes);
+
+	virtual void DetermineLinkGeometry(
+		TMap<TSharedRef<SWidget>,
+		FArrangedWidget>& PinGeometries,
+		FArrangedChildren& ArrangedNodes, 
+		TSharedRef<SWidget>& OutputPinWidget,
+		UEdGraphPin* OutputPin,
+		UEdGraphPin* InputPin,
+		/*out*/ FArrangedWidget*& StartWidgetGeometry,
+		/*out*/ FArrangedWidget*& EndWidgetGeometry
+		);
+
+	virtual void SetIncompatiblePinDrawState(const TSharedPtr<SGraphPin>& StartPin, const TSet< TSharedRef<SWidget> >& VisiblePins);
+	virtual void ResetIncompatiblePinDrawState(const TSet< TSharedRef<SWidget> >& VisiblePins);
+
+	virtual void ApplyHoverDeemphasis(UEdGraphPin* OutputPin, UEdGraphPin* InputPin, /*inout*/ float& Thickness, /*inout*/ FLinearColor& WireColor);
+
+	virtual bool IsConnectionCulled( const FArrangedWidget& StartLink, const FArrangedWidget& EndLink ) const;
+};
diff --git a/Engine/Source/Editor/ReferenceViewer/Private/ReferenceViewerSchema.cpp b/Engine/Source/Editor/ReferenceViewer/Private/ReferenceViewerSchema.cpp
--- a/Engine/Source/Editor/ReferenceViewer/Private/ReferenceViewerSchema.cpp
+++ b/Engine/Source/Editor/ReferenceViewer/Private/ReferenceViewerSchema.cpp
@@ -5,8 +5,7 @@
 #include "ReferenceViewerActions.h"
 #include "GlobalEditorCommonCommands.h"
 
-// @todo ConnectionDrawingPolicy should not be in a private folder since it is part of the module's API
-#include "Editor/GraphEditor/Private/ConnectionDrawingPolicy.h"
+#include "Editor/GraphEditor/Public/ConnectionDrawingPolicy.h"
 
 // Overridden connection drawing policy to use less curvy lines between nodes
 class FReferenceViewerConnectionDrawingPolicy : public FConnectionDrawingPolicy
