Fix for missing comment bubbles on variable nodes.
#Branch UE4
#Proj GraphEditor

diff --git a/Engine/Source/Editor/GraphEditor/Private/KismetNodes/SGraphNodeK2Var.cpp b/Engine/Source/Editor/GraphEditor/Private/KismetNodes/SGraphNodeK2Var.cpp
--- a/Engine/Source/Editor/GraphEditor/Private/KismetNodes/SGraphNodeK2Var.cpp
+++ b/Engine/Source/Editor/GraphEditor/Private/KismetNodes/SGraphNodeK2Var.cpp
@@ -272,6 +272,27 @@ void SGraphNodeK2Var::UpdateGraphNode()
 				SNew(SSpacer).Size(FVector2D(0.0f, VerticalPaddingAmount))
 			];
 	}
+	// Create comment bubble
+	TSharedPtr<SCommentBubble> CommentBubble;
+
+	SAssignNew( CommentBubble, SCommentBubble )
+	.GraphNode( GraphNode )
+	.Text( this, &SGraphNode::GetNodeComment )
+	.ColorAndOpacity( this, &SGraphNodeK2Var::GetCommentColor )
+	.AllowPinning( true )
+	.EnableTitleBarBubble( true )
+	.EnableBubbleCtrls( true )
+	.GraphLOD( this, &SGraphNode::GetCurrentLOD )
+	.IsGraphNodeHovered( this, &SGraphNode::IsHovered );
+
+	GetOrAddSlot( ENodeZone::TopCenter )
+	.SlotOffset( TAttribute<FVector2D>( CommentBubble.Get(), &SCommentBubble::GetOffset ))
+	.SlotSize( TAttribute<FVector2D>( CommentBubble.Get(), &SCommentBubble::GetSize ))
+	.AllowScaling( TAttribute<bool>( CommentBubble.Get(), &SCommentBubble::IsScalingAllowed ))
+	.VAlign( VAlign_Top )
+	[
+		CommentBubble.ToSharedRef()
+	];
 
 	// Create widgets for each of the real pins
 	CreatePinWidgets();
