This change modifies the comment move code to consider all nodes in the panel rather than just visible.
#UE-4945 - Comment boxes no longer move all of the nodes within the box if the nodes are offscreen.
#Branch UE4
#Proj GraphEditor

diff --git a/Engine/Source/Editor/GraphEditor/Private/SGraphNodeComment.cpp b/Engine/Source/Editor/GraphEditor/Private/SGraphNodeComment.cpp
--- a/Engine/Source/Editor/GraphEditor/Private/SGraphNodeComment.cpp
+++ b/Engine/Source/Editor/GraphEditor/Private/SGraphNodeComment.cpp
@@ -250,7 +250,7 @@ void SGraphNodeComment::HandleSelection(bool bSelected, bool bUpdateNodesUnderCo
 				const FSlateRect CommentRect( NodePosition.X, NodePosition.Y, NodePosition.X + NodeSize.X, NodePosition.Y + NodeSize.Y );
 
 				TSharedPtr< SGraphPanel > Panel = Comment->GetOwnerPanel();
-				FChildren* PanelChildren = Panel->GetChildren();
+				FChildren* PanelChildren = Panel->GetAllChildren();
 				int32 NumChildren = PanelChildren->Num();
 				CommentNode->ClearNodesUnderComment();
 
diff --git a/Engine/Source/Editor/GraphEditor/Private/SNodePanel.cpp b/Engine/Source/Editor/GraphEditor/Private/SNodePanel.cpp
--- a/Engine/Source/Editor/GraphEditor/Private/SNodePanel.cpp
+++ b/Engine/Source/Editor/GraphEditor/Private/SNodePanel.cpp
@@ -247,6 +247,11 @@ FChildren* SNodePanel::GetChildren()
 	return &VisibleChildren;
 }
 
+FChildren* SNodePanel::GetAllChildren()
+{
+	return &Children;
+}
+
 float SNodePanel::GetZoomAmount() const
 {
 	if (bAllowContinousZoomInterpolation)
diff --git a/Engine/Source/Editor/GraphEditor/Public/SNodePanel.h b/Engine/Source/Editor/GraphEditor/Public/SNodePanel.h
--- a/Engine/Source/Editor/GraphEditor/Public/SNodePanel.h
+++ b/Engine/Source/Editor/GraphEditor/Public/SNodePanel.h
@@ -595,6 +595,9 @@ public:
 	/** Returns if the panel has been panned or zoomed since the last update */
 	bool HasMoved() const;
 
+	/** Returns all the panel children rather than only visible */
+	FChildren* GetAllChildren();
+
 protected:
 	/** Initialize members */
 	void Construct();
