Fix crash on triggering event dispatcher

diff --git a/Engine/Source/Editor/AnimGraph/Classes/AnimationConduitGraphSchema.h b/Engine/Source/Editor/AnimGraph/Classes/AnimationConduitGraphSchema.h
--- a/Engine/Source/Editor/AnimGraph/Classes/AnimationConduitGraphSchema.h
+++ b/Engine/Source/Editor/AnimGraph/Classes/AnimationConduitGraphSchema.h
@@ -14,5 +14,6 @@ class UAnimationConduitGraphSchema : public UEdGraphSchema_K2
 	virtual void CreateDefaultNodesForGraph(UEdGraph& Graph) const override;
 	virtual bool CanDuplicateGraph(UEdGraph* InSourceGraph) const override { return false; }
 	virtual void GetGraphDisplayInformation(const UEdGraph& Graph, /*out*/ FGraphDisplayInfo& DisplayInfo) const override;
+	virtual bool DoesSupportEventDispatcher() const	override {	return false; }
 	// End of UEdGraphSchema interface
 };
diff --git a/Engine/Source/Editor/AnimGraph/Classes/AnimationGraphSchema.h b/Engine/Source/Editor/AnimGraph/Classes/AnimationGraphSchema.h
--- a/Engine/Source/Editor/AnimGraph/Classes/AnimationGraphSchema.h
+++ b/Engine/Source/Editor/AnimGraph/Classes/AnimationGraphSchema.h
@@ -51,6 +51,7 @@ class UAnimationGraphSchema : public UEdGraphSchema_K2
 	virtual void GetContextMenuActions(const UEdGraph* CurrentGraph, const UEdGraphNode* InGraphNode, const UEdGraphPin* InGraphPin, FMenuBuilder* MenuBuilder, bool bIsDebugging) const override;
 	virtual FString GetPinDisplayName(const UEdGraphPin* Pin) const override;
 	virtual bool CanDuplicateGraph(UEdGraph* InSourceGraph) const override {	return false; }
+	virtual bool DoesSupportEventDispatcher() const	override {	return false; }
 	// End UEdGraphSchema interface.
 
 	// Begin UEdGraphSchema_K2 interface
diff --git a/Engine/Source/Editor/AnimGraph/Classes/AnimationStateMachineSchema.h b/Engine/Source/Editor/AnimGraph/Classes/AnimationStateMachineSchema.h
--- a/Engine/Source/Editor/AnimGraph/Classes/AnimationStateMachineSchema.h
+++ b/Engine/Source/Editor/AnimGraph/Classes/AnimationStateMachineSchema.h
@@ -78,5 +78,6 @@ class UAnimationStateMachineSchema : public UEdGraphSchema
 	virtual void DroppedAssetsOnPin(const TArray<FAssetData>& Assets, const FVector2D& GraphPosition, UEdGraphPin* Pin) const override;
 	virtual void GetAssetsNodeHoverMessage(const TArray<FAssetData>& Assets, const UEdGraphNode* HoverNode, FString& OutTooltipText, bool& OutOkIcon) const override;
 	virtual void GetAssetsPinHoverMessage(const TArray<FAssetData>& Assets, const UEdGraphPin* HoverPin, FString& OutTooltipText, bool& OutOkIcon) const override;
+	virtual bool DoesSupportEventDispatcher() const	override {	return false; }
 	// End UEdGraphSchema interface
 };
diff --git a/Engine/Source/Editor/AnimGraph/Classes/AnimationTransitionSchema.h b/Engine/Source/Editor/AnimGraph/Classes/AnimationTransitionSchema.h
--- a/Engine/Source/Editor/AnimGraph/Classes/AnimationTransitionSchema.h
+++ b/Engine/Source/Editor/AnimGraph/Classes/AnimationTransitionSchema.h
@@ -14,6 +14,7 @@ class UAnimationTransitionSchema : public UEdGraphSchema_K2
 	virtual void CreateDefaultNodesForGraph(UEdGraph& Graph) const override;
 	virtual bool CanDuplicateGraph(UEdGraph* InSourceGraph) const override { return false; }
 	virtual void GetGraphDisplayInformation(const UEdGraph& Graph, /*out*/ FGraphDisplayInfo& DisplayInfo) const;
+	virtual bool DoesSupportEventDispatcher() const	override {	return false; }
 	// End UEdGraphSchema interface.
 
 private:
diff --git a/Engine/Source/Editor/GraphEditor/Private/SGraphPanel.cpp b/Engine/Source/Editor/GraphEditor/Private/SGraphPanel.cpp
--- a/Engine/Source/Editor/GraphEditor/Private/SGraphPanel.cpp
+++ b/Engine/Source/Editor/GraphEditor/Private/SGraphPanel.cpp
@@ -808,8 +808,13 @@ FReply SGraphPanel::OnDrop( const FGeometry& MyGeometry, const FDragDropEvent& D
 	if (Operation->IsOfType<FGraphEditorDragDropAction>())
 	{
 		check(GraphObj);
-		TSharedPtr<FGraphEditorDragDropAction> DragConn = StaticCastSharedPtr<FGraphEditorDragDropAction>(Operation);
-		return DragConn->DroppedOnPanel( SharedThis( this ), DragDropEvent.GetScreenSpacePosition(), NodeAddPosition, *GraphObj);
+		if (GraphObj->GetSchema()->DoesSupportEventDispatcher())
+		{
+			TSharedPtr<FGraphEditorDragDropAction> DragConn = StaticCastSharedPtr<FGraphEditorDragDropAction>(Operation);
+			return DragConn->DroppedOnPanel(SharedThis(this), DragDropEvent.GetScreenSpacePosition(), NodeAddPosition, *GraphObj);
+		}
+
+		return FReply::Unhandled();
 	}
 	else if (Operation->IsOfType<FActorDragDropGraphEdOp>())
 	{
diff --git a/Engine/Source/Runtime/Engine/Classes/EdGraph/EdGraphSchema.h b/Engine/Source/Runtime/Engine/Classes/EdGraph/EdGraphSchema.h
--- a/Engine/Source/Runtime/Engine/Classes/EdGraph/EdGraphSchema.h
+++ b/Engine/Source/Runtime/Engine/Classes/EdGraph/EdGraphSchema.h
@@ -516,6 +516,13 @@ class ENGINE_API UEdGraphSchema : public UObject
 	 */
 	virtual bool DoesSupportPinWatching() const	{ return false; }
 
+	/**
+	 *
+	 *	Determine if this graph supports event dispatcher
+	 * 
+	 *	@return True if this schema supports Event Dispatcher 
+	 */
+	virtual bool DoesSupportEventDispatcher() const	{	return true; }
 	/**
 	 * Checks to see if the specified pin is being watched by the graph's debug system.
 	 * 
