CIS fix mac

diff --git a/Engine/Source/Editor/GraphEditor/Private/KismetPins/SGraphPinDataTableRowName.cpp b/Engine/Source/Editor/GraphEditor/Private/KismetPins/SGraphPinDataTableRowName.cpp
--- a/Engine/Source/Editor/GraphEditor/Private/KismetPins/SGraphPinDataTableRowName.cpp
+++ b/Engine/Source/Editor/GraphEditor/Private/KismetPins/SGraphPinDataTableRowName.cpp
@@ -1,6 +1,5 @@
 // Copyright 1998-2014 Epic Games, Inc. All Rights Reserved.
 
-#pragma once
 #include "GraphEditorCommon.h"
 #include "SGraphPinDataTableRowName.h"
 
