To fix the comment in comment selection issues this changelist  adds a node sort to the node panel and adds some functionality to help comments manage their depth in relation to one another.
#UE-3724 - LIVE: BP: If a comment box is created inside of another larger comment box the smaller box is no longer editable
#Branch UE4
#Proj UnrealEd, GraphEditor

#info this issue came about after recent changes to slate resulted in OnHitTest no longer being called by design. This changelist seeks to eliminate the requirement for this by correctly sorting the comments in arrangedchildren.

diff --git a/Engine/Source/Editor/GraphEditor/Private/SGraphNodeComment.cpp b/Engine/Source/Editor/GraphEditor/Private/SGraphNodeComment.cpp
--- a/Engine/Source/Editor/GraphEditor/Private/SGraphNodeComment.cpp
+++ b/Engine/Source/Editor/GraphEditor/Private/SGraphNodeComment.cpp
@@ -230,6 +230,35 @@ FReply SGraphNodeComment::OnMouseButtonDoubleClick( const FGeometry& InMyGeometr
 	}
 }
 
+FReply SGraphNodeComment::OnMouseButtonUp( const FGeometry& MyGeometry, const FPointerEvent& MouseEvent )
+{
+	if ( (MouseEvent.GetEffectingButton() == EKeys::LeftMouseButton) && bUserIsDragging )
+	{
+		bUserIsDragging = false;
+
+		// Resize the node	
+		UserSize.X = FMath::RoundToFloat(UserSize.X);
+		UserSize.Y = FMath::RoundToFloat(UserSize.Y);
+
+		GetNodeObj()->ResizeNode(UserSize);
+
+		// End resize transaction
+		ResizeTransactionPtr.Reset();
+
+		// Update contained child Nodes
+		HandleSelection( bIsSelected, true );
+
+		return FReply::Handled().ReleaseMouseCapture();
+	}
+	return FReply::Unhandled();
+}
+
+int32 SGraphNodeComment::GetSortDepth() const
+{
+	UEdGraphNode_Comment* CommentNode = Cast<UEdGraphNode_Comment>( GraphNode );
+	return CommentNode ? CommentNode->CommentDepth : -1;
+}
+
 void SGraphNodeComment::HandleSelection(bool bSelected, bool bUpdateNodesUnderComment) const
 {
 	const FVector2D NodeSize = GetDesiredSize();
@@ -239,20 +268,19 @@ void SGraphNodeComment::HandleSelection(bool bSelected, bool bUpdateNodesUnderCo
 		if ((!this->bIsSelected && bSelected) || bUpdateNodesUnderComment)
 		{
 			SGraphNodeComment* Comment = const_cast<SGraphNodeComment*> (this);
-
 			UEdGraphNode_Comment* CommentNode = Cast<UEdGraphNode_Comment>(GraphNode);
 
 			if (CommentNode)
 			{
-
 				// Get our geo
 				const FVector2D NodePosition = GetPosition();
 				const FSlateRect CommentRect( NodePosition.X, NodePosition.Y, NodePosition.X + NodeSize.X, NodePosition.Y + NodeSize.Y );
 
-				TSharedPtr< SGraphPanel > Panel = Comment->GetOwnerPanel();
+				TSharedPtr<SGraphPanel> Panel = Comment->GetOwnerPanel();
 				FChildren* PanelChildren = Panel->GetAllChildren();
 				int32 NumChildren = PanelChildren->Num();
 				CommentNode->ClearNodesUnderComment();
+				int32 MinDepth = 0;
 
 				for ( int32 NodeIndex=0; NodeIndex < NumChildren; ++NodeIndex )
 				{
@@ -260,15 +288,25 @@ void SGraphNodeComment::HandleSelection(bool bSelected, bool bUpdateNodesUnderCo
 
 					UObject* GraphObject = SomeNodeWidget->GetObjectBeingDisplayed();
 
-					const FVector2D SomeNodePosition = SomeNodeWidget->GetPosition();
-					const FVector2D SomeNodeSize = SomeNodeWidget->GetDesiredSize();
-
-					const FSlateRect NodeGeometryGraphSpace( SomeNodePosition.X, SomeNodePosition.Y, SomeNodePosition.X + SomeNodeSize.X, SomeNodePosition.Y + SomeNodeSize.Y );
-					if ( FSlateRect::IsRectangleContained( CommentRect, NodeGeometryGraphSpace ) )
+					if( GraphObject != CommentNode )
 					{
-						CommentNode->AddNodeUnderComment(GraphObject);
+						const FVector2D SomeNodePosition = SomeNodeWidget->GetPosition();
+						const FVector2D SomeNodeSize = SomeNodeWidget->GetDesiredSize();
+
+						const FSlateRect NodeGeometryGraphSpace( SomeNodePosition.X, SomeNodePosition.Y, SomeNodePosition.X + SomeNodeSize.X, SomeNodePosition.Y + SomeNodeSize.Y );
+						if( FSlateRect::DoRectanglesIntersect( CommentRect, NodeGeometryGraphSpace ) )
+						{
+							MinDepth = FMath::Min( MinDepth, SomeNodeWidget->GetSortDepth() - 1 );
+
+							if ( FSlateRect::IsRectangleContained( CommentRect, NodeGeometryGraphSpace ) )
+							{
+								CommentNode->AddNodeUnderComment(GraphObject);
+							}
+						}
 					}
 				}
+				// Fix Depth to include any overlapped comments
+				CommentNode->CommentDepth = FMath::Min( CommentNode->CommentDepth, MinDepth );
 			}
 		}
 		bIsSelected = bSelected;
@@ -378,7 +416,7 @@ FSlateColor SGraphNodeComment::GetCommentTitleBarColor() const
 bool SGraphNodeComment::CanBeSelected(const FVector2D& MousePositionInNode) const
 {
 	const EResizableWindowZone InMouseZone = FindMouseZone(MousePositionInNode);
-	return CRWZ_TitleBar == MouseZone;
+	return CRWZ_TitleBar == InMouseZone;
 }
 
 FVector2D SGraphNodeComment::GetDesiredSizeForMarquee() const
diff --git a/Engine/Source/Editor/GraphEditor/Private/SGraphPanel.cpp b/Engine/Source/Editor/GraphEditor/Private/SGraphPanel.cpp
--- a/Engine/Source/Editor/GraphEditor/Private/SGraphPanel.cpp
+++ b/Engine/Source/Editor/GraphEditor/Private/SGraphPanel.cpp
@@ -880,14 +880,7 @@ void SGraphPanel::AddGraphNode( const TSharedRef<SNodePanel::SNode>& NodeToAdd )
 		NodeGuidMap.Add(Node->NodeGuid, GraphNode);
 	}
 
-	if (Node && Node->IsA( UEdGraphNode_Comment::StaticClass()))
-	{
-		SNodePanel::AddGraphNodeToBack(NodeToAdd);
-	}
-	else
-	{
-		SNodePanel::AddGraphNode(NodeToAdd);
-	}
+	SNodePanel::AddGraphNode(NodeToAdd);
 }
 
 void SGraphPanel::RemoveAllNodes()
@@ -1094,7 +1087,6 @@ void SGraphPanel::JumpToNode(const UEdGraphNode* JumpToMe, bool bRequestRename)
 	}
 }
 
-
 void SGraphPanel::JumpToPin(const UEdGraphPin* JumpToMe)
 {
 	if (JumpToMe != NULL)
@@ -1133,4 +1125,4 @@ void SGraphPanel::AddReferencedObjects( FReferenceCollector& Collector )
 {
 	Collector.AddReferencedObject( GraphObj );
 	Collector.AddReferencedObject( GraphObjToDiff );
-}
\ No newline at end of file
+}
diff --git a/Engine/Source/Editor/GraphEditor/Private/SNodePanel.cpp b/Engine/Source/Editor/GraphEditor/Private/SNodePanel.cpp
--- a/Engine/Source/Editor/GraphEditor/Private/SNodePanel.cpp
+++ b/Engine/Source/Editor/GraphEditor/Private/SNodePanel.cpp
@@ -1066,14 +1066,6 @@ void SNodePanel::AddGraphNode( const TSharedRef<SNodePanel::SNode>& NodeToAdd )
 	NodeToWidgetLookup.Add( NodeToAdd->GetObjectBeingDisplayed(), NodeToAdd );
 }
 
-/** Add a node in the 'back plane' of the panel */
-void SNodePanel::AddGraphNodeToBack( const TSharedRef<SNodePanel::SNode>& NodeToAdd )
-{
-	// add into the back plane
-	Children.Insert( NodeToAdd, 0 );
-	NodeToWidgetLookup.Add( NodeToAdd->GetObjectBeingDisplayed(), NodeToAdd );
-}
-
 /** Remove all nodes from the panel */
 void SNodePanel::RemoveAllNodes()
 {
@@ -1093,6 +1085,15 @@ void SNodePanel::PopulateVisibleChildren(const FGeometry& AllottedGeometry)
 			VisibleChildren.Add(SomeChild);
 		}
 	}
+	// Depth Sort Nodes
+	if( VisibleChildren.Num() > 0 )
+	{
+		struct SNodeLessThanSort
+		{
+			FORCEINLINE bool operator()(const TSharedRef<SNodePanel::SNode>& A, const TSharedRef<SNodePanel::SNode>& B) const { return A.Get() < B.Get(); }
+		};
+		VisibleChildren.Sort( SNodeLessThanSort() );
+	}
 }
 
 // Is the given node being observed by a widget in this panel?
diff --git a/Engine/Source/Editor/GraphEditor/Public/SGraphNodeComment.h b/Engine/Source/Editor/GraphEditor/Public/SGraphNodeComment.h
--- a/Engine/Source/Editor/GraphEditor/Public/SGraphNodeComment.h
+++ b/Engine/Source/Editor/GraphEditor/Public/SGraphNodeComment.h
@@ -10,6 +10,7 @@ public:
 
 	// Begin SWidget Interface
 	virtual FReply OnMouseButtonDoubleClick( const FGeometry& InMyGeometry, const FPointerEvent& InMouseEvent ) override;
+	virtual FReply OnMouseButtonUp( const FGeometry& MyGeometry, const FPointerEvent& MouseEvent ) override;
 	virtual void Tick(const FGeometry& AllottedGeometry, const double InCurrentTime, const float InDeltaTime) override;
 	virtual FReply OnDrop( const FGeometry& MyGeometry, const FDragDropEvent& DragDropEvent ) override;
 	virtual void OnDragEnter( const FGeometry& MyGeometry, const FDragDropEvent& DragDropEvent ) override;
@@ -19,6 +20,7 @@ public:
 	virtual const FSlateBrush* GetShadowBrush(bool bSelected) const override;
 	virtual void GetOverlayBrushes(bool bSelected, const FVector2D WidgetSize, TArray<FOverlayBrushInfo>& Brushes) const;
 	virtual bool ShouldAllowCulling() const override { return false; }
+	virtual int32 GetSortDepth() const override;
 	// End SNodePanel::SNode interface
 
 	// Begin SPanel Interface
diff --git a/Engine/Source/Editor/GraphEditor/Public/SNodePanel.h b/Engine/Source/Editor/GraphEditor/Public/SNodePanel.h
--- a/Engine/Source/Editor/GraphEditor/Public/SNodePanel.h
+++ b/Engine/Source/Editor/GraphEditor/Public/SNodePanel.h
@@ -482,6 +482,15 @@ public:
 			return GetDesiredSize();
 		}
 
+		// Returns node sort depth, defaults to and is generally 0 for most nodes
+		virtual int32 GetSortDepth() const { return 0; }
+
+		// Node Sort Operator
+		bool operator < ( const SNodePanel::SNode& NodeIn ) const
+		{
+			return GetSortDepth() < NodeIn.GetSortDepth();
+		}
+
 	protected:
 		SNode()
 		: BorderImage( FCoreStyle::Get().GetBrush( "NoBorder" ) )
@@ -655,9 +664,6 @@ protected:
 	/** Add a slot to the CanvasPanel dynamically */
 	virtual void AddGraphNode(const TSharedRef<SNode>& NodeToAdd);
 
-	/** Add a node widget to the back of the panel */
-	virtual void AddGraphNodeToBack(const TSharedRef<SNode>& NodeToAdd);
-
 	/** Remove all nodes from the panel */
 	virtual void RemoveAllNodes();
 	
diff --git a/Engine/Source/Editor/UnrealEd/Private/EdGraphNode_Comment.cpp b/Engine/Source/Editor/UnrealEd/Private/EdGraphNode_Comment.cpp
--- a/Engine/Source/Editor/UnrealEd/Private/EdGraphNode_Comment.cpp
+++ b/Engine/Source/Editor/UnrealEd/Private/EdGraphNode_Comment.cpp
@@ -28,6 +28,7 @@ UEdGraphNode_Comment::UEdGraphNode_Comment(const FObjectInitializer& ObjectIniti
 	bCommentBubbleVisible = true;
 	bCanResizeNode = true;
 	bCanRenameNode = true;
+	CommentDepth = -1;
 }
 
 void UEdGraphNode_Comment::AddReferencedObjects(UObject* InThis, FReferenceCollector& Collector) 
@@ -114,6 +115,10 @@ void UEdGraphNode_Comment::ResizeNode(const FVector2D& NewSize)
 
 void UEdGraphNode_Comment::AddNodeUnderComment(UObject* Object)
 {
+	if( UEdGraphNode_Comment* ChildComment = Cast<UEdGraphNode_Comment>(Object))
+	{
+		CommentDepth = FMath::Min( CommentDepth, ChildComment->CommentDepth - 1 );
+	}
 	NodesUnderComment.Add(Object);
 }
 
diff --git a/Engine/Source/Editor/UnrealEd/Public/EdGraphNode_Comment.h b/Engine/Source/Editor/UnrealEd/Public/EdGraphNode_Comment.h
--- a/Engine/Source/Editor/UnrealEd/Public/EdGraphNode_Comment.h
+++ b/Engine/Source/Editor/UnrealEd/Public/EdGraphNode_Comment.h
@@ -37,6 +37,10 @@ class UEdGraphNode_Comment : public UEdGraphNode
 	UPROPERTY(EditAnywhere, Category=Comment)
 	TEnumAsByte<ECommentBoxMode::Type> MoveMode;
 
+	/** comment Depth */
+	UPROPERTY(EditAnywhere, Category=Comment, meta=(DisplayName="CommentDepth"))
+	int32 CommentDepth;
+
 public:
 
 	// Begin UObject Interface
