Brushes and sub-classes of ABrush can no longer be spawned by Blueprints using the "Spawn Actor from Class" node.

diff --git a/Engine/Source/Editor/GraphEditor/Private/KismetNodes/SGraphNodeSpawnActorFromClass.cpp b/Engine/Source/Editor/GraphEditor/Private/KismetNodes/SGraphNodeSpawnActorFromClass.cpp
--- a/Engine/Source/Editor/GraphEditor/Private/KismetNodes/SGraphNodeSpawnActorFromClass.cpp
+++ b/Engine/Source/Editor/GraphEditor/Private/KismetNodes/SGraphNodeSpawnActorFromClass.cpp
@@ -45,9 +45,10 @@ class SGraphPinActorBasedClass : public SGraphPinObject
 			if(NULL != InClass)
 			{
 				const bool bActorBased = InClass->IsChildOf(AActor::StaticClass());
+				const bool bNotBrushBased = !InClass->IsChildOf(ABrush::StaticClass());
 				const bool bBlueprintType = UEdGraphSchema_K2::IsAllowableBlueprintVariableType(InClass);
 				const bool bNotAbstract = !InClass->HasAnyClassFlags(CLASS_Abstract);
-				return bActorBased && bBlueprintType && bNotAbstract;
+				return bActorBased && bNotBrushBased && bBlueprintType && bNotAbstract;
 			}
 			return false;
 		}
@@ -55,8 +56,9 @@ class SGraphPinActorBasedClass : public SGraphPinObject
 		virtual bool IsUnloadedClassAllowed(const FClassViewerInitializationOptions& InInitOptions, const TSharedRef< const IUnloadedBlueprintData > InUnloadedClassData, TSharedRef< FClassViewerFilterFuncs > InFilterFuncs) override
 		{
 			const bool bActorBased = InUnloadedClassData->IsChildOf(AActor::StaticClass());
+			const bool bNotBrushBased = !InUnloadedClassData->IsChildOf(ABrush::StaticClass());
 			const bool bNotAbstract = !InUnloadedClassData->HasAnyClassFlags(CLASS_Abstract);
-			return bActorBased && bNotAbstract;
+			return bActorBased && bNotBrushBased && bNotAbstract;
 		}
 	};
 
@@ -147,4 +149,4 @@ void SGraphNodeSpawnActorFromClass::CreatePinWidgets()
 	}
 }
 
-#undef LOCTEXT_NAMESPACE
\ No newline at end of file
+#undef LOCTEXT_NAMESPACE
