---- Merging with SlateDev branch ----

diff --git a/Engine/Source/Editor/DetailCustomizations/Private/AnimSequenceDetails.cpp b/Engine/Source/Editor/DetailCustomizations/Private/AnimSequenceDetails.cpp
--- a/Engine/Source/Editor/DetailCustomizations/Private/AnimSequenceDetails.cpp
+++ b/Engine/Source/Editor/DetailCustomizations/Private/AnimSequenceDetails.cpp
@@ -328,6 +328,7 @@ void SAnimationRefPoseViewport::CleanupComponent(USceneComponent* Component)
 
 SAnimationRefPoseViewport::SAnimationRefPoseViewport()
 	: PreviewScene(FPreviewScene::ConstructionValues())
+	, PreviewComponent(nullptr)
 {
 }
 
@@ -339,6 +340,11 @@ void SAnimationRefPoseViewport::Construct(const FArguments& InArgs)
 	RefPoseTypeHandle = InArgs._RefPoseTypeHandle;
 	RefFrameIndexPropertyHandle = InArgs._RefFrameIndexPropertyHandle;
 
+	// Create the preview component
+	PreviewComponent = ConstructObject<UDebugSkelMeshComponent>( UDebugSkelMeshComponent::StaticClass() );
+	PreviewComponent->MeshComponentUpdateFlag = EMeshComponentUpdateFlag::AlwaysTickPoseAndRefreshBones;
+	PreviewScene.AddComponent( PreviewComponent, FTransform::Identity );
+
 	this->ChildSlot
 	[
 		SNew(SVerticalBox)
@@ -368,30 +374,26 @@ void SAnimationRefPoseViewport::Construct(const FArguments& InArgs)
 			.PreviewInstance(this, &SAnimationRefPoseViewport::GetPreviewInstance)
 			.DraggableBars(this, &SAnimationRefPoseViewport::GetBars)
 			.OnBarDrag(this, &SAnimationRefPoseViewport::OnBarDrag)
+			.OnTickPlayback(this, &SAnimationRefPoseViewport::OnTickPreview)
 			.bAllowZoom(true)
 		]
 	];
 
-
-	// Create a viewport client
-	LevelViewportClient	= MakeShareable(new FAnimationSegmentViewportClient(PreviewScene));
+	// Create the viewport
+	LevelViewportClient = MakeShareable( new FAnimationSegmentViewportClient( PreviewScene ) );
 
 	LevelViewportClient->ViewportType = LVT_Perspective;
 	LevelViewportClient->bSetListenerPosition = false;
 	LevelViewportClient->SetViewLocation( EditorViewportDefs::DefaultPerspectiveViewLocation );
-	LevelViewportClient->SetViewRotation( EditorViewportDefs::DefaultPerspectiveViewRotation );	
+	LevelViewportClient->SetViewRotation( EditorViewportDefs::DefaultPerspectiveViewRotation );
 
-	SceneViewport = MakeShareable(new FSceneViewport(LevelViewportClient.Get(), ViewportWidget));
+	SceneViewport = MakeShareable( new FSceneViewport( LevelViewportClient.Get(), ViewportWidget ) );
 	LevelViewportClient->Viewport = SceneViewport.Get();
-	LevelViewportClient->SetRealtime(true);
-	LevelViewportClient->VisibilityDelegate.BindSP(this, &SAnimationRefPoseViewport::IsVisible);
-	LevelViewportClient->SetViewMode(VMI_Lit);
+	LevelViewportClient->SetRealtime( true );
+	LevelViewportClient->VisibilityDelegate.BindSP( this, &SAnimationRefPoseViewport::IsVisible );
+	LevelViewportClient->SetViewMode( VMI_Lit );
 
-	ViewportWidget->SetViewportInterface(SceneViewport.ToSharedRef());
-
-	PreviewComponent = ConstructObject<UDebugSkelMeshComponent>(UDebugSkelMeshComponent::StaticClass());
-	PreviewComponent->MeshComponentUpdateFlag = EMeshComponentUpdateFlag::AlwaysTickPoseAndRefreshBones;
-	PreviewScene.AddComponent(PreviewComponent, FTransform::Identity);
+	ViewportWidget->SetViewportInterface( SceneViewport.ToSharedRef() );
 
 	InitSkeleton();
 }
@@ -438,66 +440,64 @@ void SAnimationRefPoseViewport::InitSkeleton()
 	}
 }
 
-void SAnimationRefPoseViewport::Tick(const FGeometry& AllottedGeometry, const double InCurrentTime, const float InDeltaTime)
+void SAnimationRefPoseViewport::OnTickPreview( double InCurrentTime, float InDeltaTime )
 {
-	SCompoundWidget::Tick(AllottedGeometry, InCurrentTime, InDeltaTime);
+	LevelViewportClient->Invalidate();
+}
 
+void SAnimationRefPoseViewport::Tick( const FGeometry& AllottedGeometry, const double InCurrentTime, const float InDeltaTime )
+{
 	class UDebugSkelMeshComponent* Component = PreviewComponent;
 
-	FString TargetSkeletonName = TargetSkeleton? TargetSkeleton->GetName() : FName(NAME_None).ToString();
+	FString TargetSkeletonName = TargetSkeleton ? TargetSkeleton->GetName() : FName( NAME_None ).ToString();
 
-	if (Component != NULL)
+	if ( Component != NULL )
 	{
 		// Reinit the skeleton if the anim ref has changed
 		InitSkeleton();
 
-		if (Component->IsPreviewOn() && AnimRef != NULL)
+		if ( Component->IsPreviewOn() && AnimRef != NULL )
 		{
-			if(PreviewComponent != NULL && PreviewComponent->PreviewInstance != NULL)
+			if ( PreviewComponent != NULL && PreviewComponent->PreviewInstance != NULL )
 			{
 				uint8 RefPoseType;
-				RefPoseTypeHandle->GetValue(RefPoseType);
-				if (RefPoseType == ABPT_AnimFrame)
+				RefPoseTypeHandle->GetValue( RefPoseType );
+				if ( RefPoseType == ABPT_AnimFrame )
 				{
 					int RefFrameIndex;
-					RefFrameIndexPropertyHandle->GetValue(RefFrameIndex);
-					float Fraction = (AnimRef->NumFrames > 0)? FMath::Clamp<float>((float)RefFrameIndex/(float)AnimRef->NumFrames, 0.f, 1.f) : 0.f;
+					RefFrameIndexPropertyHandle->GetValue( RefFrameIndex );
+					float Fraction = ( AnimRef->NumFrames > 0 ) ? FMath::Clamp<float>( (float)RefFrameIndex / (float)AnimRef->NumFrames, 0.f, 1.f ) : 0.f;
 					float RefTime = AnimRef->SequenceLength * Fraction;
-					PreviewComponent->PreviewInstance->SetPosition(RefTime, false);
-					PreviewComponent->PreviewInstance->SetPlaying(false);
-				}
-				else
-				{
-					PreviewComponent->PreviewInstance->SetPlaying(true);
+					PreviewComponent->PreviewInstance->SetPosition( RefTime, false );
+					PreviewComponent->PreviewInstance->SetPlaying( false );
+					LevelViewportClient->Invalidate();
 				}
-				// reinvalidate the viewport
-				LevelViewportClient->Invalidate();
 			}
 
-			Description->SetText(FText::Format( LOCTEXT("Previewing", "Previewing {0}"), FText::FromString(Component->GetPreviewText()) ).ToString());
+			Description->SetText( FText::Format( LOCTEXT( "Previewing", "Previewing {0}" ), FText::FromString( Component->GetPreviewText() ) ).ToString() );
 		}
-		else if (Component->AnimBlueprintGeneratedClass)
+		else if ( Component->AnimBlueprintGeneratedClass )
 		{
-			Description->SetText(FText::Format( LOCTEXT("Previewing", "Previewing {0}"), FText::FromString(Component->AnimBlueprintGeneratedClass->GetName()) ).ToString());
+			Description->SetText( FText::Format( LOCTEXT( "Previewing", "Previewing {0}" ), FText::FromString( Component->AnimBlueprintGeneratedClass->GetName() ) ).ToString() );
 		}
-		else if (AnimRef && AnimRef->GetSkeleton() != TargetSkeleton)
+		else if ( AnimRef && AnimRef->GetSkeleton() != TargetSkeleton )
 		{
-			Description->SetText(FText::Format( LOCTEXT("IncorrectSkeleton", "The preview asset doesn't work for the skeleton '{0}'"), FText::FromString(TargetSkeletonName) ).ToString());
+			Description->SetText( FText::Format( LOCTEXT( "IncorrectSkeleton", "The preview asset doesn't work for the skeleton '{0}'" ), FText::FromString( TargetSkeletonName ) ).ToString() );
 		}
-		else if (Component->SkeletalMesh == NULL)
+		else if ( Component->SkeletalMesh == NULL )
 		{
-			Description->SetText(FText::Format( LOCTEXT("NoMeshFound", "No skeletal mesh found for skeleton '{0}'"), FText::FromString(TargetSkeletonName) ).ToString());
+			Description->SetText( FText::Format( LOCTEXT( "NoMeshFound", "No skeletal mesh found for skeleton '{0}'" ), FText::FromString( TargetSkeletonName ) ).ToString() );
 		}
 		else
 		{
-			Description->SetText(FText::Format( LOCTEXT("SelectAnimation", "Select animation that works for skeleton '{0}'"), FText::FromString(TargetSkeletonName) ).ToString());
+			Description->SetText( FText::Format( LOCTEXT( "SelectAnimation", "Select animation that works for skeleton '{0}'" ), FText::FromString( TargetSkeletonName ) ).ToString() );
 		}
 
-		Component->GetScene()->GetWorld()->Tick(LEVELTICK_All, InDeltaTime);
+		Component->GetScene()->GetWorld()->Tick( LEVELTICK_All, InDeltaTime );
 	}
 	else
 	{
-		Description->SetText(FText::Format( LOCTEXT("NoMeshFound", "No skeletal mesh found for skeleton '{0}'"), FText::FromString(TargetSkeletonName) ).ToString());
+		Description->SetText( FText::Format( LOCTEXT( "NoMeshFound", "No skeletal mesh found for skeleton '{0}'" ), FText::FromString( TargetSkeletonName ) ).ToString() );
 	}
 }
 
diff --git a/Engine/Source/Editor/DetailCustomizations/Private/AnimSequenceDetails.h b/Engine/Source/Editor/DetailCustomizations/Private/AnimSequenceDetails.h
--- a/Engine/Source/Editor/DetailCustomizations/Private/AnimSequenceDetails.h
+++ b/Engine/Source/Editor/DetailCustomizations/Private/AnimSequenceDetails.h
@@ -24,6 +24,7 @@ public:
 	virtual ~FAnimSequenceDetails();
 
 private:
+	
 	TWeakObjectPtr<USkeleton> TargetSkeleton;
 
 	TSharedPtr<IPropertyHandle> AdditiveAnimTypeHandle;
@@ -69,10 +70,14 @@ public:
 
 	void Construct(const FArguments& InArgs);
 
-	virtual void Tick(const FGeometry& AllottedGeometry, const double InCurrentTime, const float InDeltaTime) override;
+	virtual void Tick( const FGeometry& AllottedGeometry, const double InCurrentTime, const float InDeltaTime ) override;
 	void RefreshViewport();
 
 private:
+
+	/** Called to tick the preview during playback */
+	void OnTickPreview( double InCurrentTime, float InDeltaTime );
+
 	void InitSkeleton();
 
 	TSharedPtr<FEditorViewportClient> LevelViewportClient;
diff --git a/Engine/Source/Editor/DetailCustomizations/Private/DialogueStructsCustomizations.cpp b/Engine/Source/Editor/DetailCustomizations/Private/DialogueStructsCustomizations.cpp
--- a/Engine/Source/Editor/DetailCustomizations/Private/DialogueStructsCustomizations.cpp
+++ b/Engine/Source/Editor/DetailCustomizations/Private/DialogueStructsCustomizations.cpp
@@ -200,8 +200,6 @@ void SSpeakerDropDown::Tick( const FGeometry& AllottedGeometry, const double InC
 		ComboBox->ClearSelection();
 		ComboBox->RefreshOptions();
 	}
-
-	SCompoundWidget::Tick( AllottedGeometry, InCurrentTime, InDeltaTime );
 }
 
 struct FSpeakerMatcher
@@ -439,8 +437,6 @@ void STargetSetDropDown::Tick( const FGeometry& AllottedGeometry, const double I
 		ComboBox->ClearSelection();
 		ComboBox->RefreshOptions();
 	}
-
-	SCompoundWidget::Tick( AllottedGeometry, InCurrentTime, InDeltaTime );
 }
 
 void STargetSetDropDown::OnSelectionChanged(TSharedPtr<FTargetSet> TargetSet, ESelectInfo::Type SelectInfo)
@@ -599,7 +595,7 @@ public:
 	SLATE_END_ARGS()
 
 	void Construct( const FArguments& InArgs, const TSharedRef<IPropertyHandle>& InPropertyHandle, const TSharedRef<FAssetThumbnailPool>& InAssetThumbnailPool );
-	virtual void Tick(  const FGeometry& AllottedGeometry, const double InCurrentTime, const float InDeltaTime ) override;
+	virtual void Tick( const FGeometry& AllottedGeometry, const double InCurrentTime, const float InDeltaTime ) override;
 
 private:
 	bool IsDialogueWaveValid() const;
@@ -786,12 +782,8 @@ void SValidatedDialogueContextHeaderWidget::Construct( const FArguments& InArgs,
 		];
 	}
 }
-END_SLATE_FUNCTION_BUILD_OPTIMIZATION
-
-void SValidatedDialogueContextHeaderWidget::Tick(  const FGeometry& AllottedGeometry, const double InCurrentTime, const float InDeltaTime )
+void SValidatedDialogueContextHeaderWidget::Tick( const FGeometry& AllottedGeometry, const double InCurrentTime, const float InDeltaTime )
 {
-	SCompoundWidget::Tick( AllottedGeometry, InCurrentTime, InDeltaTime );
-
 	if( DialogueWaveParameterPropertyHandle.IsValid() && DialogueWaveParameterPropertyHandle->IsValidHandle() )
 	{
 		// Get dialogue wave.
diff --git a/Engine/Source/Editor/DetailCustomizations/Private/DialogueWaveWidgets.cpp b/Engine/Source/Editor/DetailCustomizations/Private/DialogueWaveWidgets.cpp
--- a/Engine/Source/Editor/DetailCustomizations/Private/DialogueWaveWidgets.cpp
+++ b/Engine/Source/Editor/DetailCustomizations/Private/DialogueWaveWidgets.cpp
@@ -166,8 +166,6 @@ void SDialogueVoicePropertyEditor::Construct( const FArguments& InArgs, const TS
 
 void SDialogueVoicePropertyEditor::Tick( const FGeometry& AllottedGeometry, const double InCurrentTime, const float InDeltaTime )
 {
-	SCompoundWidget::Tick( AllottedGeometry, InCurrentTime, InDeltaTime );
-
 	UDialogueVoice* CurrentDialogueVoice = NULL;
 	if( DialogueVoicePropertyHandle->IsValidHandle() )
 	{
@@ -465,7 +463,6 @@ float STargetsSummaryWidget::GetPreferredWidthForWrapping() const
 
 void STargetsSummaryWidget::Tick( const FGeometry& AllottedGeometry, const double InCurrentTime, const float InDeltaTime )
 {
-	SCompoundWidget::Tick( AllottedGeometry, InCurrentTime, InDeltaTime );
 	AllottedWidth = AllottedGeometry.Size.X;
 
 	if( TargetsPropertyHandle->IsValidHandle() )
@@ -473,15 +470,9 @@ void STargetsSummaryWidget::Tick( const FGeometry& AllottedGeometry, const doubl
 		uint32 TargetCount;
 		TargetsPropertyHandle->GetNumChildren(TargetCount);
 
-		bool bRefreshList = false;
 		if( TargetCount != DisplayedTargets.Num() )
 		{
 			// The array sizes differ so we need to refresh the list
-			bRefreshList = true;
-		}
-
-		if( bRefreshList )
-		{
 			GenerateContent();
 		}
 	}
@@ -915,12 +906,8 @@ void SDialogueContextHeaderWidget::Construct( const FArguments& InArgs, const TS
 		];
 	}
 }
-END_SLATE_FUNCTION_BUILD_OPTIMIZATION
-
 void SDialogueContextHeaderWidget::Tick( const FGeometry& AllottedGeometry, const double InCurrentTime, const float InDeltaTime )
 {
-	SCompoundWidget::Tick( AllottedGeometry, InCurrentTime, InDeltaTime );
-
 	if( !IsSpeakerValid() )
 	{
 		if( SpeakerErrorHint.IsValid() ) { SpeakerErrorHint->SetError( LOCTEXT("NullSpeakerError", "Speaker can not be \"None\".") ); }
@@ -948,15 +935,15 @@ bool SDialogueContextHeaderWidget::IsSpeakerValid() const
 	{
 		const TSharedPtr<IPropertyHandle> SpeakerPropertyHandle = ContextPropertyHandle->GetChildHandle("Speaker");
 
-		const UDialogueVoice* Speaker = NULL;
+		const UDialogueVoice* Speaker = nullptr;
 		if( SpeakerPropertyHandle.IsValid() && SpeakerPropertyHandle->IsValidHandle() )
 		{
-			UObject* Object = NULL;
+			UObject* Object = nullptr;
 			SpeakerPropertyHandle->GetValue(Object);
 			Speaker = Cast<UDialogueVoice>(Object);
 		}
 
-		Result = (Speaker != NULL);
+		Result = ( Speaker != nullptr );
 	}
 
 	return Result;
@@ -980,15 +967,15 @@ bool SDialogueContextHeaderWidget::IsTargetSetValid() const
 		{
 			TSharedPtr<IPropertyHandle> TargetPropertyHandle = TargetsArrayPropertyHandle->GetElement(i);
 
-			const UDialogueVoice* Target = NULL;
+			const UDialogueVoice* Target = nullptr;
 			if( TargetPropertyHandle.IsValid() && TargetPropertyHandle->IsValidHandle() )
 			{
-				UObject* Object = NULL;
+				UObject* Object = nullptr;
 				TargetPropertyHandle->GetValue(Object);
 				Target = Cast<UDialogueVoice>(Object);
 			}
 
-			if( Target == NULL )
+			if ( Target == nullptr )
 			{
 				Result = false;
 				break;
diff --git a/Engine/Source/Editor/DeviceProfileEditor/Private/SDeviceProfileEditor.cpp b/Engine/Source/Editor/DeviceProfileEditor/Private/SDeviceProfileEditor.cpp
--- a/Engine/Source/Editor/DeviceProfileEditor/Private/SDeviceProfileEditor.cpp
+++ b/Engine/Source/Editor/DeviceProfileEditor/Private/SDeviceProfileEditor.cpp
@@ -69,12 +69,12 @@ public:
 
 	// SCompoundWidget interface
 
-	virtual void Tick(const FGeometry& AllottedGeometry, const double InCurrentTime, const float InDeltaTime) override;
+	virtual void Tick( const FGeometry& AllottedGeometry, const double InCurrentTime, const float InDeltaTime ) override;
 
 private:
 
-	/** Holds a timer for checking whether the device profile configuration file needs to be checked out. */
-	float DefaultConfigCheckOutTimer;
+	/** Holds the last time checking whether the device profile configuration file needs to be checked out. */
+	double LastDefaultConfigCheckOutTime;
 
 	/** Holds a flag indicating whether the section's configuration file needs to be checked out. */
 	bool bIsDefaultConfigCheckOutNeeded;
@@ -132,7 +132,7 @@ bool SDeviceProfileSourceControl::IsCheckOutAvailable() const
 
 void SDeviceProfileSourceControl::Construct(const FArguments& InArgs)
 {
-	DefaultConfigCheckOutTimer = 0.0f;
+	LastDefaultConfigCheckOutTime = 0.0;
 	bIsDefaultConfigCheckOutNeeded = true;
 
 	const FString RelativeConfigFilePath = FString::Printf(TEXT("%sDefault%ss.ini"), *FPaths::SourceConfigDir(), *UDeviceProfile::StaticClass()->GetName());
@@ -212,18 +212,14 @@ void SDeviceProfileSourceControl::Construct(const FArguments& InArgs)
 }
 
 
-void SDeviceProfileSourceControl::Tick(const FGeometry& AllottedGeometry, const double InCurrentTime, const float InDeltaTime)
+void SDeviceProfileSourceControl::Tick( const FGeometry& AllottedGeometry, const double InCurrentTime, const float InDeltaTime )
 {
-	SCompoundWidget::Tick(AllottedGeometry, InCurrentTime, InDeltaTime);
-
 	// cache selected settings object's configuration file state
-	DefaultConfigCheckOutTimer += InDeltaTime;
-
-	if(DefaultConfigCheckOutTimer >= 1.0f)
+	if(InCurrentTime - LastDefaultConfigCheckOutTime >= 1.0f)
 	{
 		bIsDefaultConfigCheckOutNeeded = (FPaths::FileExists(AbsoluteConfigFilePath) && IFileManager::Get().IsReadOnly(*AbsoluteConfigFilePath));
 
-		DefaultConfigCheckOutTimer = 0.0f;
+		LastDefaultConfigCheckOutTime = InCurrentTime;
 	}
 }
 
diff --git a/Engine/Source/Editor/DistCurveEditor/Private/SCurveEditorViewport.cpp b/Engine/Source/Editor/DistCurveEditor/Private/SCurveEditorViewport.cpp
--- a/Engine/Source/Editor/DistCurveEditor/Private/SCurveEditorViewport.cpp
+++ b/Engine/Source/Editor/DistCurveEditor/Private/SCurveEditorViewport.cpp
@@ -102,8 +102,6 @@ TSharedPtr<SScrollBar> SCurveEditorViewport::GetVerticalScrollBar() const
 
 void SCurveEditorViewport::Tick( const FGeometry& AllottedGeometry, const double InCurrentTime, const float InDeltaTime )
 {
-	SCompoundWidget::Tick(AllottedGeometry, InCurrentTime, InDeltaTime);
-
 	// Check to see if we need to update the scrollbars due to a size change
 	const int32 CurrentHeight = AllottedGeometry.Size.Y;
 	if (CurrentHeight != PrevViewportHeight)
diff --git a/Engine/Source/Editor/EditorWidgets/Private/SAssetDiscoveryIndicator.cpp b/Engine/Source/Editor/EditorWidgets/Private/SAssetDiscoveryIndicator.cpp
--- a/Engine/Source/Editor/EditorWidgets/Private/SAssetDiscoveryIndicator.cpp
+++ b/Engine/Source/Editor/EditorWidgets/Private/SAssetDiscoveryIndicator.cpp
@@ -33,7 +33,7 @@ void SAssetDiscoveryIndicator::Construct( const FArguments& InArgs )
 
 		if ( InArgs._FadeIn )
 		{
-			FadeAnimation.Play();
+			FadeAnimation.Play( this->AsShared() );
 		}
 		else
 		{
@@ -90,7 +90,7 @@ void SAssetDiscoveryIndicator::OnAssetRegistryFileLoadProgress(int32 NumDiscover
 
 		if ( Progress >= 1.f )
 		{
-			FadeAnimation.PlayReverse();
+			FadeAnimation.PlayReverse(this->AsShared());
 		}
 	}
 }
diff --git a/Engine/Source/Editor/EditorWidgets/Private/SObjectNameEditableTextBox.cpp b/Engine/Source/Editor/EditorWidgets/Private/SObjectNameEditableTextBox.cpp
--- a/Engine/Source/Editor/EditorWidgets/Private/SObjectNameEditableTextBox.cpp
+++ b/Engine/Source/Editor/EditorWidgets/Private/SObjectNameEditableTextBox.cpp
@@ -14,6 +14,9 @@ const float SObjectNameEditableTextBox::CommittingAnimOffsetPercent = 0.2f;
 
 void SObjectNameEditableTextBox::Construct( const FArguments& InArgs )
 {
+	LastCommittedTime = 0.0;
+	bUpdateHighlightSpring = false;
+
 	Objects = InArgs._Objects;
 
 	ChildSlot
@@ -31,25 +34,35 @@ void SObjectNameEditableTextBox::Construct( const FArguments& InArgs )
 		];
 }
 
+EActiveTickReturnType SObjectNameEditableTextBox::UpdateHighlightSpringState( double InCurrentTime, float InDeltaTime )
+{
+	if ( (float)(InCurrentTime - LastCommittedTime) <= HighlightTargetEffectDuration )
+	{
+		bUpdateHighlightSpring = true;
+	}
+	else
+	{
+		bUpdateHighlightSpring = false;
+	}
+
+	return bUpdateHighlightSpring ? EActiveTickReturnType::KeepTicking : EActiveTickReturnType::StopTicking;
+}
+
 void SObjectNameEditableTextBox::Tick( const FGeometry& AllottedGeometry, const double InCurrentTime, const float InDeltaTime )
 {
-	// Call parent implementation.
-	SCompoundWidget::Tick( AllottedGeometry, InCurrentTime, InDeltaTime );
+	const bool bShouldAppearFocused = HasKeyboardFocus();
 
-	// Update highlight 'target' effect
+	if ( bUpdateHighlightSpring || bShouldAppearFocused )
 	{
+		// Update highlight 'target' effect
 		const float HighlightLeftX = HighlightRectLeftOffset;
 		const float HighlightRightX = HighlightRectRightOffset + AllottedGeometry.Size.X;
 
 		HighlightTargetLeftSpring.SetTarget( HighlightLeftX );
 		HighlightTargetRightSpring.SetTarget( HighlightRightX );
 
-		float TimeSinceHighlightInteraction = (float)( InCurrentTime - LastCommittedTime );
-		if( TimeSinceHighlightInteraction <= HighlightTargetEffectDuration )
-		{
-			HighlightTargetLeftSpring.Tick( InDeltaTime );
-			HighlightTargetRightSpring.Tick( InDeltaTime );
-		}
+		HighlightTargetLeftSpring.Tick( InDeltaTime );
+		HighlightTargetRightSpring.Tick( InDeltaTime );
 	}
 }
 
@@ -243,6 +256,7 @@ void SObjectNameEditableTextBox::OnNameTextCommitted(const FText& NewText, EText
 					{
 						Actor->SetActorLabel(TrimmedText.ToString());
 						LastCommittedTime = FSlateApplication::Get().GetCurrentTime();
+						RegisterActiveTick( 0.f, FWidgetActiveTickDelegate::CreateSP( this, &SObjectNameEditableTextBox::UpdateHighlightSpringState ) );
 					}
 				}
 			}
@@ -262,6 +276,7 @@ void SObjectNameEditableTextBox::OnNameTextCommitted(const FText& NewText, EText
 						{
 							Actor->SetActorLabel(TrimmedText.ToString());
 							LastCommittedTime = FSlateApplication::Get().GetCurrentTime();
+							RegisterActiveTick( 0.f, FWidgetActiveTickDelegate::CreateSP( this, &SObjectNameEditableTextBox::UpdateHighlightSpringState ) );
 						}
 					}
 				}
diff --git a/Engine/Source/Editor/EditorWidgets/Private/SObjectNameEditableTextBox.h b/Engine/Source/Editor/EditorWidgets/Private/SObjectNameEditableTextBox.h
--- a/Engine/Source/Editor/EditorWidgets/Private/SObjectNameEditableTextBox.h
+++ b/Engine/Source/Editor/EditorWidgets/Private/SObjectNameEditableTextBox.h
@@ -26,6 +26,9 @@ protected:
 	virtual int32 OnPaint( const FPaintArgs& Args, const FGeometry& AllottedGeometry, const FSlateRect& MyClippingRect, FSlateWindowElementList& OutDrawElements, int32 LayerId, const FWidgetStyle& InWidgetStyle, bool bParentEnabled ) const override;
 
 private:
+	/** Updates whether the highlight spring effect is happening */
+	EActiveTickReturnType UpdateHighlightSpringState( double InCurrentTime, float InDeltaTime );
+
 	/** Getter for the Text attribute of the editable text inside this widget */
 	FText GetNameText() const;
 
@@ -88,4 +91,7 @@ private:
 
 	/** The text box used to edit object names */ 
 	TSharedPtr< SEditableTextBox > TextBox;
+
+	// Temp flag to trigger a highlight spring update in the passive tick (because that's where the geometry is)
+	bool bUpdateHighlightSpring;
 };
diff --git a/Engine/Source/Editor/EditorWidgets/Private/STransportControl.cpp b/Engine/Source/Editor/EditorWidgets/Private/STransportControl.cpp
--- a/Engine/Source/Editor/EditorWidgets/Private/STransportControl.cpp
+++ b/Engine/Source/Editor/EditorWidgets/Private/STransportControl.cpp
@@ -3,8 +3,6 @@
 #include "EditorWidgetsPrivatePCH.h"
 #include "STransportControl.h"
 
-
-
 #define LOCTEXT_NAMESPACE "STransportControl"
 
 BEGIN_SLATE_FUNCTION_BUILD_OPTIMIZATION
@@ -148,7 +146,27 @@ void STransportControl::Construct( const STransportControl::FArguments& InArgs )
 }
 END_SLATE_FUNCTION_BUILD_OPTIMIZATION
 
-/** @return The appropriate icon to draw for the compile status of the selected blueprint */
+bool STransportControl::IsTickable() const
+{
+	// Only bother if an active tick delegate was provided
+	return TransportControlArgs.OnTickPlayback.IsBound() && TransportControlArgs.OnGetPlaybackMode.IsBound();
+}
+
+void STransportControl::Tick( float DeltaTime )
+{
+	const auto PlaybackMode = TransportControlArgs.OnGetPlaybackMode.Execute();
+	const bool bIsPlaying = PlaybackMode == EPlaybackMode::PlayingForward || PlaybackMode == EPlaybackMode::PlayingReverse || PlaybackMode == EPlaybackMode::Recording;
+
+	if ( bIsPlaying && !ActiveTickHandle.IsValid() )
+	{
+		ActiveTickHandle = RegisterActiveTick( 0.f, FWidgetActiveTickDelegate::CreateSP( this, &STransportControl::TickPlayback ) );
+	}
+	else if ( !bIsPlaying && ActiveTickHandle.IsValid() )
+	{
+		UnRegisterActiveTick( ActiveTickHandle.Pin().ToSharedRef() );
+	}
+}
+
 const FSlateBrush* STransportControl::GetForwardStatusIcon() const
 {
 	const auto PlaybackMode = TransportControlArgs.OnGetPlaybackMode.Execute();
@@ -205,4 +223,10 @@ FSlateColor STransportControl::GetLoopStatusColor() const
 	return FSlateColor(FLinearColor(1.f, 1.f, 1.f));
 }
 
+EActiveTickReturnType STransportControl::TickPlayback( double InCurrentTime, float InDeltaTime )
+{
+	TransportControlArgs.OnTickPlayback.Execute( InCurrentTime, InDeltaTime );
+	return EActiveTickReturnType::KeepTicking;
+}
+
 #undef LOCTEXT_NAMESPACE
diff --git a/Engine/Source/Editor/EditorWidgets/Private/STransportControl.h b/Engine/Source/Editor/EditorWidgets/Private/STransportControl.h
--- a/Engine/Source/Editor/EditorWidgets/Private/STransportControl.h
+++ b/Engine/Source/Editor/EditorWidgets/Private/STransportControl.h
@@ -3,7 +3,7 @@
 #pragma once
 
 
-class STransportControl : public ITransportControl
+class STransportControl : public ITransportControl, public FTickableEditorObject
 {
 public:
 	SLATE_BEGIN_ARGS(STransportControl)
@@ -19,6 +19,13 @@ public:
 	 */
 	void Construct( const FArguments& InArgs );
 
+	virtual ~STransportControl() {}
+
+	// Begin FTickableObjectBase implementation
+	virtual bool IsTickable() const override;
+	virtual void Tick( float DeltaTime ) override;
+	virtual TStatId GetStatId() const override { RETURN_QUICK_DECLARE_CYCLE_STAT( STransportControl, STATGROUP_Tickables ); }
+	// End FTickableObjectBase
 private:
 	const FSlateBrush* GetForwardStatusIcon() const;
 	FText GetForwardStatusTooltip() const;
@@ -26,8 +33,17 @@ private:
 	const FSlateBrush* GetBackwardStatusIcon() const;
 	FSlateColor GetLoopStatusColor() const;
 
+	/** Executes the OnTickPlayback delegate */
+	EActiveTickReturnType TickPlayback( double InCurrentTime, float InDeltaTime );
+
 	FReply OnToggleLooping();
 
 private:
+	/** The handle to the active tick */
+	TWeakPtr<FActiveTickHandle> ActiveTickHandle;
+
+	/** Whether the active tick is currently registered */
+	bool bIsActiveTickRegistered;
+
 	FTransportControlArgs TransportControlArgs;
 };
diff --git a/Engine/Source/Editor/EditorWidgets/Public/ITransportControl.h b/Engine/Source/Editor/EditorWidgets/Public/ITransportControl.h
--- a/Engine/Source/Editor/EditorWidgets/Public/ITransportControl.h
+++ b/Engine/Source/Editor/EditorWidgets/Public/ITransportControl.h
@@ -15,6 +15,7 @@ namespace EPlaybackMode
 
 DECLARE_DELEGATE_RetVal(bool, FOnGetLooping)
 DECLARE_DELEGATE_RetVal(EPlaybackMode::Type, FOnGetPlaybackMode)
+DECLARE_DELEGATE_TwoParams(FOnTickPlayback, double /*InCurrentTime*/, float /*InDeltaTime*/)
 
 struct FTransportControlArgs
 {
@@ -29,6 +30,7 @@ struct FTransportControlArgs
 		, OnToggleLooping()
 		, OnGetLooping()
 		, OnGetPlaybackMode()
+		, OnTickPlayback()
 	{}
 
 	FOnClicked OnForwardPlay;
@@ -41,6 +43,7 @@ struct FTransportControlArgs
 	FOnClicked OnToggleLooping;
 	FOnGetLooping OnGetLooping;
 	FOnGetPlaybackMode OnGetPlaybackMode;
+	FOnTickPlayback OnTickPlayback;
 };
 
 /**
diff --git a/Engine/Source/Editor/FoliageEdit/Private/SFoliageEdit.cpp b/Engine/Source/Editor/FoliageEdit/Private/SFoliageEdit.cpp
--- a/Engine/Source/Editor/FoliageEdit/Private/SFoliageEdit.cpp
+++ b/Engine/Source/Editor/FoliageEdit/Private/SFoliageEdit.cpp
@@ -340,11 +340,6 @@ SFoliageEdit::~SFoliageEdit()
 	FEditorDelegates::NewCurrentLevel.RemoveAll(this);
 }
 
-void SFoliageEdit::Tick(const FGeometry& AllottedGeometry, const double InCurrentTime, const float InDeltaTime)
-{
-	AssetThumbnailPool->Tick(InDeltaTime);
-}
-
 TSharedPtr<FAssetThumbnail> SFoliageEdit::CreateThumbnail(UStaticMesh* InStaticMesh)
 {
 	return MakeShareable(new FAssetThumbnail(InStaticMesh, 80, 80, AssetThumbnailPool));
diff --git a/Engine/Source/Editor/FoliageEdit/Private/SFoliageEdit.h b/Engine/Source/Editor/FoliageEdit/Private/SFoliageEdit.h
--- a/Engine/Source/Editor/FoliageEdit/Private/SFoliageEdit.h
+++ b/Engine/Source/Editor/FoliageEdit/Private/SFoliageEdit.h
@@ -20,10 +20,6 @@ public:
 
 	~SFoliageEdit();
 
-	// SWidget interface
-	virtual void Tick(const FGeometry& AllottedGeometry, const double InCurrentTime, const float InDeltaTime) override;
-	// End of SWidget interface
-
 	/** Creates the thumbnail for the passed in Static Mesh. Used by the MeshListView items. */
 	TSharedPtr<FAssetThumbnail> CreateThumbnail(UStaticMesh* InStaticMesh);
 
diff --git a/Engine/Source/Editor/GameProjectGeneration/Private/SGameProjectDialog.cpp b/Engine/Source/Editor/GameProjectGeneration/Private/SGameProjectDialog.cpp
--- a/Engine/Source/Editor/GameProjectGeneration/Private/SGameProjectDialog.cpp
+++ b/Engine/Source/Editor/GameProjectGeneration/Private/SGameProjectDialog.cpp
@@ -31,6 +31,9 @@ void SGameProjectDialog::Construct( const FArguments& InArgs )
 		ProjectBrowser = SNew(SProjectBrowser);
 	}
 
+	FadeAnimation.AddCurve( 0.f, 0.5f, ECurveEaseFunction::QuadOut );
+	RegisterActiveTick( 0.f, FWidgetActiveTickDelegate::CreateSP( this, &SGameProjectDialog::TriggerFadeInPostConstruct ) );
+
 	TSharedPtr<SWidget> Content;
 
 	if (InArgs._AllowProjectCreate && InArgs._AllowProjectOpening)
@@ -185,37 +188,20 @@ void SGameProjectDialog::Construct( const FArguments& InArgs )
 		check(false);
 	}
 }
-END_SLATE_FUNCTION_BUILD_OPTIMIZATION
-
-
-/* SWidget overrides
-*****************************************************************************/
 
-void SGameProjectDialog::Tick( const FGeometry& AllottedGeometry, const double InCurrentTime, const float InDeltaTime )
+EActiveTickReturnType SGameProjectDialog::TriggerFadeInPostConstruct(double InCurrentTime, float InDeltaTime)
 {
-	SCompoundWidget::Tick(AllottedGeometry, InCurrentTime, InDeltaTime);
-
 	// Play the intro fade in the first frame after the widget is created.
 	// We start it now instead of Construct because there is a lot of elapsed time between Construct and when we
 	// see the dialog and the beginning of the animation is cut off.
-	if (FadeAnimation.IsAtStart())
-	{
-		FadeIn();
-	}
-}
+	FadeAnimation.Play(this->AsShared());
 
+	return EActiveTickReturnType::KeepTicking;
+}
 
 /* SGameProjectDialog implementation
 *****************************************************************************/
 
-void SGameProjectDialog::FadeIn( )
-{
-	FadeAnimation = FCurveSequence();
-	FadeCurve = FadeAnimation.AddCurve(0.f, 0.5f, ECurveEaseFunction::QuadOut);
-	FadeAnimation.Play();
-}
-
-
 bool SGameProjectDialog::OpenProject( const FString& ProjectFile )
 {
 	FText FailReason;
@@ -258,7 +244,7 @@ FReply SGameProjectDialog::ShowProjectBrowser( )
 
 FLinearColor SGameProjectDialog::HandleCustomContentColorAndOpacity( ) const
 {
-	return FLinearColor(1,1,1, FadeCurve.GetLerp());
+	return FLinearColor(1,1,1, FadeAnimation.GetLerp());
 }
 
 
diff --git a/Engine/Source/Editor/GameProjectGeneration/Private/SGameProjectDialog.h b/Engine/Source/Editor/GameProjectGeneration/Private/SGameProjectDialog.h
--- a/Engine/Source/Editor/GameProjectGeneration/Private/SGameProjectDialog.h
+++ b/Engine/Source/Editor/GameProjectGeneration/Private/SGameProjectDialog.h
@@ -27,21 +27,7 @@ public:
 	/** Constructs this widget with InArgs */
 	void Construct( const FArguments& InArgs );
 
-public:
-
-	// Begin SWidget overrides
-
-	virtual void Tick( const FGeometry& AllottedGeometry, const double InCurrentTime, const float InDeltaTime ) override;
-
-	// End SWidget overrides
-
 protected:
-
-	/**
-	 * Fades the entire dialog in over a short period of time.
-	 */
-	void FadeIn( );
-
 	/**
 	 * Opens the specified project.
 	 *
@@ -62,6 +48,8 @@ protected:
 	FReply ShowProjectBrowser( );
 
 private:
+	/** Ensures the fade-in animation is played post-construct */
+	EActiveTickReturnType TriggerFadeInPostConstruct( double InCurrentTime, float InDeltaTime );
 
 	// Callback for getting the color of the custom content area.
 	FLinearColor HandleCustomContentColorAndOpacity() const;
@@ -83,9 +71,6 @@ private:
 	// Holds the fading animation.
 	FCurveSequence FadeAnimation;
 
-	// Holds the fading animation curve.
-	FCurveHandle FadeCurve;
-
 	/** The switcher widget to control which screen is in view */
 	TSharedPtr<SWidgetSwitcher> ContentAreaSwitcher;
 	TSharedPtr<SProjectBrowser> ProjectBrowser;
diff --git a/Engine/Source/Editor/GameProjectGeneration/Private/SNewClassDialog.cpp b/Engine/Source/Editor/GameProjectGeneration/Private/SNewClassDialog.cpp
--- a/Engine/Source/Editor/GameProjectGeneration/Private/SNewClassDialog.cpp
+++ b/Engine/Source/Editor/GameProjectGeneration/Private/SNewClassDialog.cpp
@@ -682,12 +682,8 @@ void SNewClassDialog::Construct( const FArguments& InArgs )
 		ParentClassListView->SetSelection(ParentClassItemsSource[0], ESelectInfo::Direct);
 	}
 }
-END_SLATE_FUNCTION_BUILD_OPTIMIZATION
-
 void SNewClassDialog::Tick( const FGeometry& AllottedGeometry, const double InCurrentTime, const float InDeltaTime )
 {
-	SCompoundWidget::Tick(AllottedGeometry, InCurrentTime, InDeltaTime);
-
 	// Every few seconds, the class name/path is checked for validity in case the disk contents changed and the location is now valid or invalid.
 	// After class creation, periodic checks are disabled to prevent a brief message indicating that the class you created already exists.
 	// This feature is re-enabled if the user did not restart and began editing parameters again.
diff --git a/Engine/Source/Editor/GraphEditor/Private/SGraphEditorImpl.cpp b/Engine/Source/Editor/GraphEditor/Private/SGraphEditorImpl.cpp
--- a/Engine/Source/Editor/GraphEditor/Private/SGraphEditorImpl.cpp
+++ b/Engine/Source/Editor/GraphEditor/Private/SGraphEditorImpl.cpp
@@ -87,19 +87,21 @@ void SGraphEditorImpl::NotifyGraphChanged()
 
 void SGraphEditorImpl::OnGraphChanged(const FEdGraphEditAction& InAction)
 {
-	if (!bNeedsRefresh)
+	if ( !bIsActiveTickRegistered )
 	{
 		// Remove the old user interface nodes
 		GraphPanel->PurgeVisualRepresentation();
-	}
 
-	bNeedsRefresh = true;
+		// Trigger the refresh
+		bIsActiveTickRegistered = true;
+		RegisterActiveTick( 0.f, FWidgetActiveTickDelegate::CreateSP( this, &SGraphEditorImpl::TriggerRefresh ) );
+	}
 }
 
-void SGraphEditorImpl::GraphEd_OnPanelUpdated()
-{
-	bNeedsRefresh = false;
-}
+//void SGraphEditorImpl::GraphEd_OnPanelUpdated()
+//{
+//	
+//}
 
 const TSet<class UObject*>& SGraphEditorImpl::GetSelectedNodes() const
 {
@@ -156,6 +158,8 @@ void SGraphEditorImpl::Construct( const FArguments& InArgs )
 	OnNavigateHistoryForward = InArgs._OnNavigateHistoryForward;
 	OnNodeSpawnedByKeymap = InArgs._GraphEvents.OnNodeSpawnedByKeymap;
 
+	bIsActiveTickRegistered = false;
+
 	// Make sure that the editor knows about what kinds
 	// of commands GraphEditor can do.
 	FGraphEditorCommands::Register();
@@ -188,8 +192,7 @@ void SGraphEditorImpl::Construct( const FArguments& InArgs )
 
 	GraphPinForMenu = NULL;
 	EdGraphObj = InArgs._GraphToEdit;
-	bNeedsRefresh = false;
-		
+
 	OnFocused = InArgs._GraphEvents.OnFocused;
 	OnCreateActionMenu = InArgs._GraphEvents.OnCreateActionMenu;
 	
@@ -241,7 +244,7 @@ void SGraphEditorImpl::Construct( const FArguments& InArgs )
 			.OnVerifyTextCommit( InArgs._GraphEvents.OnVerifyTextCommit )
 			.OnTextCommitted( InArgs._GraphEvents.OnTextCommitted )
 			.OnSpawnNodeByShortcut( InArgs._GraphEvents.OnSpawnNodeByShortcut )
-			.OnUpdateGraphPanel( this, &SGraphEditorImpl::GraphEd_OnPanelUpdated )
+			//.OnUpdateGraphPanel( this, &SGraphEditorImpl::GraphEd_OnPanelUpdated )
 			.OnDisallowedPinConnection( InArgs._GraphEvents.OnDisallowedPinConnection )
 			.ShowGraphStateOverlay(InArgs._ShowGraphStateOverlay)
 		]
@@ -372,12 +375,6 @@ SGraphEditorImpl::~SGraphEditorImpl()
 
 void SGraphEditorImpl::Tick( const FGeometry& AllottedGeometry, const double InCurrentTime, const float InDeltaTime )
 {
-	if (bNeedsRefresh)
-	{
-		bNeedsRefresh = false;
-		GraphPanel->Update();
-	}	
-
 	// If locked to another graph editor, and our panel has moved, synchronise the locked graph editor accordingly
 	if ((EdGraphObj != NULL) && GraphPanel.IsValid())
 	{
@@ -388,6 +385,14 @@ void SGraphEditorImpl::Tick( const FGeometry& AllottedGeometry, const double InC
 	}
 }
 
+EActiveTickReturnType SGraphEditorImpl::TriggerRefresh( double InCurrentTime, float InDeltaTime )
+{
+	GraphPanel->Update();
+
+	bIsActiveTickRegistered = false;
+	return EActiveTickReturnType::StopTicking;
+}
+
 void SGraphEditorImpl::OnClosedActionMenu()
 {
 	GraphPanel->OnStopMakingConnection(/*bForceStop=*/ true);
diff --git a/Engine/Source/Editor/GraphEditor/Private/SGraphEditorImpl.h b/Engine/Source/Editor/GraphEditor/Private/SGraphEditorImpl.h
--- a/Engine/Source/Editor/GraphEditor/Private/SGraphEditorImpl.h
+++ b/Engine/Source/Editor/GraphEditor/Private/SGraphEditorImpl.h
@@ -51,9 +51,6 @@ private:
 
 	UEdGraphPin* GraphPinForMenu;
 
-	/** Do we need to refresh next tick? */
-	bool bNeedsRefresh;
-
 	/** Info on the appearance */
 	TAttribute<FGraphAppearanceInfo> Appearance;
 
@@ -108,7 +105,7 @@ public:
 
 	FActionMenuContent GraphEd_OnGetContextMenuFor(const FGraphContextMenuArguments& SpawnInfo);
 
-	void GraphEd_OnPanelUpdated();
+	//void GraphEd_OnPanelUpdated();
 
 	// SWidget interface
 	virtual void Tick( const FGeometry& AllottedGeometry, const double InCurrentTime, const float InDeltaTime ) override;
@@ -157,7 +154,11 @@ protected:
 	// SGraphEditor interface
 	virtual void OnGraphChanged( const FEdGraphEditAction& InAction ) override;
 	// End of SGraphEditorInterface
+
 private:
+	/** One-off active tick to ensure the graph is refreshes as needed */
+	EActiveTickReturnType TriggerRefresh( double InCurrentTime, float InDeltaTime );
+
 	FText GetZoomText() const;
 
 	FSlateColor GetZoomTextColorAndOpacity() const;
@@ -165,5 +166,8 @@ private:
 	bool IsGraphEditable() const;
 
 	bool IsLocked() const;
+
+private:
+	bool bIsActiveTickRegistered;
 };
 
diff --git a/Engine/Source/Editor/GraphEditor/Private/SGraphNode.cpp b/Engine/Source/Editor/GraphEditor/Private/SGraphNode.cpp
--- a/Engine/Source/Editor/GraphEditor/Private/SGraphNode.cpp
+++ b/Engine/Source/Editor/GraphEditor/Private/SGraphNode.cpp
@@ -35,10 +35,8 @@ void SNodeTitle::Construct(const FArguments& InArgs, UEdGraphNode* InNode)
 	RebuildWidget();
 }
 
-void SNodeTitle::Tick(const FGeometry& AllottedGeometry, const double InCurrentTime, const float InDeltaTime)
+void SNodeTitle::Tick( const FGeometry& AllottedGeometry, const double InCurrentTime, const float InDeltaTime )
 {
-	SCompoundWidget::Tick(AllottedGeometry, InCurrentTime, InDeltaTime);
-
 	// Checks to see if the cached string is valid, and if not, updates it.
 	if (TitleText.Get().CompareTo(CachedTitle) != 0)
 	{
@@ -1085,7 +1083,7 @@ TSharedPtr<SGraphPin> SGraphNode::FindWidgetForPin( UEdGraphPin* ThePin ) const
 
 void SGraphNode::PlaySpawnEffect()
 {
-	SpawnAnim.Play();
+	SpawnAnim.Play( this->AsShared() );
 }
 
 FVector2D SGraphNode::GetContentScale() const
diff --git a/Engine/Source/Editor/GraphEditor/Private/SGraphNodeDocumentation.cpp b/Engine/Source/Editor/GraphEditor/Private/SGraphNodeDocumentation.cpp
--- a/Engine/Source/Editor/GraphEditor/Private/SGraphNodeDocumentation.cpp
+++ b/Engine/Source/Editor/GraphEditor/Private/SGraphNodeDocumentation.cpp
@@ -332,8 +332,6 @@ void SGraphNodeDocumentation::Tick( const FGeometry& AllottedGeometry, const dou
 		GraphNode->NodeHeight = 0.f;
 		UpdateGraphNode();
 	}
-
-	SGraphNodeResizable::Tick( AllottedGeometry, InCurrentTime, InDeltaTime );
 }
 
 /////////////////////////////////////////////////////
diff --git a/Engine/Source/Editor/GraphEditor/Private/SGraphPanel.cpp b/Engine/Source/Editor/GraphEditor/Private/SGraphPanel.cpp
--- a/Engine/Source/Editor/GraphEditor/Private/SGraphPanel.cpp
+++ b/Engine/Source/Editor/GraphEditor/Private/SGraphPanel.cpp
@@ -89,7 +89,9 @@ void SGraphPanel::Construct( const SGraphPanel::FArguments& InArgs )
 	}
 
 	BounceCurve.AddCurve(0.0f, 1.0f);
-	BounceCurve.Play();
+
+	FEditorDelegates::BeginPIE.AddRaw( this, &SGraphPanel::OnBeginPIE );
+	FEditorDelegates::EndPIE.AddRaw( this, &SGraphPanel::OnEndPIE );
 
 	// Register for notifications
 	MyRegisteredGraphChangedDelegate = FOnGraphChanged::FDelegate::CreateSP(this, &SGraphPanel::OnGraphChanged);
@@ -100,6 +102,9 @@ void SGraphPanel::Construct( const SGraphPanel::FArguments& InArgs )
 
 SGraphPanel::~SGraphPanel()
 {
+	FEditorDelegates::BeginPIE.RemoveAll( this );
+	FEditorDelegates::EndPIE.RemoveAll( this );
+
 	this->GraphObj->RemoveOnGraphChangedHandler(MyRegisteredGraphChangedDelegate);
 }
 
@@ -264,8 +269,8 @@ int32 SGraphPanel::OnPaint( const FPaintArgs& Args, const FGeometry& AllottedGeo
 							{
 								FPaintGeometry BouncedGeometry = CurWidget.Geometry.ToPaintGeometry(OverlayInfo.OverlayOffset, OverlayBrush->ImageSize, 1.f);
 
-								// Handle bouncing
-								const float BounceValue = FMath::Sin(2.0f * PI * BounceCurve.GetLerpLooping());
+								// Handle bouncing during PIE
+								const float BounceValue = FMath::Sin(2.0f * PI * BounceCurve.GetLerp());
 								BouncedGeometry.DrawPosition += (OverlayInfo.AnimationEnvelope * BounceValue * ZoomFactor);
 
 								CurWidgetsMaxLayerId++;
@@ -1061,27 +1066,25 @@ bool SGraphPanel::IsRectVisible(const FVector2D &TopLeft, const FVector2D &Botto
 
 bool SGraphPanel::JumpToRect(const FVector2D &TopLeft, const FVector2D &BottomRight)
 {
-	ZoomTargetTopLeft = TopLeft;
-	ZoomTargetBottomRight = BottomRight;
-	bDeferredZoomingToFit = true;
+	ZoomToTarget(TopLeft, BottomRight);
+
 	return true;
 }
 
 void SGraphPanel::JumpToNode(const UEdGraphNode* JumpToMe, bool bRequestRename)
 {
-	bDeferredZoomingToFit = false;
-
-	if (JumpToMe != NULL)
+	if (JumpToMe != nullptr)
 	{
 		if (bRequestRename)
 		{
 			TSharedRef<SNode>* pWidget = NodeToWidgetLookup.Find(JumpToMe);
-			if (pWidget != NULL)
+			if (pWidget != nullptr)
 			{
 				TSharedRef<SGraphNode> GraphNode = StaticCastSharedRef<SGraphNode>(*pWidget);
 				GraphNode->RequestRename();
 			}
 		}
+
 		// Select this node, and request that we jump to it.
 		SelectAndCenterObject(JumpToMe, true);
 	}
@@ -1089,12 +1092,24 @@ void SGraphPanel::JumpToNode(const UEdGraphNode* JumpToMe, bool bRequestRename)
 
 void SGraphPanel::JumpToPin(const UEdGraphPin* JumpToMe)
 {
-	if (JumpToMe != NULL)
+	if (JumpToMe != nullptr)
 	{
 		JumpToNode(JumpToMe->GetOwningNode(), false);
 	}
 }
 
+void SGraphPanel::OnBeginPIE( const bool bIsSimulating )
+{
+	// Play the bounce curve on a continuous loop during PIE
+	BounceCurve.Play( this->AsShared(), true );
+}
+
+void SGraphPanel::OnEndPIE( const bool bIsSimulating )
+{
+	// Stop the bounce curve
+	BounceCurve.JumpToEnd();
+}
+
 void SGraphPanel::OnGraphChanged(const FEdGraphEditAction& EditAction)
 {
 	if ((EditAction.Graph == GraphObj) &&
diff --git a/Engine/Source/Editor/GraphEditor/Private/SGraphPanel.h b/Engine/Source/Editor/GraphEditor/Private/SGraphPanel.h
--- a/Engine/Source/Editor/GraphEditor/Private/SGraphPanel.h
+++ b/Engine/Source/Editor/GraphEditor/Private/SGraphPanel.h
@@ -245,6 +245,12 @@ private:
 
 	FOnGraphChanged::FDelegate MyRegisteredGraphChangedDelegate;
 private:
+	/** Called when PIE begins */
+	void OnBeginPIE( const bool bIsSimulating );
+
+	/** Called when PIE ends */
+	void OnEndPIE( const bool bIsSimulating );
+
 	/** Called when watched graph changes */
 	void OnGraphChanged( const FEdGraphEditAction& InAction );
 
diff --git a/Engine/Source/Editor/GraphEditor/Private/SGraphPin.cpp b/Engine/Source/Editor/GraphEditor/Private/SGraphPin.cpp
--- a/Engine/Source/Editor/GraphEditor/Private/SGraphPin.cpp
+++ b/Engine/Source/Editor/GraphEditor/Private/SGraphPin.cpp
@@ -728,8 +728,6 @@ void SGraphPin::Tick( const FGeometry& AllottedGeometry, const double InCurrentT
 {
 	CachedNodeOffset = AllottedGeometry.AbsolutePosition/AllottedGeometry.Scale - OwnerNodePtr.Pin()->GetUnscaledPosition();
 	CachedNodeOffset.Y += AllottedGeometry.Size.Y * 0.5f;
-
-	SBorder::Tick(AllottedGeometry, InCurrentTime, InDeltaTime);
 }
 
 UEdGraphPin* SGraphPin::GetPinObj() const
diff --git a/Engine/Source/Editor/GraphEditor/Private/SNodePanel.cpp b/Engine/Source/Editor/GraphEditor/Private/SNodePanel.cpp
--- a/Engine/Source/Editor/GraphEditor/Private/SNodePanel.cpp
+++ b/Engine/Source/Editor/GraphEditor/Private/SNodePanel.cpp
@@ -132,7 +132,7 @@ void FGraphSelectionManager::SetSelectionSet(FGraphPanelSelectionSet& NewSet)
 
 void FGraphSelectionManager::SetNodeSelection(SelectedItemType Node, bool bSelect)
 {
-	ensureMsg(Node != NULL, TEXT("Node is invalid"));
+	ensureMsg(Node != nullptr, TEXT("Node is invalid"));
 	if (bSelect)
 	{
 		SelectedNodes.Add(Node);
@@ -293,7 +293,6 @@ void SNodePanel::Construct()
 	TotalMouseDelta = 0;
 	TotalMouseDeltaY = 0;
 	bDeferredZoomToSelection = false;
-	bDeferredZoomingToFit = false;
 	bDeferredZoomToNodeExtents = false;
 
 	ZoomTargetTopLeft = FVector2D::ZeroVector;
@@ -304,17 +303,17 @@ void SNodePanel::Construct()
 	bTeleportInsteadOfScrollingWhenZoomingToFit = false;
 
 	DeferredSelectionTargetObjects.Empty();
-	DeferredMovementTargetObject = NULL;
+	DeferredMovementTargetObject = nullptr;
 
 	bIsPanning = false;
 	bIsZoomingWithTrackpad = false;
 	IsEditable.Set(true);
 
 	ZoomLevelFade = FCurveSequence( 0.0f, 1.0f );
-	ZoomLevelFade.Play();
+	ZoomLevelFade.Play( this->AsShared() );
 
 	ZoomLevelGraphFade = FCurveSequence( 0.0f, 0.5f );
-	ZoomLevelGraphFade.Play();
+	ZoomLevelGraphFade.Play( this->AsShared() );
 
 	PastePosition = FVector2D::ZeroVector;
 
@@ -409,9 +408,35 @@ void SNodePanel::OnEndNodeInteraction(const TSharedRef<SNode>& InNodeToDrag)
 {
 }
 
-// Ticks this widget.  Override in derived classes, but always call the parent implementation.
-void SNodePanel::Tick(const FGeometry& AllottedGeometry, const double InCurrentTime, const float InDeltaTime)
+EActiveTickReturnType SNodePanel::HandleZoomToFit(double InCurrentTime, float InDeltaTime)
 {
+	const FVector2D DesiredViewCenter = ( ZoomTargetTopLeft + ZoomTargetBottomRight ) * 0.5f;
+	const bool bDoneScrolling = ScrollToLocation(CachedGeometry, DesiredViewCenter, bTeleportInsteadOfScrollingWhenZoomingToFit ? 1000.0f : InDeltaTime);
+	bool bDoneZooming = ZoomToLocation(CachedGeometry.Size, ZoomTargetBottomRight - ZoomTargetTopLeft, bDoneScrolling);
+
+	if (bDoneZooming && bDoneScrolling)
+	{
+		// One final push to make sure we centered in the end
+		ViewOffset = DesiredViewCenter - ( 0.5f * CachedGeometry.Scale * CachedGeometry.Size / GetZoomAmount() );
+		
+		// Reset ZoomPadding
+		ZoomPadding = NodePanelDefs::DefaultZoomPadding;
+		ZoomTargetTopLeft = FVector2D::ZeroVector;
+		ZoomTargetBottomRight = FVector2D::ZeroVector;
+
+		DeferredMovementTargetObject = nullptr;
+
+		return EActiveTickReturnType::StopTicking;
+	}
+
+	return EActiveTickReturnType::KeepTicking;
+}
+
+void SNodePanel::Tick( const FGeometry& AllottedGeometry, const double InCurrentTime, const float InDeltaTime )
+{
+	CachedGeometry = AllottedGeometry;
+	bool bWasActiveTickRegisteredThisFrame = false;
+
 	if(DeferredSelectionTargetObjects.Num() > 0)
 	{
 		FGraphPanelSelectionSet NewSelectionSet;
@@ -436,11 +461,14 @@ void SNodePanel::Tick(const FGeometry& AllottedGeometry, const double InCurrentT
 			bDeferredZoomToNodeExtents = false;
 		}
 	}
-	else if(DeferredMovementTargetObject != NULL && GetBoundsForNodes(true, ZoomTargetTopLeft, ZoomTargetBottomRight, ZoomPadding))
+	else if(DeferredMovementTargetObject != nullptr && GetBoundsForNodes(true, ZoomTargetTopLeft, ZoomTargetBottomRight, ZoomPadding))
 	{
-		// Zoom to fit the target node
-		bDeferredZoomingToFit = true;
-		DeferredMovementTargetObject = NULL;
+		DeferredMovementTargetObject = nullptr;
+		if (!ActiveTickHandle.IsValid())
+		{
+			ActiveTickHandle = RegisterActiveTick(0.f, FWidgetActiveTickDelegate::CreateSP(this, &SNodePanel::HandleZoomToFit));
+			bWasActiveTickRegisteredThisFrame = true;
+		}
 	}
 	
 	// Zoom to node extents
@@ -451,27 +479,11 @@ void SNodePanel::Tick(const FGeometry& AllottedGeometry, const double InCurrentT
 		if( GetBoundsForNodes(bDeferredZoomToSelection, ZoomTargetTopLeft, ZoomTargetBottomRight, ZoomPadding))
 		{
 			bDeferredZoomToSelection = false;
-			bDeferredZoomingToFit = true;
-		}
-	}
-
-	if (bDeferredZoomingToFit)
-	{
-		const FVector2D DesiredViewCenter = (ZoomTargetTopLeft + ZoomTargetBottomRight) * 0.5f;
-		const bool bDoneScrolling = ScrollToLocation(AllottedGeometry, DesiredViewCenter, bTeleportInsteadOfScrollingWhenZoomingToFit ? 1000.0f : InDeltaTime);
-		bool bDoneZooming = ZoomToLocation(AllottedGeometry.Size, ZoomTargetBottomRight - ZoomTargetTopLeft, bDoneScrolling);
-
-		if (bDoneZooming && bDoneScrolling)
-		{	
-			// One final push to make sure we centered in the end
-			ViewOffset = DesiredViewCenter - (0.5f * AllottedGeometry.Scale * AllottedGeometry.Size / GetZoomAmount());
-			// Reset ZoomPadding
-			ZoomPadding = NodePanelDefs::DefaultZoomPadding;
-			ZoomTargetTopLeft = FVector2D::ZeroVector;
-			ZoomTargetBottomRight = FVector2D::ZeroVector;
-
-			bDeferredZoomingToFit = false;
-			DeferredMovementTargetObject = NULL;
+			if (!ActiveTickHandle.IsValid())
+			{
+				ActiveTickHandle = RegisterActiveTick(0.f, FWidgetActiveTickDelegate::CreateSP(this, &SNodePanel::HandleZoomToFit));
+				bWasActiveTickRegisteredThisFrame = true;
+			}
 		}
 	}
 
@@ -493,7 +505,7 @@ void SNodePanel::Tick(const FGeometry& AllottedGeometry, const double InCurrentT
 	OldZoomAmount = GetZoomAmount();
 	OldViewOffset = ViewOffset;
 
-	SPanel::Tick(AllottedGeometry, InCurrentTime, InDeltaTime);
+	SPanel::Tick( AllottedGeometry, InCurrentTime, InDeltaTime );
 }
 
 // The system calls this method to notify the widget that a mouse button was pressed within it. This event is bubbled.
@@ -514,7 +526,7 @@ FReply SNodePanel::OnMouseButtonDown( const FGeometry& MyGeometry, const FPointe
 		ReplyState.CaptureMouse( SharedThis(this) );
 		ReplyState.UseHighPrecisionMouseMovement( SharedThis(this) );
 
-		DeferredMovementTargetObject = NULL; // clear any interpolation when you manually zoom
+		DeferredMovementTargetObject = nullptr; // clear any interpolation when you manually zoom
 		TotalMouseDeltaY = 0;
 
 		if (!FSlateApplication::Get().IsUsingTrackpad()) // on trackpad we don't know yet if user wants to zoom or bring up the context menu
@@ -571,7 +583,7 @@ FReply SNodePanel::OnMouseButtonDown( const FGeometry& MyGeometry, const FPointe
 
 		SoftwareCursorPosition = PanelCoordToGraphCoord( MyGeometry.AbsoluteToLocal( MouseEvent.GetScreenSpacePosition() ) );
 
-		DeferredMovementTargetObject = NULL; // clear any interpolation when you manually pan
+		DeferredMovementTargetObject = nullptr; // clear any interpolation when you manually pan
 
 		// RIGHT BUTTON is for dragging and Context Menu.
 		return ReplyState;
@@ -683,7 +695,7 @@ FReply SNodePanel::OnMouseMove(const FGeometry& MyGeometry, const FPointerEvent&
 				// Update the amount to pan panel
 				UpdateViewOffset(MyGeometry, MouseEvent.GetScreenSpacePosition());
 
-				const bool bCursorInDeadZone = TotalMouseDelta <= FSlateApplication::Get().GetDragTriggerDistnace();
+				const bool bCursorInDeadZone = TotalMouseDelta <= FSlateApplication::Get().GetDragTriggerDistance();
 
 				if ( NodeBeingDragged.IsValid() )
 				{
@@ -717,7 +729,7 @@ FReply SNodePanel::OnMouseMove(const FGeometry& MyGeometry, const FPointerEvent&
 							for (FGraphPanelSelectionSet::TIterator NodeIt(SelectionManager.SelectedNodes); NodeIt; ++NodeIt)
 							{
 								TSharedRef<SNode>* pWidget = NodeToWidgetLookup.Find(*NodeIt);
-								if (pWidget != NULL)
+								if (pWidget != nullptr)
 								{
 									SNode& Widget = pWidget->Get();
 									FDefferedNodePosition NodePosition = { &Widget, Widget.GetPosition() + DeltaPos };
@@ -780,7 +792,7 @@ FReply SNodePanel::OnMouseButtonUp( const FGeometry& MyGeometry, const FPointerE
 	// Did the user move the cursor sufficiently far, or is it in a dead zone?
 	// In Dead zone     - implies actions like summoning context menus and general clicking.
 	// Out of Dead Zone - implies dragging actions like moving nodes and marquee selection.
-	const bool bCursorInDeadZone = TotalMouseDelta <= FSlateApplication::Get().GetDragTriggerDistnace();
+	const bool bCursorInDeadZone = TotalMouseDelta <= FSlateApplication::Get().GetDragTriggerDistance();
 
 	// Set to true later if we need to finish with the software cursor
 	bool bRemoveSoftwareCursor = false;
@@ -1057,6 +1069,13 @@ void SNodePanel::SelectAndCenterObject(const UObject* ObjectToSelect, bool bCent
 	{
 		DeferredMovementTargetObject = ObjectToSelect;
 	}
+
+	auto PinnedActiveTickHandle = ActiveTickHandle.Pin();
+	if (PinnedActiveTickHandle.IsValid())
+	{
+		// Stop zooming for now and let the a new zoom target be established
+		UnRegisterActiveTick(PinnedActiveTickHandle.ToSharedRef());
+	}
 }
 
 /** Add a slot to the CanvasPanel dynamically */
@@ -1099,7 +1118,7 @@ void SNodePanel::PopulateVisibleChildren(const FGeometry& AllottedGeometry)
 // Is the given node being observed by a widget in this panel?
 bool SNodePanel::Contains(UObject* Node) const
 {
-	return NodeToWidgetLookup.Find(Node) != NULL;
+	return NodeToWidgetLookup.Find(Node) != nullptr;
 }
 
 void SNodePanel::RestoreViewSettings(const FVector2D& InViewOffset, float InZoomAmount)
@@ -1115,8 +1134,13 @@ void SNodePanel::RestoreViewSettings(const FVector2D& InViewOffset, float InZoom
 	else
 	{
 		ZoomLevel = ZoomLevels->GetNearestZoomLevel(InZoomAmount);
-		bDeferredZoomingToFit = false;
 		bDeferredZoomToNodeExtents = false;
+
+		auto PinnedActiveTickHandle = ActiveTickHandle.Pin();
+		if (PinnedActiveTickHandle.IsValid())
+		{
+			UnRegisterActiveTick(PinnedActiveTickHandle.ToSharedRef());
+		}
 	}
 
 	PostChangedZoom();
@@ -1369,7 +1393,7 @@ bool SNodePanel::GetBoundsForNodes(bool bSelectionSetOnly, FVector2D& MinCorner,
 		for (FGraphPanelSelectionSet::TConstIterator NodeIt(SelectionManager.GetSelectedNodes()); NodeIt; ++NodeIt)
 		{
 			TSharedRef<SNode>* pWidget = NodeToWidgetLookup.Find(*NodeIt);
-			if (pWidget != NULL)
+			if (pWidget != nullptr)
 			{
 				SNode& Widget = pWidget->Get();
 				const FVector2D Lower = Widget.GetPosition();
@@ -1455,7 +1479,7 @@ bool SNodePanel::ZoomToLocation(const FVector2D& CurrentSizeWithoutZoom, const F
 			// Animate to it
 			PreviousZoomLevel = ZoomLevel;
 			ZoomLevel = FMath::Clamp(DesiredZoom, 0, NumZoomLevels-1);
-			ZoomLevelGraphFade.Play();
+			ZoomLevelGraphFade.Play( this->AsShared() );
 			return false;
 		}
 		else
@@ -1465,7 +1489,7 @@ bool SNodePanel::ZoomToLocation(const FVector2D& CurrentSizeWithoutZoom, const F
 			{
 				// Zooming out; do it instantly
 				ZoomLevel = PreviousZoomLevel = DesiredZoom;
-				ZoomLevelFade.Play();
+				ZoomLevelFade.Play( this->AsShared() );
 			}
 			else
 			{
@@ -1473,7 +1497,7 @@ bool SNodePanel::ZoomToLocation(const FVector2D& CurrentSizeWithoutZoom, const F
 				if (bDoneScrolling)
 				{
 					ZoomLevel = PreviousZoomLevel = DesiredZoom;
-					ZoomLevelFade.Play();
+					ZoomLevelFade.Play( this->AsShared() );
 				}
 			}
 		}
@@ -1486,13 +1510,24 @@ bool SNodePanel::ZoomToLocation(const FVector2D& CurrentSizeWithoutZoom, const F
 
 void SNodePanel::ZoomToFit(bool bOnlySelection)
 {
-	// Deferred zoom/pan action
-	bDeferredZoomingToFit = true;
 	bDeferredZoomToNodeExtents = true;
 	bDeferredZoomToSelection = bOnlySelection;
 	ZoomPadding = NodePanelDefs::DefaultZoomPadding;
 }
 
+void SNodePanel::ZoomToTarget(const FVector2D& TopLeft, const FVector2D& BottomRight)
+{
+	bDeferredZoomToNodeExtents = false;
+
+	ZoomTargetTopLeft = TopLeft;
+	ZoomTargetBottomRight = BottomRight;
+
+	if (!ActiveTickHandle.IsValid())
+	{
+		ActiveTickHandle = RegisterActiveTick(0.f, FWidgetActiveTickDelegate::CreateSP(this, &SNodePanel::HandleZoomToFit));
+	}
+}
+
 void SNodePanel::ChangeZoomLevel(int32 ZoomLevelDelta, const FVector2D& WidgetSpaceZoomOrigin, bool bOverrideZoomLimiting)
 {
 	// We want to zoom into this point; i.e. keep it the same fraction offset into the panel
@@ -1525,7 +1560,7 @@ void SNodePanel::ChangeZoomLevel(int32 ZoomLevelDelta, const FVector2D& WidgetSp
 	}
 
 	// Note: This happens even when maxed out at a stop; so the user sees the animation and knows that they're at max zoom in/out
-	ZoomLevelFade.Play();
+	ZoomLevelFade.Play( this->AsShared() );
 
 	// Re-center the screen so that it feels like zooming around the cursor.
 	this->ViewOffset = PointToMaintainGraphSpace - WidgetSpaceZoomOrigin / GetZoomAmount();
@@ -1551,7 +1586,7 @@ FVector2D SNodePanel::GetPastePosition() const
 
 bool SNodePanel::HasDeferredObjectFocus() const
 {
-	return DeferredMovementTargetObject != NULL;
+	return DeferredMovementTargetObject != nullptr;
 }
 
 void SNodePanel::PostChangedZoom()
diff --git a/Engine/Source/Editor/GraphEditor/Public/SNodePanel.h b/Engine/Source/Editor/GraphEditor/Public/SNodePanel.h
--- a/Engine/Source/Editor/GraphEditor/Public/SNodePanel.h
+++ b/Engine/Source/Editor/GraphEditor/Public/SNodePanel.h
@@ -538,7 +538,7 @@ public:
 	// End of SPanel interface
 
 	// SWidget interface
-	virtual void Tick(const FGeometry& AllottedGeometry, const double InCurrentTime, const float InDeltaTime) override;
+	virtual void Tick( const FGeometry& AllottedGeometry, const double InCurrentTime, const float InDeltaTime ) override;
 	virtual FReply OnMouseButtonDown(const FGeometry& MyGeometry, const FPointerEvent& MouseEvent) override;
 	virtual FReply OnMouseButtonUp(const FGeometry& MyGeometry, const FPointerEvent& MouseEvent) override;
 	virtual FReply OnMouseMove(const FGeometry& MyGeometry, const FPointerEvent& MouseEvent) override;
@@ -582,7 +582,10 @@ public:
 		return 16.f;
 	}
 
-	// Zooms out to fit either all nodes or only the selected ones
+	/** 
+	 * Zooms out to fit either all nodes or only the selected ones.
+	 * @param bOnlySelection Whether to zoom to fit around only the current selection (if false, will zoom to the extents of all nodes)
+	 */
 	void ZoomToFit(bool bOnlySelection);
 
 	/** Get the bounding area for the currently selected nodes 
@@ -611,6 +614,13 @@ protected:
 	/** Initialize members */
 	void Construct();
 
+	/**
+	 * Zooms to the specified target rect
+	 * @param TopLeft The top left corner of the target rect
+	 * @param BottomRight The bottom right corner of the target rect
+	 */
+	void ZoomToTarget(const FVector2D& TopLeft, const FVector2D& BottomRight);
+
 	/** Update the new view offset location  */
 	void UpdateViewOffset(const FGeometry& MyGeometry, const FVector2D& TargetPosition);
 
@@ -795,16 +805,9 @@ protected:
 	/** Deferred zoom to node extents */
 	bool bDeferredZoomToNodeExtents;
 
-	/** Are we currently zooming to fit? */
-	bool bDeferredZoomingToFit;
-
 	/** Zoom selection padding */
 	float ZoomPadding;
 
-	/** Zoom target rectangle */
-	FVector2D ZoomTargetTopLeft;
-	FVector2D ZoomTargetBottomRight;
-
 	/** Allow continous zoom interpolation? */
 	bool bAllowContinousZoomInterpolation;
 
@@ -843,4 +846,19 @@ protected:
 
 	/** The current transaction for undo/redo */
 	TSharedPtr<FScopedTransaction> ScopedTransactionPtr;
+
+private:
+	/** Active tick that handles deferred zooming until the target zoom is reached */
+	EActiveTickReturnType HandleZoomToFit(double InCurrentTime, float InDeltaTime);
+
+private:
+	/** The handle to the active tick */
+	TWeakPtr<FActiveTickHandle> ActiveTickHandle;
+
+	/** Cached geometry for use within the active tick */
+	FGeometry CachedGeometry;
+
+	/** Zoom target rectangle */
+	FVector2D ZoomTargetTopLeft;
+	FVector2D ZoomTargetBottomRight;
 };
diff --git a/Engine/Source/Editor/IntroTutorials/Private/SEditorTutorials.cpp b/Engine/Source/Editor/IntroTutorials/Private/SEditorTutorials.cpp
--- a/Engine/Source/Editor/IntroTutorials/Private/SEditorTutorials.cpp
+++ b/Engine/Source/Editor/IntroTutorials/Private/SEditorTutorials.cpp
@@ -65,6 +65,7 @@ void SEditorTutorials::HandleCloseClicked()
 	OnCloseClicked.ExecuteIfBound();
 	OnTutorialClosed.ExecuteIfBound();
 
+	OverlayContent.Reset();
 	ContentBox->ClearChildren();
 	OnTutorialExited.ExecuteIfBound();
 }
diff --git a/Engine/Source/Editor/IntroTutorials/Private/STutorialButton.cpp b/Engine/Source/Editor/IntroTutorials/Private/STutorialButton.cpp
--- a/Engine/Source/Editor/IntroTutorials/Private/STutorialButton.cpp
+++ b/Engine/Source/Editor/IntroTutorials/Private/STutorialButton.cpp
@@ -25,11 +25,10 @@ void STutorialButton::Construct(const FArguments& InArgs)
 	bTutorialAvailable = false;
 	bTutorialCompleted = false;
 	bTutorialDismissed = false;
-	bDeferTutorialOpen = true;
 	AlertStartTime = 0.0f;
 
 	PulseAnimation.AddCurve(0.0f, TutorialButtonConstants::PulseAnimationLength, ECurveEaseFunction::Linear);
-	PulseAnimation.Play();
+	RegisterActiveTick( 0.f, FWidgetActiveTickDelegate::CreateSP( this, &STutorialButton::OpenTutorialPostConstruct ) );
 
 	ChildSlot
 	[
@@ -47,31 +46,32 @@ void STutorialButton::Construct(const FArguments& InArgs)
 	];
 }
 
-void STutorialButton::Tick(const FGeometry& AllottedGeometry, const double InCurrentTime, const float InDeltaTime)
+EActiveTickReturnType STutorialButton::OpenTutorialPostConstruct( double InCurrentTime, float InDeltaTime )
 {
-	if (bDeferTutorialOpen)
-	{
-		RefreshStatus();
+	// Begin playing the pulse animation on a loop
+	PulseAnimation.Play(this->AsShared(), true);
 
-		if (bTutorialAvailable && CachedAttractTutorial != nullptr && !bTutorialDismissed && !bTutorialCompleted)
-		{
-			// kick off the attract tutorial if the user hasn't dismissed it and hasn't completed it
-			FIntroTutorials& IntroTutorials = FModuleManager::GetModuleChecked<FIntroTutorials>(TEXT("IntroTutorials"));
-			const bool bRestart = true;
-			IntroTutorials.LaunchTutorial(CachedAttractTutorial, bRestart, ContextWindow);
-		}
+	RefreshStatus();
 
-		if(ShouldShowAlert())
-		{
-			AlertStartTime = FPlatformTime::Seconds();
-		}
+	if ( bTutorialAvailable && CachedAttractTutorial != nullptr && !bTutorialDismissed && !bTutorialCompleted )
+	{
+		// kick off the attract tutorial if the user hasn't dismissed it and hasn't completed it
+		FIntroTutorials& IntroTutorials = FModuleManager::GetModuleChecked<FIntroTutorials>( TEXT( "IntroTutorials" ) );
+		const bool bRestart = true;
+		IntroTutorials.LaunchTutorial( CachedAttractTutorial, bRestart, ContextWindow );
+	}
 
-		if(CachedLaunchTutorial != nullptr)
-		{
-			TutorialTitle = CachedLaunchTutorial->Title;
-		}
+	if ( ShouldShowAlert() )
+	{
+		AlertStartTime = FPlatformTime::Seconds();
+	}
+
+	if ( CachedLaunchTutorial != nullptr )
+	{
+		TutorialTitle = CachedLaunchTutorial->Title;
 	}
-	bDeferTutorialOpen = false;
+
+	return EActiveTickReturnType::StopTicking;
 }
 
 static void GetAnimationValues(float InAnimationProgress, float& OutAlphaFactor0, float& OutPulseFactor0, float& OutAlphaFactor1, float& OutPulseFactor1)
@@ -90,13 +90,13 @@ int32 STutorialButton::OnPaint(const FPaintArgs& Args, const FGeometry& Allotted
 {
 	LayerId = SCompoundWidget::OnPaint(Args, AllottedGeometry, MyClippingRect, OutDrawElements, LayerId, InWidgetStyle, bParentEnabled) + 1000;
 
-	if (ShouldShowAlert())
+	if ( PulseAnimation.IsPlaying() )
 	{
 		float AlphaFactor0 = 0.0f;
 		float AlphaFactor1 = 0.0f;
 		float PulseFactor0 = 0.0f;
 		float PulseFactor1 = 0.0f;
-		GetAnimationValues(PulseAnimation.GetLerpLooping(), AlphaFactor0, PulseFactor0, AlphaFactor1, PulseFactor1);
+		GetAnimationValues(PulseAnimation.GetLerp(), AlphaFactor0, PulseFactor0, AlphaFactor1, PulseFactor1);
 
 		const FSlateBrush* PulseBrush = FEditorStyle::Get().GetBrush(TEXT("TutorialLaunch.Circle"));
 		const FLinearColor PulseColor = FEditorStyle::Get().GetColor(TEXT("TutorialLaunch.Circle.Color"));
@@ -285,6 +285,15 @@ void STutorialButton::RefreshStatus()
 	bTutorialCompleted = (CachedLaunchTutorial != nullptr) && GetDefault<UTutorialStateSettings>()->HaveCompletedTutorial(CachedLaunchTutorial);
 	bTutorialDismissed = ((CachedAttractTutorial != nullptr) && GetDefault<UTutorialStateSettings>()->IsTutorialDismissed(CachedAttractTutorial)) ||
 						 ((CachedLaunchTutorial != nullptr) && GetDefault<UTutorialStateSettings>()->IsTutorialDismissed(CachedLaunchTutorial));
+
+	if (ShouldShowAlert())
+	{
+		PulseAnimation.Resume();
+	}
+	else
+	{
+		PulseAnimation.Pause();
+	}
 }
 
 void STutorialButton::HandleTutorialExited()
diff --git a/Engine/Source/Editor/IntroTutorials/Private/STutorialButton.h b/Engine/Source/Editor/IntroTutorials/Private/STutorialButton.h
--- a/Engine/Source/Editor/IntroTutorials/Private/STutorialButton.h
+++ b/Engine/Source/Editor/IntroTutorials/Private/STutorialButton.h
@@ -17,9 +17,11 @@ class STutorialButton : public SCompoundWidget
 
 	virtual int32 OnPaint(const FPaintArgs& Args, const FGeometry& AllottedGeometry, const FSlateRect& MyClippingRect, FSlateWindowElementList& OutDrawElements, int32 LayerId, const FWidgetStyle& InWidgetStyle, bool bParentEnabled) const override;
 	virtual FReply OnMouseButtonDown(const FGeometry& MyGeometry, const FPointerEvent& MouseEvent) override;
-	virtual void Tick(const FGeometry& AllottedGeometry, const double InCurrentTime, const float InDeltaTime) override;
-
+	
 private:
+	/** Opens the tutorial post-construct */
+	EActiveTickReturnType OpenTutorialPostConstruct( double InCurrentTime, float InDeltaTime );
+
 	/** Handle clicking the tutorial button */
 	FReply HandleButtonClicked();
 
@@ -48,9 +50,7 @@ private:
 	void HandleTutorialExited();
 
 private:
-	/** Flag to defer tutorial open until the first Tick() */
-	bool bDeferTutorialOpen;
-
+	
 	/** Whether we have a tutorial for this context */
 	bool bTutorialAvailable;
 
diff --git a/Engine/Source/Editor/IntroTutorials/Private/STutorialContent.cpp b/Engine/Source/Editor/IntroTutorials/Private/STutorialContent.cpp
--- a/Engine/Source/Editor/IntroTutorials/Private/STutorialContent.cpp
+++ b/Engine/Source/Editor/IntroTutorials/Private/STutorialContent.cpp
@@ -48,11 +48,14 @@ void STutorialContent::Construct(const FArguments& InArgs, UEditorTutorial* InTu
 
 	BorderIntroAnimation.AddCurve(0.0f, TutorialConstants::BorderIntroAnimationLength, ECurveEaseFunction::CubicOut);
 	BorderPulseAnimation.AddCurve(0.0f, TutorialConstants::BorderPulseAnimationLength, ECurveEaseFunction::Linear);
-	BorderIntroAnimation.Play();
-	BorderPulseAnimation.Play();
+	BorderIntroAnimation.Play(this->AsShared());
+	
+	// Set the border pulse to play on a loop and immediately pause it - will be resumed when needed
+	BorderPulseAnimation.Play(this->AsShared(), true);
+	BorderPulseAnimation.Pause();
 
 	ContentIntroAnimation.AddCurve(0.0f, TutorialConstants::ContentIntroAnimationLength, ECurveEaseFunction::Linear);
-	ContentIntroAnimation.Play();
+	ContentIntroAnimation.Play(this->AsShared());
 
 	if (InContent.Text.IsEmpty() == true)
 	{
@@ -184,25 +187,44 @@ void STutorialContent::Construct(const FArguments& InArgs, UEditorTutorial* InTu
 	];
 }
 
-static void GetAnimationValues(bool bIsIntro, float InAnimationProgress, float& OutAlphaFactor, float& OutPulseFactor, FLinearColor& OutShadowTint, FLinearColor& OutBorderTint)
-{
-	if(bIsIntro)
+//static void GetAnimationValues(bool bIsIntro, float InAnimationProgress, float& OutAlphaFactor, float& OutPulseFactor, FLinearColor& OutShadowTint, FLinearColor& OutBorderTint)
+//{
+//	if ( bIsIntro )
+//	{
+//		OutAlphaFactor = InAnimationProgress;
+//		OutPulseFactor = ( 1.0f - OutAlphaFactor ) * 50.0f;
+//		OutShadowTint = FLinearColor(1.0f, 1.0f, 0.0f, OutAlphaFactor);
+//		OutBorderTint = FLinearColor(1.0f, 1.0f, 0.0f, OutAlphaFactor * OutAlphaFactor);
+//	}
+//	else
+//	{
+//		OutAlphaFactor = 1.0f - ( 0.5f + ( FMath::Cos(2.0f * PI * InAnimationProgress) * 0.5f ) );
+//		OutPulseFactor = 0.5f + ( FMath::Cos(2.0f * PI * InAnimationProgress) * 0.5f );
+//		OutShadowTint = FLinearColor(1.0f, 1.0f, 0.0f, 1.0f);
+//		OutBorderTint = FLinearColor(1.0f, 1.0f, 0.0f, TutorialConstants::MinBorderOpacity + ( ( 1.0f - TutorialConstants::MinBorderOpacity ) * OutAlphaFactor ));
+//	}
+//}
+
+void STutorialContent::GetAnimationValues(float& OutAlphaFactor, float& OutPulseFactor, FLinearColor& OutShadowTint, FLinearColor& OutBorderTint) const
+{
+	if (BorderIntroAnimation.IsPlaying())
 	{
-		OutAlphaFactor = InAnimationProgress;
-		OutPulseFactor = (1.0f - OutAlphaFactor) * 50.0f;
+		OutAlphaFactor = BorderIntroAnimation.GetLerp();
+		OutPulseFactor = ( 1.0f - OutAlphaFactor ) * 50.0f;
 		OutShadowTint = FLinearColor(1.0f, 1.0f, 0.0f, OutAlphaFactor);
 		OutBorderTint = FLinearColor(1.0f, 1.0f, 0.0f, OutAlphaFactor * OutAlphaFactor);
 	}
 	else
 	{
-		OutAlphaFactor = 1.0f - (0.5f + (FMath::Cos(2.0f * PI * InAnimationProgress) * 0.5f));
-		OutPulseFactor = 0.5f + (FMath::Cos(2.0f * PI * InAnimationProgress) * 0.5f);
+		float PulseAnimationProgress = BorderPulseAnimation.GetLerp();
+		OutAlphaFactor = 1.0f - ( 0.5f + ( FMath::Cos(2.0f * PI * PulseAnimationProgress) * 0.5f ) );
+		OutPulseFactor = 0.5f + ( FMath::Cos(2.0f * PI * PulseAnimationProgress) * 0.5f );
 		OutShadowTint = FLinearColor(1.0f, 1.0f, 0.0f, 1.0f);
-		OutBorderTint = FLinearColor(1.0f, 1.0f, 0.0f, TutorialConstants::MinBorderOpacity + ((1.0f - TutorialConstants::MinBorderOpacity) * OutAlphaFactor));
+		OutBorderTint = FLinearColor(1.0f, 1.0f, 0.0f, TutorialConstants::MinBorderOpacity + ( ( 1.0f - TutorialConstants::MinBorderOpacity ) * OutAlphaFactor ));
 	}
 }
 
-int32 STutorialContent::OnPaint( const FPaintArgs& Args, const FGeometry& AllottedGeometry, const FSlateRect& MyClippingRect, FSlateWindowElementList& OutDrawElements, int32 LayerId, const FWidgetStyle& InWidgetStyle, bool bParentEnabled ) const
+int32 STutorialContent::OnPaint(const FPaintArgs& Args, const FGeometry& AllottedGeometry, const FSlateRect& MyClippingRect, FSlateWindowElementList& OutDrawElements, int32 LayerId, const FWidgetStyle& InWidgetStyle, bool bParentEnabled) const
 {
 	CachedContentGeometry = AllottedGeometry;
 	CachedContentGeometry.AppendTransform(FSlateLayoutTransform(OutDrawElements.GetWindow()->GetPositionInScreen()));
@@ -213,7 +235,7 @@ int32 STutorialContent::OnPaint( const FPaintArgs& Args, const FGeometry& Allott
 		float PulseFactor;
 		FLinearColor ShadowTint;
 		FLinearColor BorderTint;
-		GetAnimationValues(BorderIntroAnimation.IsPlaying(),  BorderIntroAnimation.IsPlaying() ? BorderIntroAnimation.GetLerp() : BorderPulseAnimation.GetLerpLooping(), AlphaFactor, PulseFactor, ShadowTint, BorderTint);
+		GetAnimationValues(AlphaFactor, PulseFactor, ShadowTint, BorderTint);
 
 		const FSlateBrush* ShadowBrush = FCoreStyle::Get().GetBrush(TEXT("Tutorials.Shadow"));
 		const FSlateBrush* BorderBrush = FCoreStyle::Get().GetBrush(TEXT("Tutorials.Border"));
@@ -466,6 +488,11 @@ void STutorialContent::HandlePaintNamedWidget(TSharedRef<SWidget> InWidget, cons
 			{
 				bIsVisible = true;
 				CachedGeometry = InGeometry;
+
+				if (!BorderPulseAnimation.IsPlaying() && Anchor.bDrawHighlight)
+				{
+					BorderPulseAnimation.Resume();
+				}
 			}
 		}
 		break;
@@ -474,6 +501,7 @@ void STutorialContent::HandlePaintNamedWidget(TSharedRef<SWidget> InWidget, cons
 
 void STutorialContent::HandleResetNamedWidget()
 {
+	BorderPulseAnimation.Pause();
 	bIsVisible = false;
 }
 
diff --git a/Engine/Source/Editor/IntroTutorials/Private/STutorialContent.h b/Engine/Source/Editor/IntroTutorials/Private/STutorialContent.h
--- a/Engine/Source/Editor/IntroTutorials/Private/STutorialContent.h
+++ b/Engine/Source/Editor/IntroTutorials/Private/STutorialContent.h
@@ -71,10 +71,10 @@ class STutorialContent : public SCompoundWidget
 	void Construct(const FArguments& InArgs, UEditorTutorial* InTutorial, const FTutorialContent& InContent);
 
 	/** SWidget implementation */
-	virtual int32 OnPaint( const FPaintArgs& Args, const FGeometry& AllottedGeometry, const FSlateRect& MyClippingRect, FSlateWindowElementList& OutDrawElements, int32 LayerId, const FWidgetStyle& InWidgetStyle, bool bParentEnabled ) const override;
-	virtual FReply OnMouseButtonUp( const FGeometry& MyGeometry, const FPointerEvent& MouseEvent ) override;
-	virtual FReply OnMouseButtonDown( const FGeometry& MyGeometry, const FPointerEvent& MouseEvent ) override;
-	virtual FReply OnMouseButtonDoubleClick( const FGeometry& InMyGeometry, const FPointerEvent& InMouseEvent ) override;
+	virtual int32 OnPaint(const FPaintArgs& Args, const FGeometry& AllottedGeometry, const FSlateRect& MyClippingRect, FSlateWindowElementList& OutDrawElements, int32 LayerId, const FWidgetStyle& InWidgetStyle, bool bParentEnabled) const override;
+	virtual FReply OnMouseButtonUp(const FGeometry& MyGeometry, const FPointerEvent& MouseEvent) override;
+	virtual FReply OnMouseButtonDown(const FGeometry& MyGeometry, const FPointerEvent& MouseEvent) override;
+	virtual FReply OnMouseButtonDoubleClick(const FGeometry& InMyGeometry, const FPointerEvent& InMouseEvent) override;
 
 	/** Helper function to generate widgets from an FTutorialContent struct */
 	static TSharedRef<SWidget> GenerateContentWidget(const FTutorialContent& InContent, TSharedPtr<IDocumentationPage>& OutDocumentationPage, const TAttribute<FText>& InHighlightText = TAttribute<FText>(), bool bAutoWrapText = true, float WrapTextAt = 0.0f);
@@ -150,6 +150,9 @@ private:
 	/** We need to override the border ourselves, rather than let the button handle it, as we are using a larger apparent hitbox */
 	const FSlateBrush* GetNextButtonBorder() const;
 
+	/** Helper to determine the proper animation values for the border pulse */
+	void GetAnimationValues(float& OutAlphaFactor, float& OutPulseFactor, FLinearColor& OutShadowTint, FLinearColor& OutBorderTint) const;
+
 private:
 
 	/** Copy of the window size we were last draw at */
diff --git a/Engine/Source/Editor/IntroTutorials/Private/STutorialRoot.cpp b/Engine/Source/Editor/IntroTutorials/Private/STutorialRoot.cpp
--- a/Engine/Source/Editor/IntroTutorials/Private/STutorialRoot.cpp
+++ b/Engine/Source/Editor/IntroTutorials/Private/STutorialRoot.cpp
@@ -33,7 +33,7 @@ void STutorialRoot::MaybeAddOverlay(TSharedRef<SWindow> InWindow)
 {
 	if(InWindow->HasOverlay())
 	{
-		// check we dont already have a widget overlay for this window
+		// check we don't already have a widget overlay for this window
 		TWeakPtr<SEditorTutorials>* FoundWidget = TutorialWidgets.Find(InWindow);
 		if(FoundWidget == nullptr)
 		{
@@ -72,7 +72,7 @@ void STutorialRoot::MaybeAddOverlay(TSharedRef<SWindow> InWindow)
 	}
 }
 
-void STutorialRoot::Tick(const FGeometry& AllottedGeometry, const double InCurrentTime, const float InDeltaTime)
+void STutorialRoot::Tick( const FGeometry& AllottedGeometry, const double InCurrentTime, const float InDeltaTime )
 {
 	TArray<TSharedRef<SWindow>> Windows = FSlateApplication::Get().GetInteractiveTopLevelWindows();
 	for(auto& Window : Windows)
diff --git a/Engine/Source/Editor/IntroTutorials/Private/STutorialsBrowser.cpp b/Engine/Source/Editor/IntroTutorials/Private/STutorialsBrowser.cpp
--- a/Engine/Source/Editor/IntroTutorials/Private/STutorialsBrowser.cpp
+++ b/Engine/Source/Editor/IntroTutorials/Private/STutorialsBrowser.cpp
@@ -21,7 +21,7 @@ DECLARE_DELEGATE_OneParam(FOnCategorySelected, const FString& /* InCategory */);
 
 namespace TutorialBrowserConstants
 {
-	const float RefreshTimerInterval = 0.25f;
+	const float RefreshTimerInterval = 1.0f;
 
 	const float ProgressUpdateInterval = 0.5f;
 }
@@ -497,6 +497,8 @@ void STutorialsBrowser::Construct(const FArguments& InArgs)
 	FAssetRegistryModule& AssetRegistryModule = FModuleManager::LoadModuleChecked<FAssetRegistryModule>(TEXT("AssetRegistry"));
 	AssetRegistryModule.Get().OnAssetAdded().AddSP(this, &STutorialsBrowser::HandleAssetAdded);
 
+	RegisterActiveTick( TutorialBrowserConstants::RefreshTimerInterval, FWidgetActiveTickDelegate::CreateSP( this, &STutorialsBrowser::TriggerReloadTutorials ) );
+
 	ChildSlot
 	[
 		SNew(SVerticalBox)
@@ -577,16 +579,15 @@ void STutorialsBrowser::Construct(const FArguments& InArgs)
 	RebuildCrumbs();
 }
 
-void STutorialsBrowser::Tick( const FGeometry& AllottedGeometry, const double InCurrentTime, const float InDeltaTime )
+EActiveTickReturnType STutorialsBrowser::TriggerReloadTutorials( double InCurrentTime, float InDeltaTime )
 {
-	RefreshTimer -= InDeltaTime;
-
-	if(bNeedsRefresh && RefreshTimer <= 0.0f)
+	if (bNeedsRefresh)
 	{
-		ReloadTutorials();
 		bNeedsRefresh = false;
-		RefreshTimer = TutorialBrowserConstants::RefreshTimerInterval;
+		ReloadTutorials();
 	}
+	
+	return EActiveTickReturnType::KeepTicking;
 }
 
 void STutorialsBrowser::SetFilter(const FString& InFilter)
diff --git a/Engine/Source/Editor/IntroTutorials/Private/STutorialsBrowser.h b/Engine/Source/Editor/IntroTutorials/Private/STutorialsBrowser.h
--- a/Engine/Source/Editor/IntroTutorials/Private/STutorialsBrowser.h
+++ b/Engine/Source/Editor/IntroTutorials/Private/STutorialsBrowser.h
@@ -45,15 +45,12 @@ class STutorialsBrowser : public SCompoundWidget
 
 	void Construct(const FArguments& Args);
 
-	/** Set the current filter string. Filters are used to only show the specifed category of tutorials (e.g. only Blueprint tutorials) */
+	/** Set the current filter string. Filters are used to only show the specified category of tutorials (e.g. only Blueprint tutorials) */
 	void SetFilter(const FString& InFilter);
 
 	/** Reload all tutorials that we know about */
 	void ReloadTutorials();
 
-	/** SWidget implementation */
-	virtual void Tick( const FGeometry& AllottedGeometry, const double InCurrentTime, const float InDeltaTime ) override;
-
 protected:
 	/** Handle generating a table row in the browser */
 	TSharedRef<ITableRow> OnGenerateTutorialRow(TSharedPtr<ITutorialListEntry> InItem, const TSharedRef<STableViewBase>& OwnerTable) const;
@@ -102,6 +99,9 @@ protected:
 
 private:
 
+	/** Triggers a reload of the tutorials */
+	EActiveTickReturnType TriggerReloadTutorials( double InCurrentTime, float InDeltaTime );
+
 	/** Root entry of the tutorials tree */
 	TSharedPtr<FTutorialListEntry_Category> RootEntry;
 
diff --git a/Engine/Source/Editor/Kismet/Private/Debugging/SKismetDebuggingView.cpp b/Engine/Source/Editor/Kismet/Private/Debugging/SKismetDebuggingView.cpp
--- a/Engine/Source/Editor/Kismet/Private/Debugging/SKismetDebuggingView.cpp
+++ b/Engine/Source/Editor/Kismet/Private/Debugging/SKismetDebuggingView.cpp
@@ -1158,9 +1158,6 @@ void SKismetDebuggingView::Tick( const FGeometry& AllottedGeometry, const double
 
 	// Refresh the list
 	DebugTreeView->RequestTreeRefresh();
-
-
-	SCompoundWidget::Tick( AllottedGeometry, InCurrentTime, InDeltaTime );
 }
 
 
diff --git a/Engine/Source/Editor/Kismet/Private/FindInBlueprints.cpp b/Engine/Source/Editor/Kismet/Private/FindInBlueprints.cpp
--- a/Engine/Source/Editor/Kismet/Private/FindInBlueprints.cpp
+++ b/Engine/Source/Editor/Kismet/Private/FindInBlueprints.cpp
@@ -1092,35 +1092,33 @@ SFindInBlueprints::~SFindInBlueprints()
 	}
 }
 
-void SFindInBlueprints::Tick( const FGeometry& AllottedGeometry, const double InCurrentTime, const float InDeltaTime )
+EActiveTickReturnType SFindInBlueprints::UpdateSearchResults( double InCurrentTime, float InDeltaTime )
 {
-	SCompoundWidget::Tick(AllottedGeometry, InCurrentTime, InDeltaTime);
-
-	if(StreamSearch.IsValid())
+	if ( StreamSearch.IsValid() )
 	{
 		bool bShouldShutdownThread = false;
 		bShouldShutdownThread = StreamSearch->IsComplete();
 
 		TArray<FSearchResult> BackgroundItemsFound;
 
-		StreamSearch->GetFilteredItems(BackgroundItemsFound);
-		if(BackgroundItemsFound.Num())
+		StreamSearch->GetFilteredItems( BackgroundItemsFound );
+		if ( BackgroundItemsFound.Num() )
 		{
-			for(auto& Item : BackgroundItemsFound)
+			for ( auto& Item : BackgroundItemsFound )
 			{
-				Item->ExpandAllChildren(TreeView);
-				ItemsFound.Add(Item);
+				Item->ExpandAllChildren( TreeView );
+				ItemsFound.Add( Item );
 			}
 			TreeView->RequestTreeRefresh();
 		}
 
 		// If the thread is complete, shut it down properly
-		if(bShouldShutdownThread)
+		if ( bShouldShutdownThread )
 		{
-			if(ItemsFound.Num() == 0)
+			if ( ItemsFound.Num() == 0 )
 			{
 				// Insert a fake result to inform user if none found
-				ItemsFound.Add(FSearchResult(new FFindInBlueprintsResult(LOCTEXT("BlueprintSearchNoResults", "No Results found"))));
+				ItemsFound.Add( FSearchResult( new FFindInBlueprintsResult( LOCTEXT( "BlueprintSearchNoResults", "No Results found" ) ) ) );
 				TreeView->RequestTreeRefresh();
 			}
 
@@ -1128,6 +1126,8 @@ void SFindInBlueprints::Tick( const FGeometry& AllottedGeometry, const double In
 			StreamSearch.Reset();
 		}
 	}
+
+	return StreamSearch.IsValid() ? EActiveTickReturnType::KeepTicking : EActiveTickReturnType::StopTicking;
 }
 
 void SFindInBlueprints::RegisterCommands()
@@ -1300,6 +1300,11 @@ void SFindInBlueprints::LaunchStreamThread(const TArray<FString>& InTokens)
 		StreamSearch->Stop();
 		StreamSearch->EnsureCompletion();
 	}
+	else
+	{
+		// If the stream search wasn't already running, register the active tick
+		RegisterActiveTick( 0.f, FWidgetActiveTickDelegate::CreateSP( this, &SFindInBlueprints::UpdateSearchResults ) );
+	}
 
 	StreamSearch = MakeShareable(new FStreamSearch(InTokens));
 }
diff --git a/Engine/Source/Editor/Kismet/Private/SBlueprintSubPalette.cpp b/Engine/Source/Editor/Kismet/Private/SBlueprintSubPalette.cpp
--- a/Engine/Source/Editor/Kismet/Private/SBlueprintSubPalette.cpp
+++ b/Engine/Source/Editor/Kismet/Private/SBlueprintSubPalette.cpp
@@ -256,15 +256,10 @@ void SBlueprintSubPalette::Construct(FArguments const& InArgs, TWeakPtr<FBluepri
 }
 
 //------------------------------------------------------------------------------
-void SBlueprintSubPalette::Tick(const FGeometry& AllottedGeometry, const double InCurrentTime, const float InDeltaTime)
+EActiveTickReturnType SBlueprintSubPalette::TriggerRefreshActionsList(double InCurrentTime, float InDeltaTime)
 {
-	SGraphPalette::Tick(AllottedGeometry, InCurrentTime, InDeltaTime);
-
-	if(bNeedsRefresh)
-	{
-		bNeedsRefresh = false;
-		RefreshActionsList(/*bPreserveExpansion =*/true);
-	}
+	RefreshActionsList(true);
+	return EActiveTickReturnType::StopTicking;
 }
 
 //------------------------------------------------------------------------------
@@ -395,7 +390,7 @@ void SBlueprintSubPalette::GenerateContextMenuEntries(FMenuBuilder& MenuBuilder)
 //------------------------------------------------------------------------------
 void SBlueprintSubPalette::RequestRefreshActionsList()
 {
-	bNeedsRefresh = true;
+	RegisterActiveTick(0.f, FWidgetActiveTickDelegate::CreateSP(this, &SBlueprintSubPalette::TriggerRefreshActionsList));
 }
 
 //------------------------------------------------------------------------------
diff --git a/Engine/Source/Editor/Kismet/Private/SBlueprintSubPalette.h b/Engine/Source/Editor/Kismet/Private/SBlueprintSubPalette.h
--- a/Engine/Source/Editor/Kismet/Private/SBlueprintSubPalette.h
+++ b/Engine/Source/Editor/Kismet/Private/SBlueprintSubPalette.h
@@ -34,15 +34,6 @@ public:
 	 */
 	void Construct(const FArguments& InArgs, TWeakPtr<FBlueprintEditor> InBlueprintEditor);
 
-	/**
-	 * Ticks this widget.
-	 *
-	 * @param  AllottedGeometry The space allotted for this widget
-	 * @param  InCurrentTime  Current absolute real time
-	 * @param  InDeltaTime  Real time passed since last tick
-	 */
-	virtual void Tick( const FGeometry& AllottedGeometry, const double InCurrentTime, const float InDeltaTime ) override;
-
 	/**
 	 * Retrieves, from the owning blueprint-editor, the blueprint currently  
 	 * being worked on.
@@ -106,6 +97,9 @@ private:
 	/** Immediately calls RefreshActionsList(), doesn't not defer until Tick() like RequestRefreshActionsList() does. */
 	void ForceRefreshActionList();
 
+	/** One-off active tick to trigger a refresh of the action list */
+	EActiveTickReturnType TriggerRefreshActionsList(double InCurrentTime, float InDeltaTime);
+
 	/**
 	 * Constructs a slate widget for the right-click context menu in this 
 	 * palette. While this isn't virtual, sub-classes can override GenerateContextMenuEntries()
@@ -117,8 +111,5 @@ private:
 
 	/** Pointer to the command list created for this (so multiple sub-palettes can have their own bindings)*/
 	TSharedPtr<FUICommandList> CommandList;
-
-	/** Flag to indicate whether or not the action list needs to be refreshed */
-	bool bNeedsRefresh;
 };
 
diff --git a/Engine/Source/Editor/Kismet/Private/SCSEditorViewportClient.cpp b/Engine/Source/Editor/Kismet/Private/SCSEditorViewportClient.cpp
--- a/Engine/Source/Editor/Kismet/Private/SCSEditorViewportClient.cpp
+++ b/Engine/Source/Editor/Kismet/Private/SCSEditorViewportClient.cpp
@@ -182,6 +182,10 @@ void FSCSEditorViewportClient::Tick(float DeltaSeconds)
 			}
 		}
 	}
+	else
+	{
+		InvalidatePreview(false);
+	}
 
 	// Tick the preview scene world.
 	if (!GIntraFrameDebuggingGameThread)
diff --git a/Engine/Source/Editor/Kismet/Private/SSCSEditorViewport.cpp b/Engine/Source/Editor/Kismet/Private/SSCSEditorViewport.cpp
--- a/Engine/Source/Editor/Kismet/Private/SSCSEditorViewport.cpp
+++ b/Engine/Source/Editor/Kismet/Private/SSCSEditorViewport.cpp
@@ -187,9 +187,7 @@ private:
 
 void SSCSEditorViewport::Construct(const FArguments& InArgs)
 {
-	// Initialize
-	bPreviewNeedsUpdating = false;
-	bResetCameraOnNextPreviewUpdate = false;
+	bIsActiveTickRegistered = false;
 
 	// Save off the Blueprint editor reference, we'll need this later
 	BlueprintEditorPtr = InArgs._BlueprintEditor;
@@ -306,10 +304,10 @@ void SSCSEditorViewport::RequestRefresh(bool bResetCamera, bool bRefreshNow)
 	else
 	{
 		// Defer the update until the next tick. This way we don't accidentally spawn the preview actor in the middle of a transaction, for example.
-		bPreviewNeedsUpdating = true;
-		if(bResetCamera)
+		if (!bIsActiveTickRegistered)
 		{
-			bResetCameraOnNextPreviewUpdate = true;
+			bIsActiveTickRegistered = true;
+			RegisterActiveTick(0.f, FWidgetActiveTickDelegate::CreateSP(this, &SSCSEditorViewport::DeferredUpdatePreview, bResetCamera));
 		}
 	}
 }
@@ -330,19 +328,15 @@ bool SSCSEditorViewport::GetIsSimulateEnabled()
 	return ViewportClient->GetIsSimulateEnabled();
 }
 
-void SSCSEditorViewport::Tick(const FGeometry& AllottedGeometry, const double InCurrentTime, const float InDeltaTime)
+EActiveTickReturnType SSCSEditorViewport::DeferredUpdatePreview(double InCurrentTime, float InDeltaTime, bool bResetCamera)
 {
-	SEditorViewport::Tick(AllottedGeometry, InCurrentTime, InDeltaTime);
-
-	// If the preview scene is no longer valid (i.e. all actors have destroyed themselves), then attempt to recreate the scene. This way we can "loop" certain "finite" Blueprints that might destroy themselves.
-	if(ViewportClient.IsValid() && (bPreviewNeedsUpdating || !ViewportClient->IsPreviewSceneValid()))
+	if (ViewportClient.IsValid())
 	{
-		ViewportClient->InvalidatePreview(bResetCameraOnNextPreviewUpdate);
-
-		// Reset for next update
-		bPreviewNeedsUpdating = false;
-		bResetCameraOnNextPreviewUpdate = false;
+		ViewportClient->InvalidatePreview(bResetCamera);
 	}
+
+	bIsActiveTickRegistered = false;
+	return EActiveTickReturnType::StopTicking;
 }
 
 AActor* SSCSEditorViewport::GetPreviewActor() const
diff --git a/Engine/Source/Editor/Kismet/Private/SSCSEditorViewport.h b/Engine/Source/Editor/Kismet/Private/SSCSEditorViewport.h
--- a/Engine/Source/Editor/Kismet/Private/SSCSEditorViewport.h
+++ b/Engine/Source/Editor/Kismet/Private/SSCSEditorViewport.h
@@ -47,19 +47,6 @@ public:
 	 */
 	void RequestRefresh(bool bResetCamera = false, bool bRefreshNow = false);
 
-	// SWidget interface
-
-	/**
-	 * Ticks this widget.
-	 *
-	 * @param  AllottedGeometry The space allotted for this widget
-	 * @param  InCurrentTime  Current absolute real time
-	 * @param  InDeltaTime  Real time passed since last tick
-	 */
-	virtual void Tick( const FGeometry& AllottedGeometry, const double InCurrentTime, const float InDeltaTime ) override;
-	
-	// End of SWidget interface
-
 	/**
 	 * Called when the selected component changes in the SCS editor.
 	 */
@@ -101,6 +88,11 @@ protected:
 	virtual TSharedRef<class FEditorViewportClient> MakeEditorViewportClient() override;
 	virtual TSharedPtr<class SWidget> MakeViewportToolbar() override;
 	virtual void BindCommands();
+
+private:
+	/** One-off active tick to update the preview */
+	EActiveTickReturnType DeferredUpdatePreview(double InCurrentTime, float InDeltaTime, bool bResetCamera);
+
 private:
 	/** Pointer back to editor tool (owner) */
 	TWeakPtr<class FBlueprintEditor> BlueprintEditorPtr;
@@ -111,10 +103,6 @@ private:
 	/** Viewport client */
 	TSharedPtr<class FSCSEditorViewportClient> ViewportClient;
 
-	/** If true, update the preview scene */
-	bool bPreviewNeedsUpdating;
-
-	/** If true, reset the camera on the next preview scene update */
-	bool bResetCameraOnNextPreviewUpdate;
-
+	/** Whether the active tick (for updating the preview) is registered */
+	bool bIsActiveTickRegistered;
 };
diff --git a/Engine/Source/Editor/Kismet/Public/FindInBlueprints.h b/Engine/Source/Editor/Kismet/Public/FindInBlueprints.h
--- a/Engine/Source/Editor/Kismet/Public/FindInBlueprints.h
+++ b/Engine/Source/Editor/Kismet/Public/FindInBlueprints.h
@@ -226,14 +226,13 @@ public:
 	void Construct(const FArguments& InArgs, TSharedPtr<class FBlueprintEditor> InBlueprintEditor);
 	~SFindInBlueprints();
 
-	/** Begin SWidget Interface */
-	virtual void Tick( const FGeometry& AllottedGeometry, const double InCurrentTime, const float InDeltaTime ) override;
-	/** End SWidget Interface */
-
 	/** Focuses this widget's search box, and changes the mode as well, and optionally the search terms */
 	void FocusForUse(bool bSetFindWithinBlueprint, FString NewSearchTerms = FString(), bool bSelectFirstResult = false);
 
+
 private:
+	/** Processes results of the ongoing async stream search */
+	EActiveTickReturnType UpdateSearchResults( double InCurrentTime, float InDeltaTime );
 
 	/** Register any Find-in-Blueprint commands */
 	void RegisterCommands();
diff --git a/Engine/Source/Editor/KismetWidgets/Private/SLevelOfDetailBranchNode.cpp b/Engine/Source/Editor/KismetWidgets/Private/SLevelOfDetailBranchNode.cpp
--- a/Engine/Source/Editor/KismetWidgets/Private/SLevelOfDetailBranchNode.cpp
+++ b/Engine/Source/Editor/KismetWidgets/Private/SLevelOfDetailBranchNode.cpp
@@ -21,7 +21,7 @@ void SLevelOfDetailBranchNode::Construct(const FArguments& InArgs)
 	ChildSlot[ChildSlotHighDetail];
 }
 
-void SLevelOfDetailBranchNode::Tick(const FGeometry& AllottedGeometry, const double InCurrentTime, const float InDeltaTime)
+void SLevelOfDetailBranchNode::Tick( const FGeometry& AllottedGeometry, const double InCurrentTime, const float InDeltaTime )
 {
 	const int32 CurrentValue = ShowLowDetailAttr.Get() ? 1 : 0;
 	if (CurrentValue != LastCachedValue)
@@ -32,6 +32,4 @@ void SLevelOfDetailBranchNode::Tick(const FGeometry& AllottedGeometry, const dou
 			(CurrentValue != 0) ? ChildSlotLowDetail : ChildSlotHighDetail
 		];
 	}
-
-	SCompoundWidget::Tick(AllottedGeometry, InCurrentTime, InDeltaTime);
 }
diff --git a/Engine/Source/Editor/KismetWidgets/Private/SScrubControlPanel.cpp b/Engine/Source/Editor/KismetWidgets/Private/SScrubControlPanel.cpp
--- a/Engine/Source/Editor/KismetWidgets/Private/SScrubControlPanel.cpp
+++ b/Engine/Source/Editor/KismetWidgets/Private/SScrubControlPanel.cpp
@@ -26,6 +26,7 @@ void SScrubControlPanel::Construct( const SScrubControlPanel::FArguments& InArgs
 	TransportControlArgs.OnToggleLooping = InArgs._OnClickedToggleLoop;
 	TransportControlArgs.OnGetLooping = InArgs._OnGetLooping;
 	TransportControlArgs.OnGetPlaybackMode = InArgs._OnGetPlaybackMode;
+	TransportControlArgs.OnTickPlayback = InArgs._OnTickPlayback;
 	
 	FTransportControlArgs TransportControlArgsForRealtimeStreamingMode;
 	TransportControlArgsForRealtimeStreamingMode.OnForwardPlay = TransportControlArgs.OnForwardPlay;
diff --git a/Engine/Source/Editor/KismetWidgets/Private/SScrubWidget.cpp b/Engine/Source/Editor/KismetWidgets/Private/SScrubWidget.cpp
--- a/Engine/Source/Editor/KismetWidgets/Private/SScrubWidget.cpp
+++ b/Engine/Source/Editor/KismetWidgets/Private/SScrubWidget.cpp
@@ -332,7 +332,7 @@ FReply SScrubWidget::OnMouseMove( const FGeometry& MyGeometry, const FPointerEve
 		else if (!bDragging)
 		{
 			DistanceDragged += FMath::Abs(MouseEvent.GetCursorDelta().X);
-			if ( DistanceDragged > FSlateApplication::Get().GetDragTriggerDistnace() )
+			if ( DistanceDragged > FSlateApplication::Get().GetDragTriggerDistance() )
 			{
 				bDragging = true;
 			}
diff --git a/Engine/Source/Editor/KismetWidgets/Private/SSingleObjectDetailsPanel.cpp b/Engine/Source/Editor/KismetWidgets/Private/SSingleObjectDetailsPanel.cpp
--- a/Engine/Source/Editor/KismetWidgets/Private/SSingleObjectDetailsPanel.cpp
+++ b/Engine/Source/Editor/KismetWidgets/Private/SSingleObjectDetailsPanel.cpp
@@ -39,7 +39,7 @@ UObject* SSingleObjectDetailsPanel::GetObjectToObserve() const
 	return NULL;
 }
 
-void SSingleObjectDetailsPanel::Tick(const FGeometry& AllottedGeometry, const double InCurrentTime, const float InDeltaTime)
+void SSingleObjectDetailsPanel::Tick( const FGeometry& AllottedGeometry, const double InCurrentTime, const float InDeltaTime )
 {
 	if (bAutoObserveObject)
 	{
@@ -57,8 +57,6 @@ void SSingleObjectDetailsPanel::Tick(const FGeometry& AllottedGeometry, const do
 			SetPropertyWindowContents(SelectedObjects);
 		}
 	}
-
-	SCompoundWidget::Tick(AllottedGeometry, InCurrentTime, InDeltaTime);
 }
 
 void SSingleObjectDetailsPanel::SetPropertyWindowContents(TArray<UObject*> Objects)
diff --git a/Engine/Source/Editor/KismetWidgets/Public/SScrubControlPanel.h b/Engine/Source/Editor/KismetWidgets/Public/SScrubControlPanel.h
--- a/Engine/Source/Editor/KismetWidgets/Public/SScrubControlPanel.h
+++ b/Engine/Source/Editor/KismetWidgets/Public/SScrubControlPanel.h
@@ -9,9 +9,8 @@
 class KISMETWIDGETS_API SScrubControlPanel : public SCompoundWidget
 {
 public:
-
-	SLATE_BEGIN_ARGS(SScrubControlPanel)
-		: _Value(0)
+	SLATE_BEGIN_ARGS( SScrubControlPanel )
+		: _Value( 0 )
 		, _NumOfKeys()
 		, _SequenceLength()
 		, _OnValueChanged()
@@ -31,6 +30,7 @@ public:
 		, _ViewInputMax()
 		, _OnSetInputViewRange()
 		, _OnCropAnimSequence()
+		, _OnTickPlayback()
 	{}
 		SLATE_ATTRIBUTE( float, Value )
 		SLATE_ATTRIBUTE( uint32, NumOfKeys )
@@ -66,7 +66,9 @@ public:
 		SLATE_ATTRIBUTE( bool, IsRealtimeStreamingMode )
 		/** Optional, additional values to draw on the timeline **/
 		SLATE_ATTRIBUTE( TArray<float>, DraggableBars )
-		SLATE_EVENT( FOnScrubBarDrag, OnBarDrag)
+		SLATE_EVENT( FOnScrubBarDrag, OnBarDrag )
+		/** Called each frame during playback */
+		SLATE_EVENT( FOnTickPlayback, OnTickPlayback )
 	SLATE_END_ARGS()
 
 	/**
@@ -75,6 +77,9 @@ public:
 	 * @param InArgs   A declaration from which to construct the widget
 	 */
 	void Construct( const FArguments& InArgs );
+
+	virtual ~SScrubControlPanel() {}
+
 protected:
 	EVisibility GetRealtimeControlVisibility(bool bIsControlForRealtimeStreamingMode) const;
 
diff --git a/Engine/Source/Editor/LevelEditor/Private/LevelViewportLayout.cpp b/Engine/Source/Editor/LevelEditor/Private/LevelViewportLayout.cpp
--- a/Engine/Source/Editor/LevelEditor/Private/LevelViewportLayout.cpp
+++ b/Engine/Source/Editor/LevelEditor/Private/LevelViewportLayout.cpp
@@ -71,6 +71,9 @@ public:
 
 private:
 	
+	/** Ticks the level viewport tab */
+	EActiveTickReturnType ActiveTick( double InCurrentTime, float InDeltaTime );
+
 	/** Reference to the owning level viewport tab */
 	TSharedPtr<FLevelViewportTabContent> LevelViewportTab;
 
@@ -86,6 +89,9 @@ void SViewportsOverlay::Construct( const FArguments& InArgs )
 {
 	const TSharedRef<SWidget>& ContentWidget = InArgs._Content.Widget;
 	LevelViewportTab = InArgs._LevelViewportTab;
+
+	//RegisterActiveTick( 0.f, FTickWidgetDelegate::CreateSP( this, &SViewportsOverlay::ActiveTick ) );
+
 	ChildSlot
 		[
 			SAssignNew( OverlayWidget, SOverlay )
@@ -96,6 +102,12 @@ void SViewportsOverlay::Construct( const FArguments& InArgs )
 		];
 }
 
+EActiveTickReturnType SViewportsOverlay::ActiveTick( double InCurrentTime, float InDeltaTime )
+{
+	// Exists to ensure slate is ticked
+	return EActiveTickReturnType::KeepTicking;
+}
+
 void SViewportsOverlay::Tick( const FGeometry& AllottedGeometry, const double InCurrentTime, const float InDeltaTime )
 {
 	CachedSize = AllottedGeometry.Size;
@@ -354,6 +366,7 @@ void FLevelViewportLayout::MaximizeViewport( TSharedRef<SLevelViewport> Viewport
 		// Update state
 		bWasMaximized = bIsMaximized;
 		bWasImmersive = bIsImmersive;
+
 		bIsMaximized = bWantMaximize;
 		bIsImmersive = bWantImmersive;
 
@@ -374,7 +387,7 @@ void FLevelViewportLayout::MaximizeViewport( TSharedRef<SLevelViewport> Viewport
 			// the current state of bIsMaximized, we'll transition to either a maximized state or a "restored" state
 			MaximizeAnimation = FCurveSequence();
 			MaximizeAnimation.AddCurve( 0.0f, ViewportLayoutDefs::RestoreTransitionTime, ECurveEaseFunction::CubicIn );
-			MaximizeAnimation.PlayReverse();
+			MaximizeAnimation.PlayReverse( ViewportsOverlayWidget->AsShared() );
 			
 			if( bWasImmersive && !bIsImmersive )
 			{
@@ -427,7 +440,7 @@ void FLevelViewportLayout::MaximizeViewport( TSharedRef<SLevelViewport> Viewport
 			// Play the "maximize" transition
 			MaximizeAnimation = FCurveSequence();
 			MaximizeAnimation.AddCurve( 0.0f, ViewportLayoutDefs::MaximizeTransitionTime, ECurveEaseFunction::CubicOut );
-			MaximizeAnimation.Play();
+			MaximizeAnimation.Play( ViewportsOverlayWidget->AsShared() );
 		}
 
 
@@ -568,7 +581,7 @@ bool FLevelViewportLayout::IsViewportImmersive( const SLevelViewport& InViewport
 EVisibility FLevelViewportLayout::OnGetNonMaximizedVisibility() const
 {
 	// The non-maximized viewports are not visible if there is a maximized viewport on top of them
-	return (!bIsQueryingLayoutMetrics && MaximizedViewport.IsValid() && !bIsTransitioning && DeferredMaximizeCommands.Num() == 0) ? EVisibility::Collapsed : EVisibility::Visible;
+	return ( !bIsQueryingLayoutMetrics && MaximizedViewport.IsValid() && !bIsTransitioning && DeferredMaximizeCommands.Num() == 0 ) ? EVisibility::Collapsed : EVisibility::Visible;
 }
 
 
@@ -698,5 +711,5 @@ void FLevelViewportLayout::Tick( float DeltaTime )
 
 bool FLevelViewportLayout::IsTickable() const
 {
-	return DeferredMaximizeCommands.Num() > 0 || (bIsTransitioning && !MaximizeAnimation.IsPlaying());
+	return DeferredMaximizeCommands.Num() > 0 || ( bIsTransitioning && !MaximizeAnimation.IsPlaying() );
 }
diff --git a/Engine/Source/Editor/LevelEditor/Private/LevelViewportLayout.h b/Engine/Source/Editor/LevelEditor/Private/LevelViewportLayout.h
--- a/Engine/Source/Editor/LevelEditor/Private/LevelViewportLayout.h
+++ b/Engine/Source/Editor/LevelEditor/Private/LevelViewportLayout.h
@@ -13,7 +13,7 @@
  * Base class for level viewport layout configurations
  * Handles maximizing and restoring well as visibility of specific viewports.
  */
-class FLevelViewportLayout : public TSharedFromThis<FLevelViewportLayout>, FTickableEditorObject
+class FLevelViewportLayout : public TSharedFromThis<FLevelViewportLayout>, public FTickableEditorObject
 {
 public:
 	/**
@@ -194,9 +194,9 @@ protected:
 	/** True if we've started an animation and are waiting for it to finish */
 	bool bIsTransitioning;
 
-	/** Curve for animating from a "restored" state to a maximized state, and vice-versa */
+	/** Curve for animating from a "restored" state to a maximized state */
 	FCurveSequence MaximizeAnimation;
-	
+
 	/** List of all of the viewports in this layout */
 	TArray< TSharedPtr< SLevelViewport > > Viewports;
 	
