This expands recent work to resolve comment depths of all comments in a region when a single comment node is moved to that region.
#UE-3724 - LIVE: BP: If a comment box is created inside of another larger comment box the smaller box is no longer editable
#Branch UE4
#Proj GraphEditor

diff --git a/Engine/Source/Editor/GraphEditor/Private/SGraphNodeComment.cpp b/Engine/Source/Editor/GraphEditor/Private/SGraphNodeComment.cpp
--- a/Engine/Source/Editor/GraphEditor/Private/SGraphNodeComment.cpp
+++ b/Engine/Source/Editor/GraphEditor/Private/SGraphNodeComment.cpp
@@ -280,7 +280,6 @@ void SGraphNodeComment::HandleSelection(bool bSelected, bool bUpdateNodesUnderCo
 				FChildren* PanelChildren = Panel->GetAllChildren();
 				int32 NumChildren = PanelChildren->Num();
 				CommentNode->ClearNodesUnderComment();
-				int32 MinDepth = 0;
 
 				for ( int32 NodeIndex=0; NodeIndex < NumChildren; ++NodeIndex )
 				{
@@ -294,19 +293,12 @@ void SGraphNodeComment::HandleSelection(bool bSelected, bool bUpdateNodesUnderCo
 						const FVector2D SomeNodeSize = SomeNodeWidget->GetDesiredSize();
 
 						const FSlateRect NodeGeometryGraphSpace( SomeNodePosition.X, SomeNodePosition.Y, SomeNodePosition.X + SomeNodeSize.X, SomeNodePosition.Y + SomeNodeSize.Y );
-						if( FSlateRect::DoRectanglesIntersect( CommentRect, NodeGeometryGraphSpace ) )
+						if( FSlateRect::IsRectangleContained( CommentRect, NodeGeometryGraphSpace ) )
 						{
-							MinDepth = FMath::Min( MinDepth, SomeNodeWidget->GetSortDepth() - 1 );
-
-							if ( FSlateRect::IsRectangleContained( CommentRect, NodeGeometryGraphSpace ) )
-							{
-								CommentNode->AddNodeUnderComment(GraphObject);
-							}
+							CommentNode->AddNodeUnderComment(GraphObject);
 						}
 					}
 				}
-				// Fix Depth to include any overlapped comments
-				CommentNode->CommentDepth = FMath::Min( CommentNode->CommentDepth, MinDepth );
 			}
 		}
 		bIsSelected = bSelected;
@@ -366,6 +358,38 @@ void SGraphNodeComment::MoveTo( const FVector2D& NewPosition, FNodeSet& NodeFilt
 	}
 }
 
+void SGraphNodeComment::EndUserInteraction() const
+{
+	// Find any parent comments and their list of child nodes
+	const FVector2D NodeSize = GetDesiredSize();
+	if( !NodeSize.IsZero() )
+	{
+		const FVector2D NodePosition = GetPosition();
+		const FSlateRect CommentRect( NodePosition.X, NodePosition.Y, NodePosition.X + NodeSize.X, NodePosition.Y + NodeSize.Y );
+
+		TSharedPtr<SGraphPanel> Panel = GetOwnerPanel();
+		FChildren* PanelChildren = Panel->GetAllChildren();
+		int32 NumChildren = PanelChildren->Num();
+
+		for ( int32 NodeIndex=0; NodeIndex < NumChildren; ++NodeIndex )
+		{
+			TSharedPtr<SGraphNodeComment> CommentWidget = StaticCastSharedRef<SGraphNodeComment>(PanelChildren->GetChildAt(NodeIndex));
+
+			if( CommentWidget.IsValid() )
+			{
+				const FVector2D SomeNodePosition = CommentWidget->GetPosition();
+				const FVector2D SomeNodeSize = CommentWidget->GetDesiredSize();
+
+				const FSlateRect NodeGeometryGraphSpace( SomeNodePosition.X, SomeNodePosition.Y, SomeNodePosition.X + SomeNodeSize.X, SomeNodePosition.Y + SomeNodeSize.Y );
+				if( FSlateRect::DoRectanglesIntersect( CommentRect, NodeGeometryGraphSpace ) )
+				{
+					CommentWidget->HandleSelection( CommentWidget->bIsSelected, true );
+				}
+			}
+		}
+	}
+}
+
 float SGraphNodeComment::GetTitleBarHeight() const
 {
 	return TitleBar.IsValid() ? TitleBar->GetDesiredSize().Y : 0.0f;
diff --git a/Engine/Source/Editor/GraphEditor/Private/SNodePanel.cpp b/Engine/Source/Editor/GraphEditor/Private/SNodePanel.cpp
--- a/Engine/Source/Editor/GraphEditor/Private/SNodePanel.cpp
+++ b/Engine/Source/Editor/GraphEditor/Private/SNodePanel.cpp
@@ -406,6 +406,7 @@ void SNodePanel::OnBeginNodeInteraction(const TSharedRef<SNode>& InNodeToDrag, c
 
 void SNodePanel::OnEndNodeInteraction(const TSharedRef<SNode>& InNodeToDrag)
 {
+	InNodeToDrag->EndUserInteraction();
 }
 
 EActiveTimerReturnType SNodePanel::HandleZoomToFit(double InCurrentTime, float InDeltaTime)
diff --git a/Engine/Source/Editor/GraphEditor/Public/SGraphNodeComment.h b/Engine/Source/Editor/GraphEditor/Public/SGraphNodeComment.h
--- a/Engine/Source/Editor/GraphEditor/Public/SGraphNodeComment.h
+++ b/Engine/Source/Editor/GraphEditor/Public/SGraphNodeComment.h
@@ -21,6 +21,7 @@ public:
 	virtual void GetOverlayBrushes(bool bSelected, const FVector2D WidgetSize, TArray<FOverlayBrushInfo>& Brushes) const;
 	virtual bool ShouldAllowCulling() const override { return false; }
 	virtual int32 GetSortDepth() const override;
+	virtual void EndUserInteraction() const override;
 	// End SNodePanel::SNode interface
 
 	// Begin SPanel Interface
diff --git a/Engine/Source/Editor/GraphEditor/Public/SNodePanel.h b/Engine/Source/Editor/GraphEditor/Public/SNodePanel.h
--- a/Engine/Source/Editor/GraphEditor/Public/SNodePanel.h
+++ b/Engine/Source/Editor/GraphEditor/Public/SNodePanel.h
@@ -472,6 +472,9 @@ public:
 			return true;
 		}
 
+		/** Called when user interaction has completed */
+		virtual void EndUserInteraction() const {}
+
 		/** 
 		 *	override, when area used to select node, should be different, than it's size
 		 *	e.g. comment node - only title bar is selectable
