Fixed code relying on SLATE_TEXT_ATTRIBUTE for STextBlock.

diff --git a/Engine/Source/Editor/GraphEditor/Private/AnimationStateNodes/SGraphNodeAnimStateEntry.cpp b/Engine/Source/Editor/GraphEditor/Private/AnimationStateNodes/SGraphNodeAnimStateEntry.cpp
--- a/Engine/Source/Editor/GraphEditor/Private/AnimationStateNodes/SGraphNodeAnimStateEntry.cpp
+++ b/Engine/Source/Editor/GraphEditor/Private/AnimationStateNodes/SGraphNodeAnimStateEntry.cpp
@@ -85,7 +85,7 @@ void SGraphNodeAnimStateEntry::AddPin(const TSharedRef<SGraphPin>& PinToAdd)
 	OutputPins.Add(PinToAdd);
 }
 
-FString SGraphNodeAnimStateEntry::GetPreviewCornerText() const
+FText SGraphNodeAnimStateEntry::GetPreviewCornerText() const
 {
-	return FString(TEXT("Entry point for state machine"));
+	return NSLOCTEXT("SGraphNodeAnimStateEntry", "CornerTextDescription", "Entry point for state machine");
 }
diff --git a/Engine/Source/Editor/GraphEditor/Private/AnimationStateNodes/SGraphNodeAnimStateEntry.h b/Engine/Source/Editor/GraphEditor/Private/AnimationStateNodes/SGraphNodeAnimStateEntry.h
--- a/Engine/Source/Editor/GraphEditor/Private/AnimationStateNodes/SGraphNodeAnimStateEntry.h
+++ b/Engine/Source/Editor/GraphEditor/Private/AnimationStateNodes/SGraphNodeAnimStateEntry.h
@@ -29,5 +29,5 @@ public:
 protected:
 	FSlateColor GetBorderBackgroundColor() const;
 
-	FString GetPreviewCornerText() const;
+	FText GetPreviewCornerText() const;
 };
diff --git a/Engine/Source/Editor/GraphEditor/Private/AnimationStateNodes/SGraphNodeAnimTransition.cpp b/Engine/Source/Editor/GraphEditor/Private/AnimationStateNodes/SGraphNodeAnimTransition.cpp
--- a/Engine/Source/Editor/GraphEditor/Private/AnimationStateNodes/SGraphNodeAnimTransition.cpp
+++ b/Engine/Source/Editor/GraphEditor/Private/AnimationStateNodes/SGraphNodeAnimTransition.cpp
@@ -108,7 +108,7 @@ TSharedRef<SWidget> SGraphNodeAnimTransition::GenerateRichTooltip()
 
 	TSharedRef<SVerticalBox> Widget = SNew(SVerticalBox);
 
-	const FString TooltipDesc = GetPreviewCornerText(false);
+	const FText TooltipDesc = GetPreviewCornerText(false);
 
 	
 	// Transition rule linearized
@@ -194,39 +194,41 @@ void SGraphNodeAnimTransition::UpdateGraphNode()
 		];
 }
 
-FString SGraphNodeAnimTransition::GetPreviewCornerText(bool bReverse) const
+FText SGraphNodeAnimTransition::GetPreviewCornerText(bool bReverse) const
 {
 	UAnimStateTransitionNode* TransNode = CastChecked<UAnimStateTransitionNode>(GraphNode);
 
 	UAnimStateNodeBase* PrevState = (bReverse ? TransNode->GetNextState() : TransNode->GetPreviousState());
 	UAnimStateNodeBase* NextState = (bReverse ? TransNode->GetPreviousState() : TransNode->GetNextState());
 
-	FString Result = TEXT("Bad transition (missing source or target)");
+	FText Result = LOCTEXT("BadTransition", "Bad transition (missing source or target)");
 
 	// Show the priority if there is any ambiguity
 	if (PrevState != NULL)
 	{
 		if (NextState != NULL)
 		{
-			Result = FString::Printf(TEXT("%s to %s"), *PrevState->GetStateName(), *NextState->GetStateName());
-		}
-
-		TArray<UAnimStateTransitionNode*> TransitionFromSource;
-		PrevState->GetTransitionList(/*out*/ TransitionFromSource);
+			TArray<UAnimStateTransitionNode*> TransitionFromSource;
+			PrevState->GetTransitionList(/*out*/ TransitionFromSource);
 
-		if (TransitionFromSource.Num() > 1)
-		{
-			// See if the priorities differ
 			bool bMultiplePriorities = false;
-			for (int32 Index = 0; (Index < TransitionFromSource.Num()) && !bMultiplePriorities; ++Index)
+			if (TransitionFromSource.Num() > 1)
 			{
-				const bool bDifferentPriority = (TransitionFromSource[Index]->PriorityOrder != TransNode->PriorityOrder);
-				bMultiplePriorities |= bDifferentPriority;
+				// See if the priorities differ
+				for (int32 Index = 0; (Index < TransitionFromSource.Num()) && !bMultiplePriorities; ++Index)
+				{
+					const bool bDifferentPriority = (TransitionFromSource[Index]->PriorityOrder != TransNode->PriorityOrder);
+					bMultiplePriorities |= bDifferentPriority;
+				}
 			}
 
 			if (bMultiplePriorities)
 			{
-				Result += FString::Printf(TEXT(" (Priority %d)"), TransNode->PriorityOrder);
+				Result = FText::Format(LOCTEXT("TransitionXToYWithPriority", "{0} to {1} (Priority {2})"), FText::FromString(PrevState->GetStateName()), FText::FromString(NextState->GetStateName()), FText::AsNumber(TransNode->PriorityOrder));
+			}
+			else
+			{
+				Result = FText::Format(LOCTEXT("TransitionXToY", "{0} to {1}"), FText::FromString(PrevState->GetStateName()), FText::FromString(NextState->GetStateName()));
 			}
 		}
 	}
diff --git a/Engine/Source/Editor/GraphEditor/Private/AnimationStateNodes/SGraphNodeAnimTransition.h b/Engine/Source/Editor/GraphEditor/Private/AnimationStateNodes/SGraphNodeAnimTransition.h
--- a/Engine/Source/Editor/GraphEditor/Private/AnimationStateNodes/SGraphNodeAnimTransition.h
+++ b/Engine/Source/Editor/GraphEditor/Private/AnimationStateNodes/SGraphNodeAnimTransition.h
@@ -35,7 +35,7 @@ private:
 	TSharedPtr<STextEntryPopup> TextEntryWidget;
 
 private:
-	FString GetPreviewCornerText(bool reverse) const;
+	FText GetPreviewCornerText(bool reverse) const;
 	FSlateColor GetTransitionColor() const;
 
 	TSharedRef<SWidget> GenerateInlineDisplayOrEditingWidget(bool bShowGraphPreview);
diff --git a/Engine/Source/Editor/GraphEditor/Private/GraphDiffControl.cpp b/Engine/Source/Editor/GraphEditor/Private/GraphDiffControl.cpp
--- a/Engine/Source/Editor/GraphEditor/Private/GraphDiffControl.cpp
+++ b/Engine/Source/Editor/GraphEditor/Private/GraphDiffControl.cpp
@@ -80,7 +80,7 @@ static void DiffR_NodeAdded( FDiffResults& Results, UEdGraphNode* Node )
 
 		FFormatNamedArguments Args;
 		Args.Add(TEXT("NodeTitle"), Node->GetNodeTitle(ENodeTitleType::ListView));
-		Diff.ToolTip =  FText::Format(LOCTEXT("DIF_AddNode", "Added Node '{NodeTitle}'"), Args).ToString();
+		Diff.ToolTip =  FText::Format(LOCTEXT("DIF_AddNode", "Added Node '{NodeTitle}'"), Args);
 		Diff.DisplayString = Diff.ToolTip;
 		Diff.DisplayColor = FLinearColor(0.3f,1.0f,0.4f);
 		Results.Add(Diff);
@@ -99,7 +99,7 @@ static void DiffR_NodeCommentChanged(FDiffResults& Results, class UEdGraphNode*
 
 		FFormatNamedArguments Args;
 		Args.Add(TEXT("NodeTitle"), Node->GetNodeTitle(ENodeTitleType::ListView));
-		Diff.ToolTip =  FText::Format(LOCTEXT("DIF_CommentModified", "Comment Modified Node '{NodeTitle}'"), Args).ToString();
+		Diff.ToolTip =  FText::Format(LOCTEXT("DIF_CommentModified", "Comment Modified Node '{NodeTitle}'"), Args);
 
 		Diff.DisplayString = Diff.ToolTip;
 		Diff.DisplayColor = FLinearColor(0.25f,0.4f,0.5f);
@@ -119,7 +119,7 @@ static void DiffR_NodeMoved(FDiffResults& Results, UEdGraphNode* Node, class UEd
 
 		FFormatNamedArguments Args;
 		Args.Add(TEXT("NodeTitle"), Node->GetNodeTitle(ENodeTitleType::ListView));
-		Diff.ToolTip = FText::Format(LOCTEXT("DIF_MoveNode", "Moved Node '{NodeTitle}'"), Args).ToString();
+		Diff.ToolTip = FText::Format(LOCTEXT("DIF_MoveNode", "Moved Node '{NodeTitle}'"), Args);
 
 		Diff.DisplayString = Diff.ToolTip;
 		Diff.DisplayColor = FLinearColor(0.9f, 0.84f, 0.43f);
@@ -143,18 +143,16 @@ static void DiffR_PinTypeChanged(FDiffResults& Results, class UEdGraphPin* Pin2,
 		if(Type1.PinCategory != Type2.PinCategory)
 		{
 			Diff.Diff = EDiffType::PIN_TYPE_CATEGORY;
-			Diff.ToolTip =  FString::Printf(*LOCTEXT("DIF_PinCategoryToolTip", "Pin '%s' Category was '%s', but is now '%s").ToString(), *Pin2->PinName, *Pin1->PinType.PinCategory, *Pin2->PinType.PinCategory);
+			Diff.ToolTip = FText::Format(LOCTEXT("DIF_PinCategoryToolTipFmt", "Pin '{0}' Category was '{1}', but is now '{2}"), FText::FromString(Pin2->PinName), FText::FromString(Pin1->PinType.PinCategory), FText::FromString(Pin2->PinType.PinCategory));
 			Diff.DisplayColor = FLinearColor(0.15f,0.53f,0.15f);
-			Diff.DisplayString =  FString::Printf(*LOCTEXT("DIF_PinCategory", "Pin Category '%s' ['%s' -> '%s']").ToString(), *Pin2->PinName, *Pin1->PinType.PinCategory, *Pin2->PinType.PinCategory);
+			Diff.DisplayString = FText::Format(LOCTEXT("DIF_PinCategoryFmt", "Pin Category '{0}' ['{1}' -> '{2}']"), FText::FromString(Pin2->PinName), FText::FromString(Pin1->PinType.PinCategory), FText::FromString(Pin2->PinType.PinCategory));
 		}
 		else if(Type1.PinSubCategory != Type2.PinSubCategory)
 		{
 			Diff.Diff = EDiffType::PIN_TYPE_SUBCATEGORY;
-			Diff.ToolTip =  FString::Printf(*LOCTEXT("DIF_PinSubCategoryToolTip", "Pin '%s' SubCategory was '%s', but is now '%s").ToString(), *Pin2->PinName, *Pin1->PinType.PinSubCategory,
-				*Pin2->PinType.PinSubCategory);
+			Diff.ToolTip = FText::Format(LOCTEXT("DIF_PinSubCategoryToolTipFmt", "Pin '{0}' SubCategory was '{1}', but is now '{2}"), FText::FromString(Pin2->PinName), FText::FromString(Pin1->PinType.PinSubCategory), FText::FromString(Pin2->PinType.PinSubCategory));
 			Diff.DisplayColor = FLinearColor(0.45f,0.53f,0.65f);
-			Diff.DisplayString =   FString::Printf(*LOCTEXT("DIF_PinSubCategory", "Pin SubCategory '%s'  ['%s' -> '%s']").ToString(), *Pin2->PinName, *Pin1->PinType.PinSubCategory,
-				*Pin2->PinType.PinSubCategory);
+			Diff.DisplayString = FText::Format(LOCTEXT("DIF_PinSubCategoryFmt", "Pin SubCategory '{0}'  ['{1}' -> '{2}']"), FText::FromString(Pin2->PinName), FText::FromString(Pin1->PinType.PinSubCategory), FText::FromString(Pin2->PinType.PinSubCategory));
 		}
 		else if(T1Obj != T2Obj && (T1Obj && T2Obj ) &&
 			(T1Obj->GetFName() != T2Obj->GetFName()))
@@ -163,29 +161,29 @@ static void DiffR_PinTypeChanged(FDiffResults& Results, class UEdGraphPin* Pin2,
 			FString Obj2 = T2Obj->GetFName().ToString();
 
 			Diff.Diff = EDiffType::PIN_TYPE_SUBCATEGORY_OBJECT;
-			Diff.ToolTip =  FString::Printf(*LOCTEXT("DIF_PinSubCategorObjToolTip", "Pin '%s' was SubCategoryObject '%s', but is now '%s").ToString(), *Pin2->PinName, *Obj1, *Obj2);
+			Diff.ToolTip = FText::Format(LOCTEXT("DIF_PinSubCategorObjToolTipFmt", "Pin '{0}' was SubCategoryObject '{1}', but is now '{2}"), FText::FromString(Pin2->PinName), FText::FromString(Obj1), FText::FromString(Obj2));
 			Diff.DisplayColor = FLinearColor(0.45f,0.13f,0.25f);
-			Diff.DisplayString = FString::Printf(*LOCTEXT("DIF_PinSubCategoryObj", "Pin SubCategoryObject '%s' ['%s' -> '%s']").ToString(), *Pin2->PinName, *Obj1, *Obj2);
+			Diff.DisplayString = FText::Format(LOCTEXT("DIF_PinSubCategoryObjFmt", "Pin SubCategoryObject '{0}' ['{1}' -> '{2}']"), FText::FromString(Pin2->PinName), FText::FromString(Obj1), FText::FromString(Obj2));
 		}
 		else if(Type1.bIsArray != Type2.bIsArray)
 		{
 			Diff.Diff = EDiffType::PIN_TYPE_IS_ARRAY;
-			FString IsArray1 = Pin1->PinType.bIsArray ? TEXT("true") : TEXT("false");
-			FString IsArray2 = Pin2->PinType.bIsArray ? TEXT("true") : TEXT("false");
+			FText IsArray1 = Pin1->PinType.bIsArray ? LOCTEXT("true", "true") : LOCTEXT("false", "false");
+			FText IsArray2 = Pin2->PinType.bIsArray ? LOCTEXT("true", "true") : LOCTEXT("false", "false");
 
-			Diff.ToolTip =  FString::Printf(*LOCTEXT("DIF_PinIsArrayToolTip", "PinType IsArray for '%s' modified. Was '%s', but is now '%s").ToString(), *Pin2->PinName, *IsArray1, *IsArray2);
-			Diff.DisplayString = FString::Printf(*LOCTEXT("DIF_PinIsArray", "Pin IsArray '%s' ['%s' -> '%s']").ToString(), *Pin2->PinName, *IsArray1, *IsArray2);
+			Diff.ToolTip = FText::Format(LOCTEXT("DIF_PinIsArrayToolTipFmt", "PinType IsArray for '{0}' modified. Was '{1}', but is now '{2}"), FText::FromString(Pin2->PinName), IsArray1, IsArray2);
+			Diff.DisplayString = FText::Format(LOCTEXT("DIF_PinIsArrayFmt", "Pin IsArray '{0}' ['{1}' -> '{2}']"), FText::FromString(Pin2->PinName), IsArray1, IsArray2);
 			Diff.DisplayColor = FLinearColor(0.45f,0.33f,0.35f);
 		}
 		else if(Type1.bIsReference != Type2.bIsReference)
 		{
 			Diff.Diff = EDiffType::PIN_TYPE_IS_REF;
-			FString IsRef1 = Pin1->PinType.bIsReference ? TEXT("true") : TEXT("false");
-			FString IsRef2 = Pin2->PinType.bIsReference ? TEXT("true") : TEXT("false");
+			FText IsRef1 = Pin1->PinType.bIsReference ? LOCTEXT("true", "true") : LOCTEXT("false", "false");
+			FText IsRef2 = Pin2->PinType.bIsReference ? LOCTEXT("true", "true") : LOCTEXT("false", "false");
 
-			Diff.ToolTip =  FString::Printf(*LOCTEXT("DIF_PinIsRefToolTip", "PinType IsReference for '%s' modified. Was '%s', but is now '%s").ToString(), *Pin2->PinName, *IsRef1, *IsRef2);
+			Diff.ToolTip = FText::Format(LOCTEXT("DIF_PinIsRefToolTipFmt", "PinType IsReference for '{0}' modified. Was '{1}', but is now '{2}"), FText::FromString(Pin2->PinName), IsRef1, IsRef2);
 			Diff.DisplayColor = FLinearColor(0.25f,0.43f,0.35f);
-			Diff.DisplayString = FString::Printf(*LOCTEXT("DIF_PinIsRef", "Pin IsReference '%s' ['%s' -> '%s']").ToString(), *Pin2->PinName, *IsRef1, *IsRef2);
+			Diff.DisplayString = FText::Format(LOCTEXT("DIF_PinIsRefFmt", "Pin IsReference '{0}' ['{1}' -> '{2}']"), FText::FromString(Pin2->PinName), IsRef1, IsRef2);
 		}
 
 		Diff.Pin1 = Pin1;
@@ -205,15 +203,15 @@ static void DiffR_PinLinkCountChanged(FDiffResults& Results, class UEdGraphPin*
 		Diff.Pin1 = Pin1;
 		if(Diff.Diff == EDiffType::PIN_LINKEDTO_NUM_INC)
 		{
-			Diff.ToolTip =  FString::Printf(*LOCTEXT("DIF_PinLinkCountIncToolTip", "Pin '%s' has more links(was %i now %i)").ToString(), *Pin1->PinName, Pin1->LinkedTo.Num() ,Pin2->LinkedTo.Num());
+			Diff.ToolTip = FText::Format(LOCTEXT("DIF_PinLinkCountIncToolTipFmt", "Pin '{0}' has more links (was {1} now {2})"), FText::FromString(Pin1->PinName), FText::AsNumber(Pin1->LinkedTo.Num()), FText::AsNumber(Pin2->LinkedTo.Num()));
 			Diff.DisplayColor = FLinearColor(0.5f,0.3f,0.85f);
-			Diff.DisplayString =  FString::Printf(*LOCTEXT("DIF_PinLinkCountInc", "Added Link to '%s'").ToString(), *Pin1->PinName);
+			Diff.DisplayString = FText::Format(LOCTEXT("DIF_PinLinkCountIncFmt", "Added Link to '{0}'"), FText::FromString(Pin1->PinName));
 		}
 		else
 		{
-			Diff.ToolTip =  FString::Printf(*LOCTEXT("DIF_PinLinkCountDecToolTip", "Pin '%s' has fewer links(was %i now %i)").ToString(), *Pin1->PinName, Pin1->LinkedTo.Num(), Pin2->LinkedTo.Num());
+			Diff.ToolTip = FText::Format(LOCTEXT("DIF_PinLinkCountDecToolTipFmt", "Pin '{0}' has fewer links (was {1} now {2})"), FText::FromString(Pin1->PinName), FText::AsNumber(Pin1->LinkedTo.Num()), FText::AsNumber(Pin2->LinkedTo.Num()));
 			Diff.DisplayColor = FLinearColor(0.5f,0.3f,0.85f);
-			Diff.DisplayString =  FString::Printf(*LOCTEXT("DIF_PinLinkCountDec", "Removed Link to '%s'").ToString(), *Pin1->PinName);
+			Diff.DisplayString = FText::Format(LOCTEXT("DIF_PinLinkCountDecFmt", "Removed Link to '{0}'"), FText::FromString(Pin1->PinName));
 		}
 
 		Results.Add(Diff);
@@ -239,10 +237,10 @@ static void DiffR_LinkedToNode(FDiffResults& Results, class UEdGraphPin* Pin1, c
 		Args.Add(TEXT("PinNameForNode1"), FText::FromString(Pin1->PinName));
 		Args.Add(TEXT("NodeName1"), Node1Name);
 		Args.Add(TEXT("NodeName2"), Node2Name);
-		Diff.ToolTip =  FText::Format(LOCTEXT("DIF_PinLinkMovedToolTip", "Pin '{PinNameForNode1}' was linked to Node '{NodeName1}', but is now linked to Node '{NodeName2}'"), Args).ToString();
+		Diff.ToolTip = FText::Format(LOCTEXT("DIF_PinLinkMovedToolTip", "Pin '{PinNameForNode1}' was linked to Node '{NodeName1}', but is now linked to Node '{NodeName2}'"), Args);
 
 		Diff.DisplayColor = FLinearColor(0.85f,0.71f,0.25f);
-		Diff.DisplayString = FText::Format(LOCTEXT("DIF_PinLinkMoved", "Link Moved  '{PinNameForNode1}' ['{NodeName1}' -> '{NodeName2}']"), Args).ToString();
+		Diff.DisplayString = FText::Format(LOCTEXT("DIF_PinLinkMoved", "Link Moved  '{PinNameForNode1}' ['{NodeName1}' -> '{NodeName2}']"), Args);
 		Results.Add(Diff);
 	}
 }
@@ -261,9 +259,9 @@ static void DiffR_PinDefaultValueChanged(FDiffResults& Results, class UEdGraphPi
 		Args.Add(TEXT("PinNameForValue1"), FText::FromString(Pin2->PinName));
 		Args.Add(TEXT("PinValue1"), FText::FromString(Pin1->GetDefaultAsString()));
 		Args.Add(TEXT("PinValue2"), FText::FromString(Pin2->GetDefaultAsString()));
-		Diff.ToolTip =  FText::Format(LOCTEXT("DIF_PinDefaultValueToolTip", "Pin '{PinNameForValue1}' Default Value was '{PinValue1}', but is now '{PinValue2}"), Args).ToString();
+		Diff.ToolTip = FText::Format(LOCTEXT("DIF_PinDefaultValueToolTip", "Pin '{PinNameForValue1}' Default Value was '{PinValue1}', but is now '{PinValue2}"), Args);
 		Diff.DisplayColor = FLinearColor(0.665f,0.13f,0.455f);
-		Diff.DisplayString = FText::Format(LOCTEXT("DIF_PinDefaultValue", "Pin Default '{PinNameForValue1}' '{PinValue1}' -> '{PinValue2}']"), Args).ToString();
+		Diff.DisplayString = FText::Format(LOCTEXT("DIF_PinDefaultValue", "Pin Default '{PinNameForValue1}' '{PinValue1}' -> '{PinValue2}']"), Args);
 		Results.Add(Diff);
 	}
 }
@@ -293,10 +291,11 @@ static void DiffR_NodePinCount(FDiffResults& Results, class UEdGraphNode* Node,
 		Args.Add(TEXT("NodeName"), NodeName);
 		Args.Add(TEXT("OriginalCount"), OriginalCount);
 		Args.Add(TEXT("NewCount"), NewCount);
-		Diff.ToolTip =  FText::Format(LOCTEXT("DIF_NodePinCountChangedToolTip", "Node '{NodeName}' had {OriginalCount} Pins, now has {NewCount} Pins"), Args).ToString();
+		Diff.ToolTip = FText::Format(LOCTEXT("DIF_NodePinCountChangedToolTip", "Node '{NodeName}' had {OriginalCount} Pins, now has {NewCount} Pins"), Args);
 		Diff.DisplayColor = FLinearColor(0.45f,0.4f,0.4f);
-		Diff.DisplayString =  (OriginalCount < NewCount) ?  (FText::Format(LOCTEXT("DIF_NodePinCountIncreased", "Add Pin '{NodeName}'"), Args).ToString()) : 
-			(FText::Format(LOCTEXT("DIF_NodePinCountDecreased", "Removed Pin '{NodeName}'"), Args).ToString());
+		Diff.DisplayString = (OriginalCount < NewCount) 
+			? (FText::Format(LOCTEXT("DIF_NodePinCountIncreased", "Add Pin '{NodeName}'"), Args)) 
+			: (FText::Format(LOCTEXT("DIF_NodePinCountDecreased", "Removed Pin '{NodeName}'"), Args));
 		Results.Add(Diff);
 	}
 }
@@ -573,7 +572,7 @@ bool FGraphDiffControl::DiffGraphs(UEdGraph* const LhsGraph, UEdGraph* const Rhs
 
 						FFormatNamedArguments Args;
 						Args.Add(TEXT("NodeTitle"), Difference.Node1->GetNodeTitle(ENodeTitleType::ListView));
-						Difference.ToolTip       = FText::Format(LOCTEXT("DIF_RemoveNode", "Removed Node '{NodeTitle}'"), Args).ToString();
+						Difference.ToolTip       = FText::Format(LOCTEXT("DIF_RemoveNode", "Removed Node '{NodeTitle}'"), Args);
 
 						Difference.DisplayString = Difference.ToolTip;
 						Difference.DisplayColor  = FLinearColor(1.f,0.4f,0.4f);
diff --git a/Engine/Source/Editor/GraphEditor/Private/KismetNodes/SGraphNodeK2Base.cpp b/Engine/Source/Editor/GraphEditor/Private/KismetNodes/SGraphNodeK2Base.cpp
--- a/Engine/Source/Editor/GraphEditor/Private/KismetNodes/SGraphNodeK2Base.cpp
+++ b/Engine/Source/Editor/GraphEditor/Private/KismetNodes/SGraphNodeK2Base.cpp
@@ -438,7 +438,7 @@ void SGraphNodeK2Base::GetNodeInfoPopups(FNodeInfoContext* Context, TArray<FGrap
 
 					FString PinName = UEdGraphSchema_K2::TypeToText(WatchPin->PinType).ToString();
 					PinName += TEXT(" ");
-					PinName += Schema->GetPinDisplayName(WatchPin);
+					PinName += Schema->GetPinDisplayName(WatchPin).ToString();
 
 					FString WatchText;
 					const FKismetDebugUtilities::EWatchTextResult WatchStatus = FKismetDebugUtilities::GetWatchText(/*inout*/ WatchText, Blueprint, ActiveObject, WatchPin);
diff --git a/Engine/Source/Editor/GraphEditor/Private/KismetNodes/SGraphNodeK2Composite.cpp b/Engine/Source/Editor/GraphEditor/Private/KismetNodes/SGraphNodeK2Composite.cpp
--- a/Engine/Source/Editor/GraphEditor/Private/KismetNodes/SGraphNodeK2Composite.cpp
+++ b/Engine/Source/Editor/GraphEditor/Private/KismetNodes/SGraphNodeK2Composite.cpp
@@ -220,15 +220,15 @@ TSharedPtr<SToolTip> SGraphNodeK2Composite::GetComplexTooltip()
 
 }
 
-FString SGraphNodeK2Composite::GetPreviewCornerText() const
+FText SGraphNodeK2Composite::GetPreviewCornerText() const
 {
 	UEdGraph* BoundGraph = GetInnerGraph();
-	return BoundGraph->GetName();
+	return FText::FromString(BoundGraph->GetName());
 }
 
-FString SGraphNodeK2Composite::GetTooltipTextForNode() const
+FText SGraphNodeK2Composite::GetTooltipTextForNode() const
 {
-	return GraphNode->GetTooltipText().ToString();
+	return GraphNode->GetTooltipText();
 }
 
 TSharedRef<SWidget> SGraphNodeK2Composite::CreateNodeBody()
diff --git a/Engine/Source/Editor/GraphEditor/Private/KismetNodes/SGraphNodeK2Composite.h b/Engine/Source/Editor/GraphEditor/Private/KismetNodes/SGraphNodeK2Composite.h
--- a/Engine/Source/Editor/GraphEditor/Private/KismetNodes/SGraphNodeK2Composite.h
+++ b/Engine/Source/Editor/GraphEditor/Private/KismetNodes/SGraphNodeK2Composite.h
@@ -18,8 +18,8 @@ protected:
 	virtual UEdGraph* GetInnerGraph() const;
 
 private:
-	FString GetPreviewCornerText() const;
-	FString GetTooltipTextForNode() const;
+	FText GetPreviewCornerText() const;
+	FText GetTooltipTextForNode() const;
 
 	TSharedRef<SWidget> CreateNodeBody();
 };
diff --git a/Engine/Source/Editor/GraphEditor/Private/KismetNodes/SGraphNodeK2CreateDelegate.cpp b/Engine/Source/Editor/GraphEditor/Private/KismetNodes/SGraphNodeK2CreateDelegate.cpp
--- a/Engine/Source/Editor/GraphEditor/Private/KismetNodes/SGraphNodeK2CreateDelegate.cpp
+++ b/Engine/Source/Editor/GraphEditor/Private/KismetNodes/SGraphNodeK2CreateDelegate.cpp
@@ -8,7 +8,7 @@ FString SGraphNodeK2CreateDelegate::FunctionDescription(const UFunction* Functio
 {
 	if(!Function || !Function->GetOuter())
 	{
-		return TEXT("Error");
+		return NSLOCTEXT("GraphNodeK2Create", "Error", "Error").ToString();
 	}
 
 	const UEdGraphSchema_K2* K2Schema = GetDefault<UEdGraphSchema_K2>();
@@ -48,7 +48,7 @@ void SGraphNodeK2CreateDelegate::Construct( const FArguments& InArgs, UK2Node* I
 	UpdateGraphNode();
 }
 
-FString SGraphNodeK2CreateDelegate::GetCurrentFunctionDescription() const
+FText SGraphNodeK2CreateDelegate::GetCurrentFunctionDescription() const
 {
 	UK2Node_CreateDelegate* Node = Cast<UK2Node_CreateDelegate>(GraphNode);
 	UFunction* FunctionSignature = Node ? Node->GetDelegateSignature() : NULL;
@@ -56,20 +56,20 @@ FString SGraphNodeK2CreateDelegate::GetCurrentFunctionDescription() const
 
 	if(!FunctionSignature || !ScopeClass)
 	{
-		return TEXT("None");
+		return NSLOCTEXT("GraphNodeK2Create", "NoneLabel", "None");
 	}
 
 	if (const auto Func = FindField<UFunction>(ScopeClass, Node->GetFunctionName()))
 	{
-		return FunctionDescription(Func);
+		return FText::FromString(FunctionDescription(Func));
 	}
 
 	if (Node->GetFunctionName() != NAME_None)
 	{
-		return FString::Printf(TEXT("Error? %s"), *Node->GetFunctionName().ToString());
+		return FText::Format(NSLOCTEXT("GraphNodeK2Create", "ErrorLabelFmt", "Error? {0}"), FText::FromName(Node->GetFunctionName()));
 	}
 
-	return TEXT("Select Function");
+	return NSLOCTEXT("GraphNodeK2Create", "SelectFunctionLabel", "Select Function");
 }
 
 TSharedRef<ITableRow> SGraphNodeK2CreateDelegate::HandleGenerateRowFunction(TSharedPtr<FFunctionItemData> FunctionItemData, const TSharedRef<STableViewBase>& OwnerTable)
@@ -77,7 +77,7 @@ TSharedRef<ITableRow> SGraphNodeK2CreateDelegate::HandleGenerateRowFunction(TSha
 	check(FunctionItemData.IsValid());
 	return SNew(STableRow< TSharedPtr<FFunctionItemData> >, OwnerTable).Content()
 		[
-			SNew( STextBlock ).Text( FunctionItemData->Description )
+			SNew( STextBlock ).Text( FText::FromString(FunctionItemData->Description) )
 		];
 }
 
diff --git a/Engine/Source/Editor/GraphEditor/Private/KismetNodes/SGraphNodeK2CreateDelegate.h b/Engine/Source/Editor/GraphEditor/Private/KismetNodes/SGraphNodeK2CreateDelegate.h
--- a/Engine/Source/Editor/GraphEditor/Private/KismetNodes/SGraphNodeK2CreateDelegate.h
+++ b/Engine/Source/Editor/GraphEditor/Private/KismetNodes/SGraphNodeK2CreateDelegate.h
@@ -23,7 +23,7 @@ public:
 	virtual void CreateBelowWidgetControls(TSharedPtr<SVerticalBox> MainBox) override;
 
 protected:
-	FString GetCurrentFunctionDescription() const;
+	FText GetCurrentFunctionDescription() const;
 	TSharedRef<ITableRow> HandleGenerateRowFunction(TSharedPtr<FFunctionItemData> FunctionItemData, const TSharedRef<STableViewBase>& OwnerTable);
 	void OnFunctionSelected(TSharedPtr<FFunctionItemData> FunctionItemData, ESelectInfo::Type SelectInfo);
 };
diff --git a/Engine/Source/Editor/GraphEditor/Private/KismetPins/SNameComboBox.cpp b/Engine/Source/Editor/GraphEditor/Private/KismetPins/SNameComboBox.cpp
--- a/Engine/Source/Editor/GraphEditor/Private/KismetPins/SNameComboBox.cpp
+++ b/Engine/Source/Editor/GraphEditor/Private/KismetPins/SNameComboBox.cpp
@@ -28,19 +28,19 @@ void SNameComboBox::Construct( const FArguments& InArgs )
 	SelectedItem = NameCombo->GetSelectedItem();
 }
 
-FString SNameComboBox::GetItemNameLabel(TSharedPtr<FName> NameItem) const
+FText SNameComboBox::GetItemNameLabel(TSharedPtr<FName> NameItem) const
 {
 	if (!NameItem.IsValid())
 	{
-		return FString();
+		return FText::GetEmpty();
 	}
 
 	return (GetTextLabelForItem.IsBound())
-		? GetTextLabelForItem.Execute(NameItem)
-		: *NameItem->ToString();
+		? FText::FromString(GetTextLabelForItem.Execute(NameItem))
+		: FText::FromName(*NameItem);
 }
 
-FString SNameComboBox::GetSelectedNameLabel() const
+FText SNameComboBox::GetSelectedNameLabel() const
 {
 	TSharedPtr<FName> StringItem = NameCombo->GetSelectedItem();
 	return GetItemNameLabel(StringItem);
diff --git a/Engine/Source/Editor/GraphEditor/Private/SGraphActionMenu.cpp b/Engine/Source/Editor/GraphEditor/Private/SGraphActionMenu.cpp
--- a/Engine/Source/Editor/GraphEditor/Private/SGraphActionMenu.cpp
+++ b/Engine/Source/Editor/GraphEditor/Private/SGraphActionMenu.cpp
@@ -46,7 +46,7 @@ void SDefaultGraphActionWidget::Construct(const FArguments& InArgs, const FCreat
 	this->ChildSlot
 	[
 		SNew(SHorizontalBox)
-		.ToolTipText( InCreateData->Action->TooltipDescription )
+		.ToolTipText( FText::FromString(InCreateData->Action->TooltipDescription) )
 		+ SHorizontalBox::Slot()
 		.AutoWidth()
 		.VAlign(VAlign_Center)
diff --git a/Engine/Source/Editor/GraphEditor/Private/SGraphEditorImpl.cpp b/Engine/Source/Editor/GraphEditor/Private/SGraphEditorImpl.cpp
--- a/Engine/Source/Editor/GraphEditor/Private/SGraphEditorImpl.cpp
+++ b/Engine/Source/Editor/GraphEditor/Private/SGraphEditorImpl.cpp
@@ -198,28 +198,28 @@ void SGraphEditorImpl::Construct( const FArguments& InArgs )
 	
 	struct Local
 	{
-		static FString GetPIENotifyText(TAttribute<FGraphAppearanceInfo> Appearance, FString DefaultText)
+		static FText GetPIENotifyText(TAttribute<FGraphAppearanceInfo> Appearance, FText DefaultText)
 		{
-			FString OverrideText = Appearance.Get().PIENotifyText;
-			return OverrideText.Len() ? OverrideText : DefaultText;
+			FText OverrideText = Appearance.Get().PIENotifyText;
+			return !OverrideText.IsEmpty() ? OverrideText : DefaultText;
 		}
 
-		static FString GetReadOnlyText(TAttribute<FGraphAppearanceInfo>Appearance, FString DefaultText)
+		static FText GetReadOnlyText(TAttribute<FGraphAppearanceInfo>Appearance, FText DefaultText)
 		{
-			FString OverrideText = Appearance.Get().ReadOnlyText;
-			return OverrideText.Len() ? OverrideText : DefaultText;
+			FText OverrideText = Appearance.Get().ReadOnlyText;
+			return !OverrideText.IsEmpty() ? OverrideText : DefaultText;
 		}
 	};
 	
-	FString DefaultPIENotify(TEXT("SIMULATING"));
-	TAttribute<FString> PIENotifyText = Appearance.IsBound() ?
-		TAttribute<FString>::Create( TAttribute<FString>::FGetter::CreateStatic(&Local::GetPIENotifyText, Appearance, DefaultPIENotify) ) :
-		TAttribute<FString>(DefaultPIENotify);
+	FText DefaultPIENotify(LOCTEXT("GraphSimulatingText", "SIMULATING"));
+	TAttribute<FText> PIENotifyText = Appearance.IsBound() ?
+		TAttribute<FText>::Create(TAttribute<FText>::FGetter::CreateStatic(&Local::GetPIENotifyText, Appearance, DefaultPIENotify)) :
+		TAttribute<FText>(DefaultPIENotify);
 
-	FString DefaultReadOnlyText(LOCTEXT("GraphReadOnlyText", "READ-ONLY").ToString());
-	TAttribute<FString> ReadOnlyText = Appearance.IsBound() ?
-		TAttribute<FString>::Create(TAttribute<FString>::FGetter::CreateStatic(&Local::GetReadOnlyText, Appearance, DefaultReadOnlyText)) :
-		TAttribute<FString>(DefaultReadOnlyText);
+	FText DefaultReadOnlyText(LOCTEXT("GraphReadOnlyText", "READ-ONLY"));
+	TAttribute<FText> ReadOnlyText = Appearance.IsBound() ?
+		TAttribute<FText>::Create(TAttribute<FText>::FGetter::CreateStatic(&Local::GetReadOnlyText, Appearance, DefaultReadOnlyText)) :
+		TAttribute<FText>(DefaultReadOnlyText);
 
 	TSharedPtr<SOverlay> OverlayWidget;
 
@@ -475,16 +475,16 @@ FActionMenuContent SGraphEditorImpl::GraphEd_OnGetContextMenuFor(const FGraphCon
 				return Content;
 			}
 
-			return FActionMenuContent( SNew(STextBlock) .Text( NSLOCTEXT("GraphEditor", "NoNodes", "No Nodes").ToString() ) );
+			return FActionMenuContent( SNew(STextBlock) .Text( NSLOCTEXT("GraphEditor", "NoNodes", "No Nodes") ) );
 		}
 		else
 		{
-			return FActionMenuContent( SNew(STextBlock)  .Text( NSLOCTEXT("GraphEditor", "CannotCreateWhileDebugging", "Cannot create new nodes in a read only graph").ToString() ) );
+			return FActionMenuContent( SNew(STextBlock)  .Text( NSLOCTEXT("GraphEditor", "CannotCreateWhileDebugging", "Cannot create new nodes in a read only graph") ) );
 		}
 	}
 	else
 	{
-		return FActionMenuContent( SNew(STextBlock) .Text( NSLOCTEXT("GraphEditor", "GraphObjectIsNull", "Graph Object is Null").ToString() ) );
+		return FActionMenuContent( SNew(STextBlock) .Text( NSLOCTEXT("GraphEditor", "GraphObjectIsNull", "Graph Object is Null") ) );
 	}
 		
 }
@@ -663,13 +663,13 @@ EVisibility SGraphEditorImpl::ReadOnlyVisibility() const
 	return EVisibility::Hidden;
 }
 
-FString SGraphEditorImpl::GetInstructionText() const
+FText SGraphEditorImpl::GetInstructionText() const
 {
 	if (Appearance.IsBound())
 	{
-		return Appearance.Get().InstructionText.ToString();
+		return Appearance.Get().InstructionText;
 	}
-	return TEXT("");
+	return FText::GetEmpty();
 }
 
 EVisibility SGraphEditorImpl::InstructionTextVisibility() const
diff --git a/Engine/Source/Editor/GraphEditor/Private/SGraphEditorImpl.h b/Engine/Source/Editor/GraphEditor/Private/SGraphEditorImpl.h
--- a/Engine/Source/Editor/GraphEditor/Private/SGraphEditorImpl.h
+++ b/Engine/Source/Editor/GraphEditor/Private/SGraphEditorImpl.h
@@ -75,7 +75,7 @@ private:
 	EVisibility ReadOnlyVisibility() const;
 
 	/** Returns dynamic text, meant to passively instruct the user on what to do in the graph */
-	FString GetInstructionText() const;
+	FText GetInstructionText() const;
 	/** Function to check whether we should show instruction text to the user */
 	EVisibility InstructionTextVisibility() const;
 	/** Returns a 0.0 to 1.0 value, denoting the instruction text's fade percent */
diff --git a/Engine/Source/Editor/GraphEditor/Private/SGraphNode.cpp b/Engine/Source/Editor/GraphEditor/Private/SGraphNode.cpp
--- a/Engine/Source/Editor/GraphEditor/Private/SGraphNode.cpp
+++ b/Engine/Source/Editor/GraphEditor/Private/SGraphNode.cpp
@@ -105,7 +105,7 @@ void SNodeTitle::RebuildWidget()
 		[
 			SNew(STextBlock)
 			.TextStyle( FEditorStyle::Get(), ExtraLineStyle )
-			.Text(Lines[Index])
+			.Text(FText::FromString(Lines[Index]))
 		];
 	}
 }
diff --git a/Engine/Source/Editor/GraphEditor/Private/SGraphPalette.cpp b/Engine/Source/Editor/GraphEditor/Private/SGraphPalette.cpp
--- a/Engine/Source/Editor/GraphEditor/Private/SGraphPalette.cpp
+++ b/Engine/Source/Editor/GraphEditor/Private/SGraphPalette.cpp
@@ -25,7 +25,7 @@ void SGraphPaletteItem::Construct(const FArguments& InArgs, FCreateWidgetForActi
 	const FSlateBrush* IconBrush = FEditorStyle::GetBrush(TEXT("NoBrush"));
 	FSlateColor IconColor = FSlateColor::UseForeground();
 	FString ToolTip = GraphAction->TooltipDescription;
-	FString IconToolTip = ToolTip;
+	FText IconToolTip = FText::FromString(ToolTip);
 	bool bIsReadOnly = false;
 
 	TSharedRef<SWidget> IconWidget = CreateIconWidget( IconToolTip, IconBrush, IconColor );
@@ -93,7 +93,7 @@ FReply SGraphPaletteItem::OnMouseButtonDown( const FGeometry& MyGeometry, const
 	return FReply::Unhandled();
 }
 
-TSharedRef<SWidget> SGraphPaletteItem::CreateIconWidget(const FString& IconToolTip, const FSlateBrush* IconBrush, const FSlateColor& IconColor)
+TSharedRef<SWidget> SGraphPaletteItem::CreateIconWidget(const FText& IconToolTip, const FSlateBrush* IconBrush, const FSlateColor& IconColor)
 {
 	return SNew(SImage)
 		.ToolTipText(IconToolTip)
@@ -101,9 +101,9 @@ TSharedRef<SWidget> SGraphPaletteItem::CreateIconWidget(const FString& IconToolT
 		.ColorAndOpacity(IconColor);
 }
 
-TSharedRef<SWidget> SGraphPaletteItem::CreateIconWidget(const FString& IconToolTip, const FSlateBrush* IconBrush, const FSlateColor& IconColor, const FString& DocLink, const FString& DocExcerpt)
+TSharedRef<SWidget> SGraphPaletteItem::CreateIconWidget(const FText& IconToolTip, const FSlateBrush* IconBrush, const FSlateColor& IconColor, const FString& DocLink, const FString& DocExcerpt)
 {
-	TSharedPtr<SToolTip> ToolTipWidget = IDocumentation::Get()->CreateToolTip(FText::FromString(IconToolTip), NULL, DocLink, DocExcerpt);
+	TSharedPtr<SToolTip> ToolTipWidget = IDocumentation::Get()->CreateToolTip(IconToolTip, NULL, DocLink, DocExcerpt);
 
 	return SNew(SImage)
 		.ToolTip(ToolTipWidget)
diff --git a/Engine/Source/Editor/GraphEditor/Private/SGraphPin.cpp b/Engine/Source/Editor/GraphEditor/Private/SGraphPin.cpp
--- a/Engine/Source/Editor/GraphEditor/Private/SGraphPin.cpp
+++ b/Engine/Source/Editor/GraphEditor/Private/SGraphPin.cpp
@@ -764,7 +764,7 @@ FVector2D SGraphPin::GetNodeOffset() const
 	return CachedNodeOffset;
 }
 
-FString SGraphPin::GetPinLabel() const
+FText SGraphPin::GetPinLabel() const
 {
 	return GetPinObj()->GetOwningNode()->GetPinDisplayName(GetPinObj());
 }
diff --git a/Engine/Source/Editor/GraphEditor/Private/SGraphPreviewer.h b/Engine/Source/Editor/GraphEditor/Private/SGraphPreviewer.h
--- a/Engine/Source/Editor/GraphEditor/Private/SGraphPreviewer.h
+++ b/Engine/Source/Editor/GraphEditor/Private/SGraphPreviewer.h
@@ -9,7 +9,7 @@ public:
 	SLATE_BEGIN_ARGS(SGraphPreviewer)
 		: _ShowGraphStateOverlay(true)
 	{}
-		SLATE_ATTRIBUTE( FString, CornerOverlayText )
+		SLATE_ATTRIBUTE( FText, CornerOverlayText )
 		/** Show overlay elements for the graph state such as the PIE and read-only borders and text */
 		SLATE_ATTRIBUTE(bool, ShowGraphStateOverlay)
 		SLATE_ARGUMENT( TSharedPtr<SWidget>, TitleBar )
diff --git a/Engine/Source/Editor/GraphEditor/Public/DiffResults.h b/Engine/Source/Editor/GraphEditor/Public/DiffResults.h
--- a/Engine/Source/Editor/GraphEditor/Public/DiffResults.h
+++ b/Engine/Source/Editor/GraphEditor/Public/DiffResults.h
@@ -61,10 +61,10 @@ struct FDiffSingleResult
 	class UEdGraphPin* Pin2;
 
 	/** String describing the error to the user */
-	FString DisplayString;
+	FText DisplayString;
 
 	/** Optional tooltip containing more information */
-	FString ToolTip; 
+	FText ToolTip; 
 
 	/** User can override color to use for display string */
 	FLinearColor DisplayColor;
@@ -80,8 +80,8 @@ FORCEINLINE bool operator==( const FDiffSingleResult& LHS, const FDiffSingleResu
 			LHS.Node2 == RHS.Node2 &&
 			LHS.Pin1 == RHS.Pin1 &&
 			LHS.Pin2 == RHS.Pin2 &&
-			LHS.DisplayString == RHS.DisplayString &&
-			LHS.ToolTip == RHS.ToolTip &&
+			LHS.DisplayString.ToString() == RHS.DisplayString.ToString() &&
+			LHS.ToolTip.ToString() == RHS.ToolTip.ToString() &&
 			LHS.DisplayColor == RHS.DisplayColor;
 }
 
diff --git a/Engine/Source/Editor/GraphEditor/Public/SGraphPalette.h b/Engine/Source/Editor/GraphEditor/Public/SGraphPalette.h
--- a/Engine/Source/Editor/GraphEditor/Public/SGraphPalette.h
+++ b/Engine/Source/Editor/GraphEditor/Public/SGraphPalette.h
@@ -33,9 +33,9 @@ public:
 
 protected:
 	/** Creates an icon type widget, including big tooltip link */
-	TSharedRef<SWidget> CreateIconWidget( const FString& IconToolTip, const FSlateBrush* IconBrush, const FSlateColor& IconColor, const FString& DocLink, const FString& DocExcerpt );
+	TSharedRef<SWidget> CreateIconWidget( const FText& IconToolTip, const FSlateBrush* IconBrush, const FSlateColor& IconColor, const FString& DocLink, const FString& DocExcerpt );
 	/** Create an icon type widget */
-	TSharedRef<SWidget> CreateIconWidget(const FString& IconToolTip, const FSlateBrush* IconBrush, const FSlateColor& IconColor);
+	TSharedRef<SWidget> CreateIconWidget(const FText& IconToolTip, const FSlateBrush* IconBrush, const FSlateColor& IconColor);
 
 
 	/* Create the text widget */
diff --git a/Engine/Source/Editor/GraphEditor/Public/SGraphPin.h b/Engine/Source/Editor/GraphEditor/Public/SGraphPin.h
--- a/Engine/Source/Editor/GraphEditor/Public/SGraphPin.h
+++ b/Engine/Source/Editor/GraphEditor/Public/SGraphPin.h
@@ -111,7 +111,7 @@ public:
 	FVector2D GetNodeOffset() const;
 
 protected:
-	FString GetPinLabel() const;
+	FText GetPinLabel() const;
 
 	/** If we should show the label on this pin */
 	EVisibility GetPinLabelVisibility() const
diff --git a/Engine/Source/Editor/GraphEditor/Public/SGraphPinComboBox.h b/Engine/Source/Editor/GraphEditor/Public/SGraphPinComboBox.h
--- a/Engine/Source/Editor/GraphEditor/Public/SGraphPinComboBox.h
+++ b/Engine/Source/Editor/GraphEditor/Public/SGraphPinComboBox.h
@@ -33,17 +33,17 @@ private:
 	void OnSelectionChangedInternal( TSharedPtr<int32> NewSelection, ESelectInfo::Type SelectInfo );
 
 	//Get string to display
-	FString OnGetVisibleTextInternal() const
+	FText OnGetVisibleTextInternal() const
 	{
-		return VisibleText.Get();
+		return FText::FromString(VisibleText.Get());
 	}
 
 	// Function to create each row of the combo widget
 	TSharedRef<ITableRow> OnGenerateComboWidget( TSharedPtr<int32> InComboIndex, const TSharedRef<STableViewBase>& OwnerTable );
 
-	FString GetRowString(int32 RowIndex) const
+	FText GetRowString(int32 RowIndex) const
 	{
-		return OnGetDisplayName.Execute(RowIndex);
+		return FText::FromString(OnGetDisplayName.Execute(RowIndex));
 	}
 
 private:
diff --git a/Engine/Source/Editor/GraphEditor/Public/SNameComboBox.h b/Engine/Source/Editor/GraphEditor/Public/SNameComboBox.h
--- a/Engine/Source/Editor/GraphEditor/Public/SNameComboBox.h
+++ b/Engine/Source/Editor/GraphEditor/Public/SNameComboBox.h
@@ -73,9 +73,9 @@ private:
 	void OnSelectionChanged(TSharedPtr<FName> Selection, ESelectInfo::Type SelectInfo);
 
 	/** Helper method to get the text for a given item in the combo box */
-	FString GetSelectedNameLabel() const;
+	FText GetSelectedNameLabel() const;
 
-	FString GetItemNameLabel(TSharedPtr<FName> StringItem) const;
+	FText GetItemNameLabel(TSharedPtr<FName> StringItem) const;
 
 private:
 
diff --git a/Engine/Source/Editor/Kismet/Private/BlueprintDetailsCustomization.cpp b/Engine/Source/Editor/Kismet/Private/BlueprintDetailsCustomization.cpp
--- a/Engine/Source/Editor/Kismet/Private/BlueprintDetailsCustomization.cpp
+++ b/Engine/Source/Editor/Kismet/Private/BlueprintDetailsCustomization.cpp
@@ -1175,7 +1175,7 @@ TSharedRef< ITableRow > FBlueprintVarActionDetails::MakeCategoryViewWidget( TSha
 {
 	return SNew(STableRow<TSharedPtr<FString>>, OwnerTable)
 		[
-			SNew(STextBlock) .Text(*Item.Get())
+			SNew(STextBlock) .Text(FText::FromString(*Item.Get()))
 		];
 }
 
@@ -2502,7 +2502,7 @@ void FBlueprintGraphActionDetails::SetNetFlags( TWeakObjectPtr<UK2Node_EditableP
 	}
 }
 
-FString FBlueprintGraphActionDetails::GetCurrentReplicatedEventString() const
+FText FBlueprintGraphActionDetails::GetCurrentReplicatedEventString() const
 {
 	const UK2Node_EditablePinBase * FunctionEntryNode = FunctionEntryNodePtr.Get();
 	const UK2Node_CustomEvent* CustomEvent = Cast<const UK2Node_CustomEvent>(FunctionEntryNode);
@@ -2518,7 +2518,7 @@ FString FBlueprintGraphActionDetails::GetCurrentReplicatedEventString() const
 		NetFlags = SuperFunction->FunctionFlags & ReplicatedNetMask;
 	}
 
-	return ReplicationSpecifierProperName(NetFlags).ToString();
+	return ReplicationSpecifierProperName(NetFlags);
 }
 
 bool FBaseBlueprintGraphActionDetails::ConditionallyCleanUpResultNode()
@@ -2743,7 +2743,7 @@ TSharedRef< ITableRow > FBlueprintDelegateActionDetails::MakeCategoryViewWidget(
 {
 	return SNew(STableRow<TSharedPtr<FString>>, OwnerTable)
 		[
-			SNew(STextBlock) .Text(*Item.Get())
+			SNew(STextBlock) .Text(FText::FromString(*Item.Get()))
 		];
 }
 
@@ -3445,7 +3445,7 @@ TSharedRef< ITableRow > FBlueprintGraphActionDetails::MakeCategoryViewWidget( TS
 {
 	return SNew(STableRow<TSharedPtr<FString>>, OwnerTable)
 		[
-			SNew(STextBlock) .Text(*Item.Get())
+			SNew(STextBlock) .Text(FText::FromString(*Item.Get()))
 		];
 }
 
@@ -3487,11 +3487,11 @@ TSharedRef<ITableRow> FBlueprintGraphActionDetails::HandleGenerateRowAccessSpeci
 		.Content()
 		[
 			SNew( STextBlock ) 
-				.Text( SpecifierName.IsValid() ? SpecifierName->LocalizedName.ToString() : FString() )
+				.Text( SpecifierName.IsValid() ? SpecifierName->LocalizedName : FText::GetEmpty() )
 		];
 }
 
-FString FBlueprintGraphActionDetails::GetCurrentAccessSpecifierName() const
+FText FBlueprintGraphActionDetails::GetCurrentAccessSpecifierName() const
 {
 	uint32 AccessSpecifierFlag = 0;
 	UK2Node_EditablePinBase * FunctionEntryNode = FunctionEntryNodePtr.Get();
@@ -3503,7 +3503,7 @@ FString FBlueprintGraphActionDetails::GetCurrentAccessSpecifierName() const
 	{
 		AccessSpecifierFlag = FUNC_AccessSpecifiers & CustomEventNode->FunctionFlags;
 	}
-	return AccessSpecifierProperName( AccessSpecifierFlag ).ToString();
+	return AccessSpecifierProperName( AccessSpecifierFlag );
 }
 
 bool FBlueprintGraphActionDetails::IsAccessSpecifierVisible() const
@@ -3862,8 +3862,8 @@ void FBlueprintInterfaceLayout::GenerateHeaderRowContent( FDetailWidgetRow& Node
 	[
 		SNew(STextBlock)
 			.Text( bShowsInheritedInterfaces ?
-			LOCTEXT("BlueprintInheritedInterfaceTitle", "Inherited Interfaces").ToString() :
-			LOCTEXT("BlueprintImplementedInterfaceTitle", "Implemented Interfaces").ToString() )
+			LOCTEXT("BlueprintInheritedInterfaceTitle", "Inherited Interfaces") :
+			LOCTEXT("BlueprintImplementedInterfaceTitle", "Implemented Interfaces") )
 			.Font( IDetailLayoutBuilder::GetDetailFont() )
 	];
 }
@@ -4181,7 +4181,7 @@ void FBlueprintGlobalOptionsDetails::CustomizeDetails(IDetailLayoutBuilder& Deta
 		.NameContent()
 		[
 			SNew(STextBlock)
-			.Text(LOCTEXT("BlueprintDetails_ParentClass", "Parent Class").ToString())
+			.Text(LOCTEXT("BlueprintDetails_ParentClass", "Parent Class"))
 			.Font(IDetailLayoutBuilder::GetDetailFont())
 		]
 		.ValueContent()
@@ -4303,7 +4303,7 @@ void FBlueprintComponentDetails::CustomizeDetails(IDetailLayoutBuilder& DetailLa
 		];
 
 		PopulateVariableCategories();
-		const FString CategoryTooltip = LOCTEXT("EditCategoryName_Tooltip", "The category of the variable; editing this will place the variable into another category or create a new one.").ToString();
+		const FText CategoryTooltip = LOCTEXT("EditCategoryName_Tooltip", "The category of the variable; editing this will place the variable into another category or create a new one.");
 
 		VariableCategory.AddCustomRow( LOCTEXT("BlueprintComponentDetails_VariableCategoryLabel", "Category") )
 		.NameContent()
@@ -4631,7 +4631,7 @@ TSharedRef< ITableRow > FBlueprintComponentDetails::MakeVariableCategoryViewWidg
 	return SNew(STableRow<TSharedPtr<FString>>, OwnerTable)
 	[
 		SNew(STextBlock)
-		.Text(*Item.Get())
+		.Text(FText::FromString(*Item.Get()))
 	];
 }
 
@@ -5079,7 +5079,7 @@ TSharedRef< ITableRow > FBlueprintDocumentationDetails::MakeExcerptViewWidget( T
 		SNew( STableRow<TSharedPtr<FString>>, OwnerTable )
 		[
 			SNew( STextBlock )
-			.Text( *Item.Get() )
+			.Text( FText::FromString(*Item.Get()) )
 		];
 }
 
diff --git a/Engine/Source/Editor/Kismet/Private/BlueprintDetailsCustomization.h b/Engine/Source/Editor/Kismet/Private/BlueprintDetailsCustomization.h
--- a/Engine/Source/Editor/Kismet/Private/BlueprintDetailsCustomization.h
+++ b/Engine/Source/Editor/Kismet/Private/BlueprintDetailsCustomization.h
@@ -447,7 +447,7 @@ private:
 	FText AccessSpecifierProperName( uint32 AccessSpecifierFlag ) const;
 	bool IsAccessSpecifierVisible() const;
 	TSharedRef<ITableRow> HandleGenerateRowAccessSpecifier( TSharedPtr<FAccessSpecifierLabel> SpecifierName, const TSharedRef<STableViewBase>& OwnerTable );
-	FString GetCurrentAccessSpecifierName() const;
+	FText GetCurrentAccessSpecifierName() const;
 	void OnAccessSpecifierSelected( TSharedPtr<FAccessSpecifierLabel> SpecifierName, ESelectInfo::Type SelectInfo );
 
 	bool GetInstanceColorVisibility() const;
@@ -469,7 +469,7 @@ private:
 		{}
 	};
 
-	FString GetCurrentReplicatedEventString() const;
+	FText GetCurrentReplicatedEventString() const;
 	FText ReplicationSpecifierProperName( uint32 ReplicationSpecifierFlag ) const;
 	TSharedRef<ITableRow> OnGenerateReplicationComboWidget( TSharedPtr<FReplicationSpecifierLabel> InNetFlag, const TSharedRef<STableViewBase>& OwnerTable );
 	
diff --git a/Engine/Source/Editor/Kismet/Private/BlueprintEditor.cpp b/Engine/Source/Editor/Kismet/Private/BlueprintEditor.cpp
--- a/Engine/Source/Editor/Kismet/Private/BlueprintEditor.cpp
+++ b/Engine/Source/Editor/Kismet/Private/BlueprintEditor.cpp
@@ -1023,16 +1023,16 @@ FGraphAppearanceInfo FBlueprintEditor::GetGraphAppearance(UEdGraph* InGraph) con
 	switch (Blueprint->BlueprintType)
 	{
 	case BPTYPE_LevelScript:
-		AppearanceInfo.CornerText = LOCTEXT("AppearanceCornerText_LevelScript", "LEVEL BLUEPRINT").ToString();
+		AppearanceInfo.CornerText = LOCTEXT("AppearanceCornerText_LevelScript", "LEVEL BLUEPRINT");
 		break;
 	case BPTYPE_MacroLibrary:
-		AppearanceInfo.CornerText = LOCTEXT("AppearanceCornerText_Macro", "MACRO").ToString();
+		AppearanceInfo.CornerText = LOCTEXT("AppearanceCornerText_Macro", "MACRO");
 		break;
 	case BPTYPE_Interface:
-		AppearanceInfo.CornerText = LOCTEXT("AppearanceCornerText_Interface", "INTERFACE").ToString();
+		AppearanceInfo.CornerText = LOCTEXT("AppearanceCornerText_Interface", "INTERFACE");
 		break;
 	default:
-		AppearanceInfo.CornerText = LOCTEXT("AppearanceCornerText_Blueprint", "BLUEPRINT").ToString();
+		AppearanceInfo.CornerText = LOCTEXT("AppearanceCornerText_Blueprint", "BLUEPRINT");
 		break;
 	}
 
@@ -6371,14 +6371,14 @@ void FBlueprintEditor::VariableListWasUpdated()
 	StartEditingDefaults(/*bAutoFocus=*/ false);
 }
 
-FString FBlueprintEditor::DefaultEditString() 
+FText FBlueprintEditor::DefaultEditString() 
 {
-	return LOCTEXT("BlueprintEditingDefaults", "Editing defaults").ToString();
+	return LOCTEXT("BlueprintEditingDefaults", "Editing defaults");
 }
 
-FString FBlueprintEditor::GetDefaultEditorTitle()
+FText FBlueprintEditor::GetDefaultEditorTitle()
 {
-	return LOCTEXT("BlueprintDefaultsTabTitle", "Blueprint Defaults").ToString();
+	return LOCTEXT("BlueprintDefaultsTabTitle", "Blueprint Defaults");
 }
 
 bool FBlueprintEditor::GetBoundsForSelectedNodes(class FSlateRect& Rect, float Padding)
@@ -6917,7 +6917,7 @@ TSharedPtr<ISCSEditorCustomization> FBlueprintEditor::CustomizeSCSEditor(USceneC
 	return TSharedPtr<ISCSEditorCustomization>();
 }
 
-FString FBlueprintEditor::GetPIEStatus() const
+FText FBlueprintEditor::GetPIEStatus() const
 {
 	UBlueprint* CurrentBlueprint = GetBlueprintObj();
 	UWorld *DebugWorld = NULL;
@@ -6947,14 +6947,14 @@ FString FBlueprintEditor::GetPIEStatus() const
 
 	if (NetMode == NM_ListenServer || NetMode == NM_DedicatedServer)
 	{
-		return TEXT("SERVER - SIMULATING");
+		return LOCTEXT("PIEStatusServerSimulating", "SERVER - SIMULATING");
 	}
 	else if (NetMode == NM_Client)
 	{
-		return TEXT("CLIENT - SIMULATING");
+		return LOCTEXT("PIEStatusClientSimulating", "CLIENT - SIMULATING");
 	}
 
-	return TEXT("SIMULATING");
+	return LOCTEXT("PIEStatusSimulating", "SIMULATING");
 }
 
 bool FBlueprintEditor::IsEditingAnimGraph() const
@@ -7014,9 +7014,9 @@ float FBlueprintEditor::GetInstructionTextOpacity(UEdGraph* InGraph) const
 	return 1.0f;
 }
 
-FString FBlueprintEditor::GetGraphDecorationString(UEdGraph* InGraph) const
+FText FBlueprintEditor::GetGraphDecorationString(UEdGraph* InGraph) const
 {
-	return TEXT("");
+	return FText::GetEmpty();
 }
 
 bool FBlueprintEditor::IsGraphInCurrentBlueprint(UEdGraph* InGraph) const
diff --git a/Engine/Source/Editor/Kismet/Private/Debugging/SKismetDebuggingView.cpp b/Engine/Source/Editor/Kismet/Private/Debugging/SKismetDebuggingView.cpp
--- a/Engine/Source/Editor/Kismet/Private/Debugging/SKismetDebuggingView.cpp
+++ b/Engine/Source/Editor/Kismet/Private/Debugging/SKismetDebuggingView.cpp
@@ -33,9 +33,9 @@ FText FDebugLineItem::GetDisplayName() const
 	return FText::GetEmpty();
 }
 
-FString FDebugLineItem::GetDescription() const
+FText FDebugLineItem::GetDescription() const
 {
-	return TEXT("");
+	return FText::GetEmpty();
 }
 
 TSharedRef<SWidget> FDebugLineItem::GenerateNameWidget()
@@ -136,9 +136,9 @@ protected:
 		return new FMessageLineItem(Message);
 	}
 
-	virtual FString GetDescription() const override
+	virtual FText GetDescription() const override
 	{
-		return Message;
+		return FText::FromString(Message);
 	}
 };
 
@@ -172,25 +172,25 @@ protected:
 	}
 protected:
 	virtual TSharedRef<SWidget> GenerateNameWidget() override;
-	virtual FString GetDescription() const override;
+	virtual FText GetDescription() const override;
 	virtual FText GetDisplayName() const override;
 	void OnNavigateToLatentNode();
 
 	class UEdGraphNode* FindAssociatedNode() const;
 };
 
-FString FLatentActionLineItem::GetDescription() const
+FText FLatentActionLineItem::GetDescription() const
 {
 	if (UObject* ParentObject = ParentObjectRef.Get())
 	{
 		if (UWorld* World = GEngine->GetWorldFromContextObject(ParentObject))
 		{
 			FLatentActionManager& LatentActionManager = World->GetLatentActionManager();
-			return LatentActionManager.GetDescription(ParentObject, UUID);
+			return FText::FromString(LatentActionManager.GetDescription(ParentObject, UUID));
 		}
 	}
 
-	return LOCTEXT("NullObject", "Object has been destroyed").ToString();
+	return LOCTEXT("NullObject", "Object has been destroyed");
 }
 
 TSharedRef<SWidget> FLatentActionLineItem::GenerateNameWidget()
@@ -294,7 +294,7 @@ public:
 		}
 	}
 protected:
-	virtual FString GetDescription() const override;
+	virtual FText GetDescription() const override;
 	virtual FText GetDisplayName() const override;
 	virtual TSharedRef<SWidget> GenerateNameWidget() override;
 
@@ -323,7 +323,7 @@ FText FWatchLineItem::GetDisplayName() const
 	}
 }
 
-FString FWatchLineItem::GetDescription() const
+FText FWatchLineItem::GetDescription() const
 {
 	if (UEdGraphPin* PinToWatch = Cast<UEdGraphPin>(ObjectRef.Get()))
 	{
@@ -342,16 +342,16 @@ FString FWatchLineItem::GetDescription() const
 				FString Description;
 				Property->ExportText_InContainer(/*ArrayElement=*/0, /*inout*/ Description, SelectedDataInPIE, SelectedDataInPIE, /*Parent=*/ NULL, PPF_PropertyWindow|PPF_BlueprintDebugView);
 
-				return Description;
+				return FText::FromString(Description);
 			}
 			else
 			{
-				return LOCTEXT("NothingToWatch", "Nothing to watch.").ToString();
+				return LOCTEXT("NothingToWatch", "Nothing to watch.");
 			}
 		}
 	}
 
-	return TEXT("");
+	return FText::GetEmpty();
 }
 
 TSharedRef<SWidget> FWatchLineItem::GenerateNameWidget()
@@ -836,15 +836,18 @@ protected:
 			];
 	}
 
-	FString GetVisitTime() const
+	FText GetVisitTime() const
 	{
 		const TSimpleRingBuffer<FKismetTraceSample>& TraceStack = FKismetDebugUtilities::GetTraceStack();
 		if (StackIndex < TraceStack.Num())
 		{
-			return FString::Printf(TEXT(" @ %.2f s"), TraceStack(StackIndex).ObservationTime - GStartTime);
+			static const FNumberFormattingOptions TimeFormatOptions = FNumberFormattingOptions()
+				.SetMinimumFractionalDigits(2)
+				.SetMaximumFractionalDigits(2);
+			return FText::Format(LOCTEXT("VisitTimeFmt", " @ {0} s"), FText::AsNumber(TraceStack(StackIndex).ObservationTime - GStartTime, &TimeFormatOptions));
 		}
 
-		return TEXT("");
+		return FText::GetEmpty();
 	}
 
 	FText GetContextObjectName() const
@@ -1011,19 +1014,19 @@ TSharedPtr<SWidget> SKismetDebuggingView::OnMakeContextMenu()
 	return MenuBuilder.MakeWidget();
 }
 
-FString SKismetDebuggingView::GetTopText() const
+FText SKismetDebuggingView::GetTopText() const
 {
 	if (GEditor->PlayWorld != NULL)
 	{
-		return LOCTEXT("ShowDebugForActors", "Showing debug info for selected actors").ToString();
+		return LOCTEXT("ShowDebugForActors", "Showing debug info for selected actors");
 	}
 	else if (BlueprintToWatchPtr.Get() != NULL)
 	{
-		return LOCTEXT("ShowDebugForBlueprint", "Showing debug info for this blueprint").ToString();
+		return LOCTEXT("ShowDebugForBlueprint", "Showing debug info for this blueprint");
 	}
 	else
 	{
-		return LOCTEXT("Idle", "Idle").ToString();
+		return LOCTEXT("Idle", "Idle");
 	}
 }
 
diff --git a/Engine/Source/Editor/Kismet/Private/Debugging/SKismetDebuggingView.h b/Engine/Source/Editor/Kismet/Private/Debugging/SKismetDebuggingView.h
--- a/Engine/Source/Editor/Kismet/Private/Debugging/SKismetDebuggingView.h
+++ b/Engine/Source/Editor/Kismet/Private/Debugging/SKismetDebuggingView.h
@@ -70,7 +70,7 @@ protected:
 	virtual FText GetDisplayName() const;
 
 	// @return The text to display in the value column, unless GenerateValueWidget is overridden
-	virtual FString GetDescription() const;
+	virtual FText GetDescription() const;
 private:
 	// Type of action (poor mans RTTI for the tree, really only used to accelerate Compare checks)
 	EDebugLineType Type;
@@ -95,7 +95,7 @@ public:
 	virtual void Tick( const FGeometry& AllottedGeometry, const double InCurrentTime, const float InDeltaTime ) override;
 	// End of SWidget interface
 protected:
-	FString GetTopText() const;
+	FText GetTopText() const;
 
 	TSharedRef<ITableRow> OnGenerateRowForWatchTree(FDebugTreeItemPtr InItem, const TSharedRef<STableViewBase>& OwnerTable);
 	void OnGetChildrenForWatchTree(FDebugTreeItemPtr InParent, TArray<FDebugTreeItemPtr>& OutChildren);
diff --git a/Engine/Source/Editor/Kismet/Private/FindInBlueprints.cpp b/Engine/Source/Editor/Kismet/Private/FindInBlueprints.cpp
--- a/Engine/Source/Editor/Kismet/Private/FindInBlueprints.cpp
+++ b/Engine/Source/Editor/Kismet/Private/FindInBlueprints.cpp
@@ -327,9 +327,9 @@ bool FFindInBlueprintsResult::ExtractContent(TSharedPtr< FJsonObject > InJsonNod
 	return bMatchesSearchResults;
 }
 
-FString FFindInBlueprintsResult::GetCategory() const
+FText FFindInBlueprintsResult::GetCategory() const
 {
-	return FString();
+	return FText::GetEmpty();
 }
 
 TSharedRef<SWidget> FFindInBlueprintsResult::CreateIcon() const
@@ -389,9 +389,9 @@ void FFindInBlueprintsResult::ExpandAllChildren( TSharedPtr< STreeView< TSharedP
 	}
 }
 
-FString FFindInBlueprintsResult::GetDisplayString() const
+FText FFindInBlueprintsResult::GetDisplayString() const
 {
-	return DisplayText.ToString();
+	return DisplayText;
 }
 
 //////////////////////////////////////////////////////////
@@ -519,30 +519,30 @@ bool FFindInBlueprintsGraphNode::ExtractContent(TSharedPtr< FJsonObject > InJson
 	return bMatchesSearchResults;
 }
 
-FString FFindInBlueprintsGraphNode::GetCategory() const
+FText FFindInBlueprintsGraphNode::GetCategory() const
 {
 	if(Class == UK2Node_CallFunction::StaticClass())
 	{
-		return LOCTEXT("CallFuctionCat", "Function Call").ToString();
+		return LOCTEXT("CallFuctionCat", "Function Call");
 	}
 	else if(Class == UK2Node_MacroInstance::StaticClass())
 	{
-		return LOCTEXT("MacroCategory", "Macro").ToString();
+		return LOCTEXT("MacroCategory", "Macro");
 	}
 	else if(Class == UK2Node_Event::StaticClass())
 	{
-		return LOCTEXT("EventCat", "Event").ToString();
+		return LOCTEXT("EventCat", "Event");
 	}
 	else if(Class == UK2Node_VariableGet::StaticClass())
 	{
-		return LOCTEXT("VariableGetCategory", "Variable Get").ToString();
+		return LOCTEXT("VariableGetCategory", "Variable Get");
 	}
 	else if(Class == UK2Node_VariableSet::StaticClass())
 	{
-		return LOCTEXT("VariableSetCategory", "Variable Set").ToString();
+		return LOCTEXT("VariableSetCategory", "Variable Set");
 	}
 
-	return LOCTEXT("NodeCategory", "Node").ToString();
+	return LOCTEXT("NodeCategory", "Node");
 }
 
 void FFindInBlueprintsGraphNode::FinalizeSearchData()
@@ -584,7 +584,7 @@ TSharedRef<SWidget> FFindInBlueprintsPin::CreateIcon() const
 	return 	SNew(SImage)
 		.Image(Brush)
 		.ColorAndOpacity(IconColor)
-		.ToolTipText(FindInBlueprintsHelpers::GetPinTypeAsString(PinType) );
+		.ToolTipText(FText::FromString(FindInBlueprintsHelpers::GetPinTypeAsString(PinType)));
 }
 
 bool FFindInBlueprintsPin::ParseSearchInfo(const TArray<FString> &InTokens, FString InKey, FText InValue, TSharedPtr< FFindInBlueprintsResult > InParent)
@@ -615,9 +615,9 @@ bool FFindInBlueprintsPin::ParseSearchInfo(const TArray<FString> &InTokens, FStr
 	return bMatchesTokens;
 }
 
-FString FFindInBlueprintsPin::GetCategory() const
+FText FFindInBlueprintsPin::GetCategory() const
 {
-	return LOCTEXT("PinCategory", "Pin").ToString();
+	return LOCTEXT("PinCategory", "Pin");
 }
 
 void FFindInBlueprintsPin::FinalizeSearchData()
@@ -664,7 +664,7 @@ TSharedRef<SWidget> FFindInBlueprintsProperty::CreateIcon() const
 	return 	SNew(SImage)
 		.Image(Brush)
 		.ColorAndOpacity(IconColor)
-		.ToolTipText( FindInBlueprintsHelpers::GetPinTypeAsString(PinType) );
+		.ToolTipText( FText::FromString(FindInBlueprintsHelpers::GetPinTypeAsString(PinType)) );
 }
 
 bool FFindInBlueprintsProperty::ParseSearchInfo(const TArray<FString> &InTokens, FString InKey, FText InValue, TSharedPtr< FFindInBlueprintsResult > InParent)
@@ -701,13 +701,13 @@ bool FFindInBlueprintsProperty::ParseSearchInfo(const TArray<FString> &InTokens,
 	return bMatchesTokens;
 }
 
-FString FFindInBlueprintsProperty::GetCategory() const
+FText FFindInBlueprintsProperty::GetCategory() const
 {
 	if(bIsSCSComponent)
 	{
-		return LOCTEXT("Component", "Component").ToString();
+		return LOCTEXT("Component", "Component");
 	}
-	return LOCTEXT("Variable", "Variable").ToString();
+	return LOCTEXT("Variable", "Variable");
 }
 
 void FFindInBlueprintsProperty::FinalizeSearchData()
@@ -846,17 +846,17 @@ bool FFindInBlueprintsGraph::ParseSearchInfo(const TArray<FString> &InTokens, FS
 	return bMatchesTokens;
 }
 
-FString FFindInBlueprintsGraph::GetCategory() const
+FText FFindInBlueprintsGraph::GetCategory() const
 {
 	if(GraphType == GT_Function)
 	{
-		return LOCTEXT("FunctionGraphCategory", "Function").ToString();
+		return LOCTEXT("FunctionGraphCategory", "Function");
 	}
 	else if(GraphType == GT_Macro)
 	{
-		return LOCTEXT("MacroGraphCategory", "Macro").ToString();
+		return LOCTEXT("MacroGraphCategory", "Macro");
 	}
-	return LOCTEXT("GraphCategory", "Graph").ToString();
+	return LOCTEXT("GraphCategory", "Graph");
 }
 
 ////////////////////////////////////
@@ -1057,7 +1057,7 @@ void SFindInBlueprints::Construct( const FArguments& InArgs, TSharedPtr<FBluepri
 				[
 					SNew(STextBlock)
 					.Font( FEditorStyle::GetFontStyle("AssetDiscoveryIndicator.DiscovertingAssetsFont") )
-					.Text( LOCTEXT("SearchResults", "Searching...").ToString() )
+					.Text( LOCTEXT("SearchResults", "Searching...") )
 					.Visibility(this, &SFindInBlueprints::GetSearchbarVisiblity)
 				]
 
@@ -1328,29 +1328,27 @@ TSharedRef<ITableRow> SFindInBlueprints::OnGenerateRow( FSearchResult InItem, co
 				[
 					SNew(STextBlock)
 					.Text(InItem.Get(), &FFindInBlueprintsResult::GetDisplayString)
-					.ToolTipText(LOCTEXT("BlueprintCatSearchToolTip", "Blueprint").ToString())
+					.ToolTipText(LOCTEXT("BlueprintCatSearchToolTip", "Blueprint"))
 				]
 			];
 	}
 	else // Functions/Event/Pin widget
 	{
-		FString CommentText;
+		FText CommentText = FText::GetEmpty();
 
 		if(!InItem->GetCommentText().IsEmpty())
 		{
 			FFormatNamedArguments Args;
 			Args.Add(TEXT("Comment"), FText::FromString(InItem->GetCommentText()));
 
-			CommentText = FText::Format(LOCTEXT("NodeComment", "Node Comment:[{Comment}]"), Args).ToString();
+			CommentText = FText::Format(LOCTEXT("NodeComment", "Node Comment:[{Comment}]"), Args);
 		}
 
-		FString Tooltip;
-
 		FFormatNamedArguments Args;
-		Args.Add(TEXT("Category"), FText::FromString(InItem->GetCategory()));
+		Args.Add(TEXT("Category"), InItem->GetCategory());
 		Args.Add(TEXT("DisplayTitle"), InItem->DisplayText);
 
-		Tooltip = FText::Format(LOCTEXT("BlueprintResultSearchToolTip", "{Category} : {DisplayTitle}"), Args).ToString();
+		FText Tooltip = FText::Format(LOCTEXT("BlueprintResultSearchToolTip", "{Category} : {DisplayTitle}"), Args);
 
 		return SNew( STableRow< TSharedPtr<FFindInBlueprintsResult> >, OwnerTable )
 			[
@@ -1461,7 +1459,7 @@ void SFindInBlueprints::OnCopyAction()
 		}
 
 		// Add the display string
-		SelectedText += SelectedItem->GetDisplayString();
+		SelectedText += SelectedItem->GetDisplayString().ToString();
 
 		// If there is a comment, add two indents and then the comment
 		FString CommentText = SelectedItem->GetCommentText();
diff --git a/Engine/Source/Editor/Kismet/Private/NativeCodeGenerationTool.cpp b/Engine/Source/Editor/Kismet/Private/NativeCodeGenerationTool.cpp
--- a/Engine/Source/Editor/Kismet/Private/NativeCodeGenerationTool.cpp
+++ b/Engine/Source/Editor/Kismet/Private/NativeCodeGenerationTool.cpp
@@ -255,7 +255,7 @@ public:
 				.AutoHeight()
 				[
 					SNew(STextBlock)
-					.Text(LOCTEXT("HeaderPath", "Header Path").ToString())
+					.Text(LOCTEXT("HeaderPath", "Header Path"))
 				]
 				+ SVerticalBox::Slot()
 				.Padding(4.0f)
@@ -272,7 +272,7 @@ public:
 				.AutoHeight()
 				[
 					SNew(STextBlock)
-					.Text(LOCTEXT("SourcePath", "Source Path").ToString())
+					.Text(LOCTEXT("SourcePath", "Source Path"))
 				]
 				+ SVerticalBox::Slot()
 				.Padding(4.0f)
diff --git a/Engine/Source/Editor/Kismet/Private/SBlueprintFavoritesPalette.cpp b/Engine/Source/Editor/Kismet/Private/SBlueprintFavoritesPalette.cpp
--- a/Engine/Source/Editor/Kismet/Private/SBlueprintFavoritesPalette.cpp
+++ b/Engine/Source/Editor/Kismet/Private/SBlueprintFavoritesPalette.cpp
@@ -146,7 +146,7 @@ struct SBlueprintFavoritesPaletteUtils
 				{
 					if (ToolTipText.IsEmpty())
 					{
-						ToolTipText = FText::FromString(FString::Printf(*LOCTEXT("ProfileAvailable", "Loads %s node favorites").ToString(), *FriendlyProfileName));
+						ToolTipText = FText::Format(LOCTEXT("ProfileAvailableFmt", "Loads {0} node favorites"), FText::FromString(FriendlyProfileName));
 					}
 
 					ProfileAction = FUIAction(
@@ -171,7 +171,7 @@ struct SBlueprintFavoritesPaletteUtils
 				TSharedRef<STextBlock> MenuTextEntry = SNew(STextBlock)
 					.TextStyle(MenuBuilder.GetStyleSet(), FEditorStyle::Join("Menu", ".Label"))
 					// @TODO how do we best localize this
-					.Text(FriendlyProfileName);
+					.Text(FText::FromString(FriendlyProfileName));
 
 				FSlateFontInfo ToolTipFont(FPaths::EngineContentDir() / TEXT("Slate/Fonts/Roboto-Regular.ttf"), 8);
 
diff --git a/Engine/Source/Editor/Kismet/Private/SBlueprintLibraryPalette.cpp b/Engine/Source/Editor/Kismet/Private/SBlueprintLibraryPalette.cpp
--- a/Engine/Source/Editor/Kismet/Private/SBlueprintLibraryPalette.cpp
+++ b/Engine/Source/Editor/Kismet/Private/SBlueprintLibraryPalette.cpp
@@ -397,13 +397,13 @@ TSharedRef<SWidget> SBlueprintLibraryPalette::ConstructClassFilterDropdownConten
 }
 
 //------------------------------------------------------------------------------
-FString SBlueprintLibraryPalette::GetFilterClassName() const
+FText SBlueprintLibraryPalette::GetFilterClassName() const
 {
-	FString FilterDisplayString = TEXT("All");
+	FText FilterDisplayString = LOCTEXT("All", "All");
 	if (FilterClass != NULL)
 	{
 		UBlueprint* Blueprint = UBlueprint::GetBlueprintFromClass(FilterClass.Get());
-		FilterDisplayString = (Blueprint != NULL) ? Blueprint->GetName() : FilterClass->GetName();
+		FilterDisplayString = FText::FromString((Blueprint != NULL) ? Blueprint->GetName() : FilterClass->GetName());
 	}
 
 	return FilterDisplayString;
diff --git a/Engine/Source/Editor/Kismet/Private/SBlueprintLibraryPalette.h b/Engine/Source/Editor/Kismet/Private/SBlueprintLibraryPalette.h
--- a/Engine/Source/Editor/Kismet/Private/SBlueprintLibraryPalette.h
+++ b/Engine/Source/Editor/Kismet/Private/SBlueprintLibraryPalette.h
@@ -52,7 +52,7 @@ private:
 	 * 
 	 * @return A string representing the selected class filter.
 	 */
-	FString GetFilterClassName() const;
+	FText GetFilterClassName() const;
 
 	/**
 	 * Clears the current class filter (if one is set), and refreshes the 
diff --git a/Engine/Source/Editor/Kismet/Private/SBlueprintPalette.cpp b/Engine/Source/Editor/Kismet/Private/SBlueprintPalette.cpp
--- a/Engine/Source/Editor/Kismet/Private/SBlueprintPalette.cpp
+++ b/Engine/Source/Editor/Kismet/Private/SBlueprintPalette.cpp
@@ -902,7 +902,7 @@ void SBlueprintPaletteItem::Construct(const FArguments& InArgs, FCreateWidgetFor
 	FString            IconToolTip = GraphAction->TooltipDescription;
 	FString			   IconDocLink, IconDocExcerpt;
 	GetPaletteItemIcon(GraphAction, Blueprint, IconBrush, IconColor, IconToolTip, IconDocLink, IconDocExcerpt);
-	TSharedRef<SWidget> IconWidget = CreateIconWidget(IconToolTip, IconBrush, IconColor, IconDocLink, IconDocExcerpt);
+	TSharedRef<SWidget> IconWidget = CreateIconWidget(FText::FromString(IconToolTip), IconBrush, IconColor, IconDocLink, IconDocExcerpt);
 
 	// Setup a meta tag for this node
 	FTutorialMetaData TagMeta("PaletteItem"); 
diff --git a/Engine/Source/Editor/Kismet/Private/SFilterableObjectList.cpp b/Engine/Source/Editor/Kismet/Private/SFilterableObjectList.cpp
--- a/Engine/Source/Editor/Kismet/Private/SFilterableObjectList.cpp
+++ b/Engine/Source/Editor/Kismet/Private/SFilterableObjectList.cpp
@@ -28,9 +28,9 @@ EVisibility SFilterableObjectList::GetFilterStatusVisibility() const
 	return IsFilterActive() ? EVisibility::Visible : EVisibility::Collapsed;
 }
 
-FString SFilterableObjectList::GetFilterStatusText() const
+FText SFilterableObjectList::GetFilterStatusText() const
 {
-	return FString::Printf(TEXT("Showing %d of %d"), FilteredObjectList.Num(), LoadedObjectList.Num());
+	return FText::Format(LOCTEXT("FilterStatus_ShowingXOfYFmt", "Showing {0} of {1}"), FText::AsNumber(FilteredObjectList.Num()), FText::AsNumber(LoadedObjectList.Num()));
 }
 
 FString SFilterableObjectList::GetSearchableText(UObject* Object)
@@ -78,7 +78,7 @@ void SFilterableObjectList::InternalConstruct()
 				.FillWidth(1)
 			[
 				SAssignNew(FilterTextBoxWidget, SSearchBox)
-					.ToolTipText( LOCTEXT("SearchBox_ToolTip", "Type words to search for").ToString() )
+					.ToolTipText( LOCTEXT("SearchBox_ToolTip", "Type words to search for") )
 					.OnTextChanged( this, &SFilterableObjectList::OnFilterTextChanged )
 			]
 
@@ -89,7 +89,7 @@ void SFilterableObjectList::InternalConstruct()
 				.VAlign(VAlign_Center)
 			[
 				SNew(SButton)
-				.ToolTipText( LOCTEXT("Refresh_ToolTip", "Search for new entries").ToString() )
+				.ToolTipText( LOCTEXT("Refresh_ToolTip", "Search for new entries") )
 				.OnClicked( this, &SFilterableObjectList::OnRefreshButtonClicked )
 				[
 					SNew(SImage)
@@ -126,7 +126,7 @@ SFilterableObjectList::FListRow SFilterableObjectList::GenerateRowForObject(UObj
 {
 	return
 		FListRow(
-			SNew(STextBlock).Text(InData->GetName()),
+			SNew(STextBlock).Text(FText::FromString(InData->GetName())),
 			FOnDragDetected()
 		);
 }
diff --git a/Engine/Source/Editor/Kismet/Private/SGraphTitleBar.cpp b/Engine/Source/Editor/Kismet/Private/SGraphTitleBar.cpp
--- a/Engine/Source/Editor/Kismet/Private/SGraphTitleBar.cpp
+++ b/Engine/Source/Editor/Kismet/Private/SGraphTitleBar.cpp
@@ -39,23 +39,24 @@ FText SGraphTitleBar::GetTitleForOneCrumb(const UEdGraph* Graph)
 	return FText::Format(LOCTEXT("BreadcrumbTitle", "{BreadcrumbDisplayName} {BreadcrumbNotes}"), Args);
 }
 
-FString SGraphTitleBar::GetTitleExtra() const
+FText SGraphTitleBar::GetTitleExtra() const
 {
 	check(EdGraphObj != NULL);
 
-	FString ExtraText = Kismet2Ptr.Pin()->GetGraphDecorationString(EdGraphObj);
+	FString ExtraText = Kismet2Ptr.Pin()->GetGraphDecorationString(EdGraphObj).ToString();
 
 	const bool bEditable = (Kismet2Ptr.Pin()->IsEditable(EdGraphObj));
 
 	if(!bEditable)
 	{
-		ExtraText += TEXT(" (READ-ONLY)");
+		ExtraText += TEXT(" (");
+		ExtraText += LOCTEXT("ReadOnlyWarningText", "READ-ONLY").ToString();
+		ExtraText += TEXT(")");
 	}
 
-	return ExtraText;
+	return FText::FromString(ExtraText);
 }
 
-
 EVisibility SGraphTitleBar::IsGraphBlueprintNameVisible() const
 {
 	return bShowBlueprintTitle ? EVisibility::Visible : EVisibility::Collapsed;
@@ -169,7 +170,7 @@ void SGraphTitleBar::Construct( const FArguments& InArgs )
 	if (BlueprintObj)
 	{
 		bShowBlueprintTitle = true;
-		BlueprintTitle = BlueprintObj->GetFriendlyName();
+		BlueprintTitle = FText::FromString(BlueprintObj->GetFriendlyName());
 
 		// Register for notifications to refresh UI
 		if( Kismet2Ptr.IsValid() )
@@ -231,7 +232,7 @@ UEdGraph* SGraphTitleBar::GetOuterGraph( UObject* Obj )
 	return NULL;
 }
 
-FString SGraphTitleBar::GetBlueprintTitle() const
+FText SGraphTitleBar::GetBlueprintTitle() const
 {
 	return BlueprintTitle;
 }
@@ -243,7 +244,7 @@ void SGraphTitleBar::Refresh()
 	{
 		if (UBlueprint* BlueprintObj = FBlueprintEditorUtils::FindBlueprintForGraph(this->EdGraphObj))
 		{
-			BlueprintTitle = BlueprintObj->GetFriendlyName();
+			BlueprintTitle = FText::FromString(BlueprintObj->GetFriendlyName());
 			RebuildBreadcrumbTrail();
 		}
 	}
diff --git a/Engine/Source/Editor/Kismet/Private/SGraphTitleBar.h b/Engine/Source/Editor/Kismet/Private/SGraphTitleBar.h
--- a/Engine/Source/Editor/Kismet/Private/SGraphTitleBar.h
+++ b/Engine/Source/Editor/Kismet/Private/SGraphTitleBar.h
@@ -55,7 +55,7 @@ protected:
 	bool bShowBlueprintTitle;
 
 	/** Blueprint title being displayed for toolbar */
-	FString BlueprintTitle;
+	FText BlueprintTitle;
 
 
 protected:
@@ -63,7 +63,7 @@ protected:
 	const FSlateBrush* GetTypeGlyph() const;
 
 	/** Get the extra title text */
-	FString GetTitleExtra() const;
+	FText GetTitleExtra() const;
 
 	/** Helper methods */
 	EVisibility IsGraphBlueprintNameVisible() const;
@@ -78,5 +78,5 @@ protected:
 	UEdGraph* GetOuterGraph( UObject* Obj );
 
 	/** Helper method used to show blueprint title in breadcrumbs */
-	FString GetBlueprintTitle() const;
+	FText GetBlueprintTitle() const;
 };
diff --git a/Engine/Source/Editor/Kismet/Private/SKismetInspector.cpp b/Engine/Source/Editor/Kismet/Private/SKismetInspector.cpp
--- a/Engine/Source/Editor/Kismet/Private/SKismetInspector.cpp
+++ b/Engine/Source/Editor/Kismet/Private/SKismetInspector.cpp
@@ -104,9 +104,9 @@ TSharedRef<SWidget> SKismetInspector::MakeContextualEditingWidget(struct FKismet
 	return ContextualEditingWidget;
 }
 
-FString SKismetInspector::GetContextualEditingWidgetTitle() const
+FText SKismetInspector::GetContextualEditingWidgetTitle() const
 {
-	FString Title = PropertyViewTitle;
+	FText Title = PropertyViewTitle;
 	if (Title.IsEmpty())
 	{
 		if (SelectedObjects.Num() == 1 && SelectedObjects[0].IsValid())
@@ -115,7 +115,7 @@ FString SKismetInspector::GetContextualEditingWidgetTitle() const
 
 			if (UEdGraphNode* Node = Cast<UEdGraphNode>(Object))
 			{
-				Title = Node->GetNodeTitle(ENodeTitleType::ListView).ToString();
+				Title = Node->GetNodeTitle(ENodeTitleType::ListView);
 			}
 			else if (USCS_Node* SCSNode = Cast<USCS_Node>(Object))
 			{
@@ -123,11 +123,11 @@ FString SKismetInspector::GetContextualEditingWidgetTitle() const
 				{
 					if (SCSNode->VariableName != NAME_None)
 					{
-						Title = FString::Printf(*LOCTEXT("TemplateFor", "Template for %s").ToString(), *(SCSNode->VariableName.ToString()));
+						Title = FText::Format(LOCTEXT("TemplateForFmt", "Template for {0}"), FText::FromName(SCSNode->VariableName));
 					}
 					else 
 					{
-						Title = FString::Printf(*LOCTEXT("Name_Template", "%s Template").ToString(), *(SCSNode->ComponentTemplate->GetClass()->GetName()));
+						Title = FText::Format(LOCTEXT("Name_TemplateFmt", "{0} Template"), FText::FromString(SCSNode->ComponentTemplate->GetClass()->GetName()));
 					}
 				}
 			}
@@ -136,13 +136,13 @@ FString SKismetInspector::GetContextualEditingWidgetTitle() const
 				// Edit the component template
 				if (UActorComponent* Template = K2Node->GetTemplateFromNode())
 				{
-					Title = FString::Printf(*LOCTEXT("Name_Template", "%s Template").ToString(), *(Template->GetClass()->GetName()));
+					Title = FText::Format(LOCTEXT("Name_TemplateFmt", "{0} Template"), FText::FromString(Template->GetClass()->GetName()));
 				}
 			}
 
 			if (Title.IsEmpty())
 			{
-				Title = UKismetSystemLibrary::GetDisplayName(Object);
+				Title = FText::FromString(UKismetSystemLibrary::GetDisplayName(Object));
 			}
 		}
 		else if (SelectedObjects.Num() > 1)
@@ -177,7 +177,7 @@ FString SKismetInspector::GetContextualEditingWidgetTitle() const
 
 			if (BaseClass)
 			{
-				Title = FString::Printf( *LOCTEXT("MultipleObjectsSelected", "%d %ss selected").ToString(), SelectedObjects.Num(), *BaseClass->GetName() );
+				Title = FText::Format(LOCTEXT("MultipleObjectsSelectedFmt", "{0} {1} selected"), FText::AsNumber(SelectedObjects.Num()), FText::FromString(BaseClass->GetName() + TEXT("s")));
 			}
 		}
 	}
@@ -266,7 +266,7 @@ void SKismetInspector::Construct(const FArguments& InArgs)
 	// Update based on the current (empty) selection set
 	TArray<UObject*> InitialSelectedObjects;
 	FKismetSelectionInfo SelectionInfo;
-	UpdateFromObjects(InitialSelectedObjects, SelectionInfo, SKismetInspector::FShowDetailsOptions(FString(), true));
+	UpdateFromObjects(InitialSelectedObjects, SelectionInfo, SKismetInspector::FShowDetailsOptions(FText::GetEmpty(), true));
 }
 
 void SKismetInspector::EnableComponentDetailsCustomization(bool bEnable)
diff --git a/Engine/Source/Editor/Kismet/Private/SMyBlueprint.cpp b/Engine/Source/Editor/Kismet/Private/SMyBlueprint.cpp
--- a/Engine/Source/Editor/Kismet/Private/SMyBlueprint.cpp
+++ b/Engine/Source/Editor/Kismet/Private/SMyBlueprint.cpp
@@ -1218,7 +1218,7 @@ void SMyBlueprint::OnActionSelectedHelper(TSharedPtr<FEdGraphSchemaAction> InAct
 			{
 				FGraphDisplayInfo DisplayInfo;
 				GraphAction->EdGraph->GetSchema()->GetGraphDisplayInformation(*GraphAction->EdGraph, DisplayInfo);
-				Inspector->ShowDetailsForSingleObject(GraphAction->EdGraph, SKismetInspector::FShowDetailsOptions(DisplayInfo.PlainName.ToString()));
+				Inspector->ShowDetailsForSingleObject(GraphAction->EdGraph, SKismetInspector::FShowDetailsOptions(DisplayInfo.PlainName));
 			}
 		}
 		else if (InAction->GetTypeId() == FEdGraphSchemaAction_K2Delegate::StaticGetTypeId())
@@ -1226,14 +1226,14 @@ void SMyBlueprint::OnActionSelectedHelper(TSharedPtr<FEdGraphSchemaAction> InAct
 			FEdGraphSchemaAction_K2Delegate* DelegateAction = (FEdGraphSchemaAction_K2Delegate*)InAction.Get();
 			if (UMulticastDelegateProperty* Property = DelegateAction->GetDelegatePoperty())
 			{
-				Inspector->ShowDetailsForSingleObject(Property, SKismetInspector::FShowDetailsOptions(Property->GetName()));
+				Inspector->ShowDetailsForSingleObject(Property, SKismetInspector::FShowDetailsOptions(FText::FromString(Property->GetName())));
 			}
 		}
 		else if (InAction->GetTypeId() == FEdGraphSchemaAction_K2Var::StaticGetTypeId())
 		{
 			FEdGraphSchemaAction_K2Var* VarAction = (FEdGraphSchemaAction_K2Var*)InAction.Get();
 
-			SKismetInspector::FShowDetailsOptions Options(VarAction->GetVariableName().ToString());
+			SKismetInspector::FShowDetailsOptions Options(FText::FromName(VarAction->GetVariableName()));
 
 			Inspector->ShowDetailsForSingleObject(VarAction->GetProperty(), Options);
 		}
@@ -1241,7 +1241,7 @@ void SMyBlueprint::OnActionSelectedHelper(TSharedPtr<FEdGraphSchemaAction> InAct
 		{
 			FEdGraphSchemaAction_K2LocalVar* VarAction = (FEdGraphSchemaAction_K2LocalVar*)InAction.Get();
 
-			SKismetInspector::FShowDetailsOptions Options(VarAction->GetVariableName().ToString());
+			SKismetInspector::FShowDetailsOptions Options(FText::FromName(VarAction->GetVariableName()));
 
 			Inspector->ShowDetailsForSingleObject(VarAction->GetProperty(), Options);
 		}
@@ -1249,7 +1249,7 @@ void SMyBlueprint::OnActionSelectedHelper(TSharedPtr<FEdGraphSchemaAction> InAct
 		{
 			FEdGraphSchemaAction_K2Enum* EnumAction = (FEdGraphSchemaAction_K2Enum*)InAction.Get();
 
-			SKismetInspector::FShowDetailsOptions Options(EnumAction->GetPathName().ToString());
+			SKismetInspector::FShowDetailsOptions Options(FText::FromName(EnumAction->GetPathName()));
 			Options.bForceRefresh = true;
 
 			Inspector->ShowDetailsForSingleObject(EnumAction->Enum, Options);
@@ -1258,7 +1258,7 @@ void SMyBlueprint::OnActionSelectedHelper(TSharedPtr<FEdGraphSchemaAction> InAct
 		{
 			FEdGraphSchemaAction_K2Struct* StructAction = (FEdGraphSchemaAction_K2Struct*)InAction.Get();
 
-			SKismetInspector::FShowDetailsOptions Options(StructAction->GetPathName().ToString());
+			SKismetInspector::FShowDetailsOptions Options(FText::FromName(StructAction->GetPathName()));
 			Options.bForceRefresh = true;
 
 			Inspector->ShowDetailsForSingleObject(StructAction->Struct, Options);
@@ -1266,7 +1266,7 @@ void SMyBlueprint::OnActionSelectedHelper(TSharedPtr<FEdGraphSchemaAction> InAct
 		else if (InAction->GetTypeId() == FEdGraphSchemaAction_K2TargetNode::StaticGetTypeId() || InAction->GetTypeId() == FEdGraphSchemaAction_K2Event::StaticGetTypeId())
 		{
 			FEdGraphSchemaAction_K2TargetNode* TargetNodeAction = (FEdGraphSchemaAction_K2TargetNode*)InAction.Get();
-			SKismetInspector::FShowDetailsOptions Options(TargetNodeAction->NodeTemplate->GetNodeTitle(ENodeTitleType::EditableTitle).ToString());
+			SKismetInspector::FShowDetailsOptions Options(TargetNodeAction->NodeTemplate->GetNodeTitle(ENodeTitleType::EditableTitle));
 			Inspector->ShowDetailsForSingleObject(TargetNodeAction->NodeTemplate, Options);
 		}
 		else
diff --git a/Engine/Source/Editor/Kismet/Private/SSCSEditor.cpp b/Engine/Source/Editor/Kismet/Private/SSCSEditor.cpp
--- a/Engine/Source/Editor/Kismet/Private/SSCSEditor.cpp
+++ b/Engine/Source/Editor/Kismet/Private/SSCSEditor.cpp
@@ -979,30 +979,30 @@ FText SSCS_RowWidget::GetMobilityToolTipText(TWeakObjectPtr<USCS_Node> SCSNode)
 	return MobilityToolTip;
 }
 
-FString SSCS_RowWidget::GetAssetName() const
+FText SSCS_RowWidget::GetAssetName() const
 {
-	FString AssetName(TEXT("None"));
+	FText AssetName = LOCTEXT("None", "None");
 	if(NodePtr.IsValid() && NodePtr->GetComponentTemplate())
 	{
 		UObject* Asset = FComponentAssetBrokerage::GetAssetFromComponent(NodePtr->GetComponentTemplate());
 		if(Asset != NULL)
 		{
-			AssetName = Asset->GetName();
+			AssetName = FText::FromString(Asset->GetName());
 		}
 	}
 
 	return AssetName;
 }
 
-FString SSCS_RowWidget::GetAssetNameTooltip() const
+FText SSCS_RowWidget::GetAssetNameTooltip() const
 {
-	FString AssetName(TEXT("None"));
+	FText AssetName = LOCTEXT("None", "None");
 	if(NodePtr.IsValid() && NodePtr->GetComponentTemplate())
 	{
 		UObject* Asset = FComponentAssetBrokerage::GetAssetFromComponent(NodePtr->GetComponentTemplate());
 		if(Asset != NULL)
 		{
-			AssetName = Asset->GetPathName();
+			AssetName = FText::FromString(Asset->GetPathName());
 		}
 	}
 
@@ -3072,7 +3072,7 @@ void SSCSEditor::RemoveComponentNode(FSCSEditorTreeNodePtrType InNodePtr)
 void SSCSEditor::UpdateSelectionFromNodes(const TArray<FSCSEditorTreeNodePtrType> &SelectedNodes )
 {
 	// Convert the selection set to an array of UObject* pointers
-	FString InspectorTitle;
+	FText InspectorTitle = FText::GetEmpty();
 	TArray<UObject*> InspectorObjects;
 	InspectorObjects.Empty(SelectedNodes.Num());
 	for (auto NodeIt = SelectedNodes.CreateConstIterator(); NodeIt; ++NodeIt)
@@ -3080,7 +3080,7 @@ void SSCSEditor::UpdateSelectionFromNodes(const TArray<FSCSEditorTreeNodePtrType
 		auto NodePtr = *NodeIt;
 		if(NodePtr.IsValid() && NodePtr->CanEditDefaults())
 		{
-			InspectorTitle = NodePtr->GetDisplayString();
+			InspectorTitle = FText::FromString(NodePtr->GetDisplayString());
 			InspectorObjects.Add(NodePtr->GetComponentTemplate());
 		}
 	}
diff --git a/Engine/Source/Editor/Kismet/Private/SSCSEditor.h b/Engine/Source/Editor/Kismet/Private/SSCSEditor.h
--- a/Engine/Source/Editor/Kismet/Private/SSCSEditor.h
+++ b/Engine/Source/Editor/Kismet/Private/SSCSEditor.h
@@ -232,8 +232,8 @@ public:
 	FString GetDocumentationLink() const;
 	FString GetDocumentationExcerptName() const;
 	
-	FString GetAssetName() const;
-	FString GetAssetNameTooltip() const;
+	FText GetAssetName() const;
+	FText GetAssetNameTooltip() const;
 	EVisibility GetAssetVisibility() const;
 
 	EVisibility GetRootLabelVisibility() const;
diff --git a/Engine/Source/Editor/Kismet/Private/WorkflowOrientedApp/SContentReference.cpp b/Engine/Source/Editor/Kismet/Private/WorkflowOrientedApp/SContentReference.cpp
--- a/Engine/Source/Editor/Kismet/Private/WorkflowOrientedApp/SContentReference.cpp
+++ b/Engine/Source/Editor/Kismet/Private/WorkflowOrientedApp/SContentReference.cpp
@@ -93,7 +93,7 @@ void SContentReference::Construct(const FArguments& InArgs)
 			.Visibility(this, &SContentReference::GetPickButtonVisibility)
 			.OnGetMenuContent( this, &SContentReference::MakeAssetPickerMenu )
 			.HasDownArrow(false)
-			.ToolTipText(LOCTEXT("PickAsset", "Pick an asset from a popup menu").ToString())
+			.ToolTipText(LOCTEXT("PickAsset", "Pick an asset from a popup menu"))
 			.ButtonContent()
 			[
 				SNew(SImage)
@@ -113,7 +113,7 @@ void SContentReference::Construct(const FArguments& InArgs)
 			.OnClicked(this, &SContentReference::OnClickFindButton)
 			.ContentPadding(0)
 			.Visibility(this, &SContentReference::GetFindButtonVisibility)
-			.ToolTipText(LOCTEXT("Find", "Find in content browser").ToString())
+			.ToolTipText(LOCTEXT("Find", "Find in content browser"))
 			[
 				SNew(SImage)
 				.Image( FEditorStyle::GetBrush(InArgs._Style, ".FindInContentBrowser") )
@@ -132,7 +132,7 @@ void SContentReference::Construct(const FArguments& InArgs)
 			.OnClicked(this, &SContentReference::OnClickClearButton)
 			.ContentPadding(1.f)
 			.Visibility( this, &SContentReference::GetClearButtonVisibility)
-			.ToolTipText(LOCTEXT("Clear", "Clear").ToString())
+			.ToolTipText(LOCTEXT("Clear", "Clear"))
 			[
 				SNew(SImage)
 				.Image( FEditorStyle::GetBrush(InArgs._Style, ".Clear") )
@@ -151,7 +151,7 @@ void SContentReference::Construct(const FArguments& InArgs)
 			.OnClicked(InArgs._OnClickedTools)
 			.ContentPadding(1.f)
 			.Visibility(this, &SContentReference::GetToolsButtonVisibility)
-			.ToolTipText(LOCTEXT("Tools", "Tools").ToString())
+			.ToolTipText(LOCTEXT("Tools", "Tools"))
 			[
 				SNew(SImage)
 				.Image( FEditorStyle::GetBrush(InArgs._Style, ".Tools") )
@@ -244,27 +244,27 @@ void SContentReference::OnAssetSelectedFromPicker(const FAssetData& AssetData)
 	OnSetReference.ExecuteIfBound(AssetData.GetAsset());
 }
 
-FString SContentReference::GetAssetShortName() const
+FText SContentReference::GetAssetShortName() const
 {
 	if (UObject* Asset = AssetReference.Get())
 	{
-		return Asset->GetName();
+		return FText::FromString(Asset->GetName());
 	}
 	else
 	{
-		return LOCTEXT("NullReference", "(None)").ToString();
+		return LOCTEXT("NullReference", "(None)");
 	}
 }
 
-FString SContentReference::GetAssetFullName() const
+FText SContentReference::GetAssetFullName() const
 {
 	if (UObject* Asset = AssetReference.Get())
 	{
-		return Asset->GetFullName();
+		return FText::FromString(Asset->GetFullName());
 	}
 	else
 	{
-		return LOCTEXT("NullReferenceTooltip", "(None)").ToString();
+		return LOCTEXT("NullReferenceTooltip", "(None)");
 	}
 }
 
diff --git a/Engine/Source/Editor/Kismet/Private/WorkflowOrientedApp/SModeWidget.cpp b/Engine/Source/Editor/Kismet/Private/WorkflowOrientedApp/SModeWidget.cpp
--- a/Engine/Source/Editor/Kismet/Private/WorkflowOrientedApp/SModeWidget.cpp
+++ b/Engine/Source/Editor/Kismet/Private/WorkflowOrientedApp/SModeWidget.cpp
@@ -72,7 +72,7 @@ void SModeWidget::Construct(const FArguments& InArgs, const FText& InText, const
 				.AutoWidth()
 				[
 					SNew(STextBlock)
-					.Text(ModeText.ToString())
+					.Text(ModeText)
 					.Font(this, &SModeWidget::GetDesiredTitleFont)
 				]
 
diff --git a/Engine/Source/Editor/Kismet/Public/BlueprintEditor.h b/Engine/Source/Editor/Kismet/Public/BlueprintEditor.h
--- a/Engine/Source/Editor/Kismet/Public/BlueprintEditor.h
+++ b/Engine/Source/Editor/Kismet/Public/BlueprintEditor.h
@@ -286,7 +286,7 @@ public:
 	EVisibility GetParentClassNameVisibility() const;
 
 	/** Returns our PIE Status - SIMULATING / SERVER / CLIENT */
-	FString GetPIEStatus() const;
+	FText GetPIEStatus() const;
 
 	/**
 	 * Util for finding a glyph for a graph
@@ -492,7 +492,7 @@ public:
 	void OnNodeTitleCommitted(const FText& NewText, ETextCommit::Type CommitInfo, UEdGraphNode* NodeBeingChanged);
 
 	/** Called by a graph title bar to get any extra information the editor would like to display */
-	virtual FString GetGraphDecorationString(UEdGraph* InGraph) const;
+	virtual FText GetGraphDecorationString(UEdGraph* InGraph) const;
 
 	/** Checks to see if the provided graph is contained within the current blueprint */
 	bool IsGraphInCurrentBlueprint(UEdGraph* InGraph) const;
@@ -573,7 +573,7 @@ protected:
 	void OnBlueprintUnloaded(UBlueprint* InBlueprint);
 
 	/** Get title for Inspector 2 tab*/
-	virtual FString GetDefaultEditorTitle();
+	virtual FText GetDefaultEditorTitle();
 
 	//@TODO: Should the breakpoint/watch modification operations be whole-blueprint, or current-graph?
 
@@ -819,7 +819,7 @@ protected:
 	void OnFinishedChangingProperties(const FPropertyChangedEvent& PropertyChangedEvent);
 
 	/** the string to show for edit defaults*/
-	static FString DefaultEditString();
+	static FText DefaultEditString();
 
 	/** On starting to rename node */
 	void OnRenameNode();
diff --git a/Engine/Source/Editor/Kismet/Public/FindInBlueprints.h b/Engine/Source/Editor/Kismet/Public/FindInBlueprints.h
--- a/Engine/Source/Editor/Kismet/Public/FindInBlueprints.h
+++ b/Engine/Source/Editor/Kismet/Public/FindInBlueprints.h
@@ -21,7 +21,7 @@ public:
 	virtual FReply OnClick();
 
 	/* Get Category for this search result */
-	virtual FString GetCategory() const;
+	virtual FText GetCategory() const;
 
 	/* Create an icon to represent the result */
 	virtual TSharedRef<SWidget>	CreateIcon() const;
@@ -81,7 +81,7 @@ public:
 	void ExpandAllChildren( TSharedPtr< STreeView< TSharedPtr< FFindInBlueprintsResult > > > InTreeView );
 
 	/** Returns the display string for the row */
-	FString GetDisplayString() const;
+	FText GetDisplayString() const;
 
 protected:
 
@@ -121,7 +121,7 @@ public:
 	virtual TSharedRef<SWidget>	CreateIcon() const override;
 	virtual bool ParseSearchInfo(const TArray<FString> &InTokens, FString InKey, FText InValue, TSharedPtr< FFindInBlueprintsResult > InParentOverride = NULL) override;
 	virtual bool ExtractContent(TSharedPtr< FJsonObject > InJsonNode, const TArray<FString>& InTokens) override;
-	virtual FString GetCategory() const override;
+	virtual FText GetCategory() const override;
 	virtual void FinalizeSearchData() override;
 	/** End FFindInBlueprintsResult Interface */
 
@@ -152,7 +152,7 @@ public:
 	/** FFindInBlueprintsResult Interface */
 	virtual TSharedRef<SWidget>	CreateIcon() const override;
 	virtual bool ParseSearchInfo(const TArray<FString> &InTokens, FString InKey, FText InValue, TSharedPtr< FFindInBlueprintsResult > InParentOverride = NULL) override;
-	virtual FString GetCategory() const override;
+	virtual FText GetCategory() const override;
 	virtual void FinalizeSearchData() override;
 	/** End FFindInBlueprintsResult Interface */
 protected:
@@ -179,7 +179,7 @@ public:
 	/** FFindInBlueprintsResult Interface */
 	virtual TSharedRef<SWidget>	CreateIcon() const override;
 	virtual bool ParseSearchInfo(const TArray<FString> &InTokens, FString InKey, FText InValue, TSharedPtr< FFindInBlueprintsResult > InParentOverride = NULL) override;
-	virtual FString GetCategory() const override;
+	virtual FText GetCategory() const override;
 	virtual void FinalizeSearchData() override;
 	/** End FFindInBlueprintsResult Interface */
 protected:
@@ -207,7 +207,7 @@ public:
 	virtual FReply OnClick() override;
 	virtual TSharedRef<SWidget>	CreateIcon() const override;
 	virtual bool ParseSearchInfo(const TArray<FString> &InTokens, FString InKey, FText InValue, TSharedPtr< FFindInBlueprintsResult > InParentOverride = NULL) override;
-	virtual FString GetCategory() const override;
+	virtual FText GetCategory() const override;
 	virtual bool ExtractContent(TSharedPtr< FJsonObject > InJsonNode, const TArray<FString>& InTokens) override;
 	/** End FFindInBlueprintsResult Interface */
 protected:
diff --git a/Engine/Source/Editor/Kismet/Public/SFilterableObjectList.h b/Engine/Source/Editor/Kismet/Public/SFilterableObjectList.h
--- a/Engine/Source/Editor/Kismet/Public/SFilterableObjectList.h
+++ b/Engine/Source/Editor/Kismet/Public/SFilterableObjectList.h
@@ -52,7 +52,7 @@ protected:
 	void OnFilterTextChanged(const FText& InFilterText);
 
 	EVisibility GetFilterStatusVisibility() const;
-	FString GetFilterStatusText() const;
+	FText GetFilterStatusText() const;
 
 	bool IsFilterActive() const;
 	void ReapplyFilter();
diff --git a/Engine/Source/Editor/Kismet/Public/SKismetInspector.h b/Engine/Source/Editor/Kismet/Public/SKismetInspector.h
--- a/Engine/Source/Editor/Kismet/Public/SKismetInspector.h
+++ b/Engine/Source/Editor/Kismet/Public/SKismetInspector.h
@@ -38,11 +38,11 @@ public:
 	/** Options for ShowDetails */
 	struct FShowDetailsOptions
 	{
-		FString ForcedTitle;
+		FText ForcedTitle;
 		bool bForceRefresh;
 		bool bHideFilterArea;
 
-		FShowDetailsOptions(const FString& InForcedTitle = FString(), bool bInForceRefresh = false)
+		FShowDetailsOptions(const FText& InForcedTitle = FText::GetEmpty(), bool bInForceRefresh = false)
 			:ForcedTitle(InForcedTitle)
 			,bForceRefresh(bInForceRefresh)
 			,bHideFilterArea(false)
@@ -74,7 +74,7 @@ protected:
 	TWeakPtr<FBlueprintEditor> Kismet2Ptr;
 
 	/** String used as the title above the property window */
-	FString PropertyViewTitle;
+	FText PropertyViewTitle;
 
 	/** Should we currently show the property view */
 	bool bShowInspectorPropertyView;
@@ -122,7 +122,7 @@ protected:
 	/**
 	 * Generates the text for the title in the contextual editing widget
 	 */
-	FString GetContextualEditingWidgetTitle() const;
+	FText GetContextualEditingWidgetTitle() const;
 
 	ECheckBoxState GetPublicViewCheckboxState() const;
 	void SetPublicViewCheckboxState(ECheckBoxState InIsChecked);
diff --git a/Engine/Source/Editor/Kismet/Public/WorkflowOrientedApp/SContentReference.h b/Engine/Source/Editor/Kismet/Public/WorkflowOrientedApp/SContentReference.h
--- a/Engine/Source/Editor/Kismet/Public/WorkflowOrientedApp/SContentReference.h
+++ b/Engine/Source/Editor/Kismet/Public/WorkflowOrientedApp/SContentReference.h
@@ -85,8 +85,8 @@ protected:
 	TSharedRef<SWidget> MakeAssetPickerMenu();
 	void OnAssetSelectedFromPicker(const FAssetData& AssetData);
 
-	FString GetAssetShortName() const;
-	FString GetAssetFullName() const;
+	FText GetAssetShortName() const;
+	FText GetAssetFullName() const;
 
 	const FSlateBrush* GetBorderImage() const;
 
diff --git a/Engine/Source/Editor/KismetCompiler/Private/KismetCompiler.cpp b/Engine/Source/Editor/KismetCompiler/Private/KismetCompiler.cpp
--- a/Engine/Source/Editor/KismetCompiler/Private/KismetCompiler.cpp
+++ b/Engine/Source/Editor/KismetCompiler/Private/KismetCompiler.cpp
@@ -902,7 +902,7 @@ void FKismetCompilerContext::CheckConnectionResponse(const FPinConnectionRespons
 {
 	if (!Response.CanSafeConnect())
 	{
-		MessageLog.Error(*FString::Printf(*LOCTEXT("FailedBuildingConnection_Error", "COMPILER ERROR: failed building connection with '%s' at @@").ToString(), *Response.Message), Node);
+		MessageLog.Error(*FString::Printf(*LOCTEXT("FailedBuildingConnection_Error", "COMPILER ERROR: failed building connection with '%s' at @@").ToString(), *Response.Message.ToString()), Node);
 	}
 }
 
diff --git a/Engine/Source/Editor/KismetWidgets/Private/SKismetLinearExpression.cpp b/Engine/Source/Editor/KismetWidgets/Private/SKismetLinearExpression.cpp
--- a/Engine/Source/Editor/KismetWidgets/Private/SKismetLinearExpression.cpp
+++ b/Engine/Source/Editor/KismetWidgets/Private/SKismetLinearExpression.cpp
@@ -254,7 +254,7 @@ TSharedRef<SWidget> SKismetLinearExpression::MakePinWidget(const UEdGraphPin* Pi
 		// Input pins with no links are displayed as their literals
 		return SNew(STextBlock)
 			.TextStyle( FEditorStyle::Get(), TEXT("KismetExpression.LiteralValue") )
-			.Text(Pin->GetDefaultAsString());
+			.Text(FText::FromString(Pin->GetDefaultAsString()));
 	}
 	else
 	{
diff --git a/Engine/Source/Editor/LandscapeEditor/Private/LandscapeEditorDetailCustomization_CopyPaste.cpp b/Engine/Source/Editor/LandscapeEditor/Private/LandscapeEditorDetailCustomization_CopyPaste.cpp
--- a/Engine/Source/Editor/LandscapeEditor/Private/LandscapeEditorDetailCustomization_CopyPaste.cpp
+++ b/Engine/Source/Editor/LandscapeEditor/Private/LandscapeEditorDetailCustomization_CopyPaste.cpp
@@ -141,7 +141,7 @@ void FLandscapeEditorDetailCustomization_CopyPaste::CustomizeDetails(IDetailLayo
 		[
 			SNew(STextBlock)
 			.Font(DetailBuilder.GetDetailFont())
-			.Text(FString().AppendChar(0xD7)) // Multiply sign
+			.Text(FText::FromString(FString().AppendChar(0xD7))) // Multiply sign
 		]
 		+ SHorizontalBox::Slot()
 		.FillWidth(1)
diff --git a/Engine/Source/Editor/LandscapeEditor/Private/LandscapeEditorDetailCustomization_NewLandscape.cpp b/Engine/Source/Editor/LandscapeEditor/Private/LandscapeEditorDetailCustomization_NewLandscape.cpp
--- a/Engine/Source/Editor/LandscapeEditor/Private/LandscapeEditorDetailCustomization_NewLandscape.cpp
+++ b/Engine/Source/Editor/LandscapeEditor/Private/LandscapeEditorDetailCustomization_NewLandscape.cpp
@@ -331,7 +331,7 @@ void FLandscapeEditorDetailCustomization_NewLandscape::CustomizeDetails(IDetailL
 		[
 			SNew(STextBlock)
 			.Font(DetailBuilder.GetDetailFont())
-			.Text(FString().AppendChar(0xD7)) // Multiply sign
+			.Text(FText::FromString(FString().AppendChar(0xD7))) // Multiply sign
 		]
 		+ SHorizontalBox::Slot()
 		.FillWidth(1)
@@ -390,7 +390,7 @@ void FLandscapeEditorDetailCustomization_NewLandscape::CustomizeDetails(IDetailL
 		[
 			SNew(STextBlock)
 			.Font(DetailBuilder.GetDetailFont())
-			.Text(FString().AppendChar(0xD7)) // Multiply sign
+			.Text(FText::FromString(FString().AppendChar(0xD7))) // Multiply sign
 		]
 		+ SHorizontalBox::Slot()
 		.FillWidth(1)
diff --git a/Engine/Source/Editor/Layers/Private/ActorLayerViewModel.cpp b/Engine/Source/Editor/Layers/Private/ActorLayerViewModel.cpp
--- a/Engine/Source/Editor/Layers/Private/ActorLayerViewModel.cpp
+++ b/Engine/Source/Editor/Layers/Private/ActorLayerViewModel.cpp
@@ -49,14 +49,14 @@ FName FActorLayerViewModel::GetFName() const
 }
 
 
-FString FActorLayerViewModel::GetName() const
+FText FActorLayerViewModel::GetName() const
 {
 	if( !Layer.IsValid() )
 	{
-		return LOCTEXT("Invalid layer Name", "").ToString();
+		return LOCTEXT("Invalid layer Name", "");
 	}
 
-	return Layer->LayerName.ToString();
+	return FText::FromName(Layer->LayerName);
 }
 
 
diff --git a/Engine/Source/Editor/Layers/Private/ActorLayerViewModel.h b/Engine/Source/Editor/Layers/Private/ActorLayerViewModel.h
--- a/Engine/Source/Editor/Layers/Private/ActorLayerViewModel.h
+++ b/Engine/Source/Editor/Layers/Private/ActorLayerViewModel.h
@@ -35,7 +35,7 @@ public:
 	FName GetFName() const;
 
 	/**	@return	The Layer's display name as a FString */
-	FString GetName() const;
+	FText GetName() const;
 
 	/**	@return Whether the Layer is visible */
 	bool IsVisible() const;
diff --git a/Engine/Source/Editor/Layers/Private/LayerViewModel.cpp b/Engine/Source/Editor/Layers/Private/LayerViewModel.cpp
--- a/Engine/Source/Editor/Layers/Private/LayerViewModel.cpp
+++ b/Engine/Source/Editor/Layers/Private/LayerViewModel.cpp
@@ -328,19 +328,19 @@ void FLayerViewModel::SelectActors( bool bSelect, bool bNotify, bool bSelectEven
 }
 
 
-FString FLayerViewModel::GetActorStatTotal( int32 StatsIndex  ) const
+FText FLayerViewModel::GetActorStatTotal( int32 StatsIndex ) const
 {
 	if( !Layer.IsValid() )
 	{
-		return TEXT( "0" );
+		return FText::AsNumber(0);
 	}
 
 	if( ActorStats.Num() <= StatsIndex )
 	{
-		return LOCTEXT("InvalidActorStatTotal", "Invalid").ToString();
+		return LOCTEXT("InvalidActorStatTotal", "Invalid");
 	}
 
-	return FString::Printf( TEXT( "%d" ), ActorStats[ StatsIndex ].Total );
+	return FText::AsNumber(ActorStats[ StatsIndex ].Total);
 }
 
 
diff --git a/Engine/Source/Editor/Layers/Private/LayerViewModel.h b/Engine/Source/Editor/Layers/Private/LayerViewModel.h
--- a/Engine/Source/Editor/Layers/Private/LayerViewModel.h
+++ b/Engine/Source/Editor/Layers/Private/LayerViewModel.h
@@ -146,7 +146,7 @@ public:
 	 *	@param	StatsIndex		The array index of the FLayerActorStats
 	 *	@return					The total actors as a FString			
 	 */
-	FString GetActorStatTotal( int32 StatsIndex  ) const;
+	FText GetActorStatTotal( int32 StatsIndex ) const;
 
 	/**	
 	 *	Selected the Actors assigned to the Layer that are of a certain type
diff --git a/Engine/Source/Editor/LevelEditor/Private/ComponentsTree.cpp b/Engine/Source/Editor/LevelEditor/Private/ComponentsTree.cpp
--- a/Engine/Source/Editor/LevelEditor/Private/ComponentsTree.cpp
+++ b/Engine/Source/Editor/LevelEditor/Private/ComponentsTree.cpp
@@ -36,14 +36,14 @@ void SComponentRowWidget::Construct( const FArguments& InArgs, FComponentTreeNod
 {
 	NodePtr = InNodePtr;
 
-	FString ComponentName;
+	FText ComponentName;
 	const FSlateBrush* ComponentIcon = FEditorStyle::GetBrush("SCS.NativeComponent");
 	FSlateColor ComponentColor = FLinearColor::White;
 
 	TSharedPtr<FComponentTreeNode> Node = NodePtr.Pin();
 	if(Node.IsValid() && Node->Component.IsValid())
 	{
-		ComponentName = Node->Component->GetName();
+		ComponentName = FText::FromString(Node->Component->GetName());
 		ComponentIcon = FClassIconFinder::FindIconForClass( Node->Component->GetClass(), TEXT("SCS.Component") );
 
 		// Color native and BP-made components in different colors
diff --git a/Engine/Source/Editor/LevelEditor/Private/SLevelEditorBuildAndSubmit.cpp b/Engine/Source/Editor/LevelEditor/Private/SLevelEditorBuildAndSubmit.cpp
--- a/Engine/Source/Editor/LevelEditor/Private/SLevelEditorBuildAndSubmit.cpp
+++ b/Engine/Source/Editor/LevelEditor/Private/SLevelEditorBuildAndSubmit.cpp
@@ -55,8 +55,8 @@ public:
 					[
 						SNew(STextBlock)
 						.Font(FEditorStyle::GetFontStyle("BuildAndSubmit.NormalFont"))
-						.Text(Item->Name)
-						.ToolTipText(Item->Name)
+						.Text(FText::FromString(Item->Name))
+						.ToolTipText(FText::FromString(Item->Name))
 					]
 				];
 		}
diff --git a/Engine/Source/Editor/LevelEditor/Private/SLevelViewport.cpp b/Engine/Source/Editor/LevelEditor/Private/SLevelViewport.cpp
--- a/Engine/Source/Editor/LevelEditor/Private/SLevelViewport.cpp
+++ b/Engine/Source/Editor/LevelEditor/Private/SLevelViewport.cpp
@@ -2487,7 +2487,7 @@ private:
 	FSlateColor GetBorderColorAndOpacity() const;
 
 	/** @return Gets the name of the preview actor.*/
-	FString OnReadText() const;
+	FText OnReadText() const;
 
 	/** @return Gets the Width of the preview viewport.*/
 	FOptionalSize OnReadWidth() const;
@@ -2802,15 +2802,15 @@ FSlateColor SActorPreview::GetBorderColorAndOpacity() const
 	return Color;
 }
 
-FString SActorPreview::OnReadText() const
+FText SActorPreview::OnReadText() const
 {
 	if( PreviewActorPtr.IsValid() )
 	{
-		return PreviewActorPtr.Get()->GetActorLabel();
+		return FText::FromString(PreviewActorPtr.Get()->GetActorLabel());
 	}
 	else
 	{
-		return TEXT("");
+		return FText::GetEmpty();
 	}
 }
 
