Added keybinding for zooming in graph editor per UE-6527 Blueprints: Can't zoom in and out without a mouse wheel (not keybindable)

diff --git a/Engine/Source/Editor/GraphEditor/Private/GraphEditorActions.cpp b/Engine/Source/Editor/GraphEditor/Private/GraphEditorActions.cpp
--- a/Engine/Source/Editor/GraphEditor/Private/GraphEditorActions.cpp
+++ b/Engine/Source/Editor/GraphEditor/Private/GraphEditorActions.cpp
@@ -74,6 +74,9 @@ void FGraphEditorCommandsImpl::RegisterCommands()
 	UI_COMMAND( CreateComment, "Create Comment", "Create a comment box", EUserInterfaceActionType::Button, FInputGesture(EKeys::C))
 
 	UI_COMMAND( FindInstancesOfCustomEvent, "Find Instances of Event", "Find the instances of this custom event", EUserInterfaceActionType::Button, FInputGesture() )
+
+	UI_COMMAND( ZoomIn, "Zoom In", "Zoom in on the graph editor", EUserInterfaceActionType::Button, FInputGesture(EKeys::Add))
+	UI_COMMAND( ZoomOut, "Zoom Out", "Zoom out from the graph editor", EUserInterfaceActionType::Button, FInputGesture(EKeys::Subtract))
 }
 
 
diff --git a/Engine/Source/Editor/GraphEditor/Private/GraphEditorModule.cpp b/Engine/Source/Editor/GraphEditor/Private/GraphEditorModule.cpp
--- a/Engine/Source/Editor/GraphEditor/Private/GraphEditorModule.cpp
+++ b/Engine/Source/Editor/GraphEditor/Private/GraphEditorModule.cpp
@@ -20,6 +20,8 @@ IMPLEMENT_MODULE(FGraphEditorModule, GraphEditor);
 
 void FGraphEditorModule::StartupModule()
 {
+	FGraphEditorCommands::Register();
+
 	TArray< TWeakPtr<SGraphEditor> >& Instances = SGraphEditor::AllInstances;
 	for (auto InstanceIt = SGraphEditor::AllInstances.CreateIterator(); InstanceIt; ++InstanceIt)
 	{
diff --git a/Engine/Source/Editor/GraphEditor/Private/SGraphPanel.cpp b/Engine/Source/Editor/GraphEditor/Private/SGraphPanel.cpp
--- a/Engine/Source/Editor/GraphEditor/Private/SGraphPanel.cpp
+++ b/Engine/Source/Editor/GraphEditor/Private/SGraphPanel.cpp
@@ -8,6 +8,10 @@
 #include "Editor/UnrealEd/Public/DragAndDrop/AssetDragDropOp.h"
 #include "Editor/UnrealEd/Public/DragAndDrop/LevelDragDropOp.h"
 
+#include "GraphEditorActions.h"
+#include "UICommandInfo.h"
+#include "InputGesture.h"
+
 #include "ConnectionDrawingPolicy.h"
 #include "BlueprintConnectionDrawingPolicy.h"
 #include "AnimGraphConnectionDrawingPolicy.h"
@@ -557,12 +561,12 @@ FReply SGraphPanel::OnKeyDown( const FGeometry& MyGeometry, const FKeyEvent& InK
 			UpdateSelectedNodesPositions(FVector2D(-GetSnapGridSize(),0.0f));
 			return FReply::Handled();
 		}
-		if ( InKeyEvent.GetKey() ==  EKeys::Subtract )
+		if(InKeyEvent.GetKey() == FGraphEditorCommands::Get().ZoomOut->GetActiveGesture()->Key)
 		{
 			ChangeZoomLevel(-1, CachedAllottedGeometryScaledSize / 2.f, InKeyEvent.IsControlDown());
 			return FReply::Handled();
 		}
-		if ( InKeyEvent.GetKey() ==  EKeys::Add )
+		if(InKeyEvent.GetKey() == FGraphEditorCommands::Get().ZoomIn->GetActiveGesture()->Key)
 		{
 			ChangeZoomLevel(+1, CachedAllottedGeometryScaledSize / 2.f, InKeyEvent.IsControlDown());
 			return FReply::Handled();
diff --git a/Engine/Source/Editor/GraphEditor/Public/GraphEditorActions.h b/Engine/Source/Editor/GraphEditor/Public/GraphEditorActions.h
--- a/Engine/Source/Editor/GraphEditor/Public/GraphEditorActions.h
+++ b/Engine/Source/Editor/GraphEditor/Public/GraphEditorActions.h
@@ -107,6 +107,10 @@ public:
 
 	// Find instances of a Custom event node
 	TSharedPtr< FUICommandInfo > FindInstancesOfCustomEvent;
+
+	// Zoom in and out on the graph editor
+	TSharedPtr< FUICommandInfo > ZoomIn;
+	TSharedPtr< FUICommandInfo > ZoomOut;
 };
 
 class GRAPHEDITOR_API FGraphEditorCommands
