Specified minimum desirable width for Graph enum dropdowns.

diff --git a/Engine/Source/Editor/GraphEditor/Private/KismetPins/SGraphPinEnum.cpp b/Engine/Source/Editor/GraphEditor/Private/KismetPins/SGraphPinEnum.cpp
--- a/Engine/Source/Editor/GraphEditor/Private/KismetPins/SGraphPinEnum.cpp
+++ b/Engine/Source/Editor/GraphEditor/Private/KismetPins/SGraphPinEnum.cpp
@@ -22,7 +22,7 @@ void SPinComboBox::Construct( const FArguments& InArgs )
 		[
 			// Wrap in configurable box to restrain height/width of menu
 			SNew(SBox)
-			.WidthOverride(200.0f)
+			.MinDesiredWidth(150.0f)
 			[
 				SNew( STextBlock ).ToolTipText(NSLOCTEXT("PinComboBox", "ToolTip", "Select enum values from the list"))
 				.Text( this, &SPinComboBox::OnGetVisibleTextInternal )
@@ -41,7 +41,7 @@ void SPinComboBox::Construct( const FArguments& InArgs )
 				[
 					SAssignNew( ComboList, SComboList )
 					.ListItemsSource( &ComboItemList )
-					.ItemHeight( 22 )
+//					.ItemHeight( 22 )
 					.OnGenerateRow( this, &SPinComboBox::OnGenerateComboWidget )
 					.OnSelectionChanged( this, &SPinComboBox::OnSelectionChangedInternal )
 				]
@@ -71,9 +71,13 @@ TSharedRef<ITableRow> SPinComboBox::OnGenerateComboWidget( TSharedPtr<int32> InC
 	return
 		SNew(STableRow< TSharedPtr<int32> >, OwnerTable)
 		[
-			SNew( STextBlock )
-			.Text( this, &SPinComboBox::GetRowString, RowIndex )
-			.Font( FEditorStyle::GetFontStyle( TEXT("PropertyWindow.NormalFont") ) )
+			SNew(SBox)
+			.MinDesiredWidth(150.0f)
+			[
+				SNew(STextBlock)
+				.Text( this, &SPinComboBox::GetRowString, RowIndex )
+				.Font( FEditorStyle::GetFontStyle( TEXT("PropertyWindow.NormalFont") ) )
+			]
 		];
 }
 
