Slate : RelativeLayoutScale is now computed in the SlatePrepass(). ComputeDesiredSize() now takes a float parameter. API-breaking change!

diff --git a/Engine/Source/Developer/LogVisualizer/Private/STimeSlider.cpp b/Engine/Source/Developer/LogVisualizer/Private/STimeSlider.cpp
--- a/Engine/Source/Developer/LogVisualizer/Private/STimeSlider.cpp
+++ b/Engine/Source/Developer/LogVisualizer/Private/STimeSlider.cpp
@@ -34,7 +34,7 @@ FReply STimeSlider::OnMouseMove( const FGeometry& MyGeometry, const FPointerEven
 	return TimeSliderController->OnMouseMove( SharedThis(this), MyGeometry, MouseEvent );
 }
 
-FVector2D STimeSlider::ComputeDesiredSize() const 
+FVector2D STimeSlider::ComputeDesiredSize( float ) const
 {
 	return FVector2D(100, 22);
 }
diff --git a/Engine/Source/Developer/LogVisualizer/Private/STimeSlider.h b/Engine/Source/Developer/LogVisualizer/Private/STimeSlider.h
--- a/Engine/Source/Developer/LogVisualizer/Private/STimeSlider.h
+++ b/Engine/Source/Developer/LogVisualizer/Private/STimeSlider.h
@@ -24,7 +24,7 @@ public:
 
 protected:
 	// SWidget interface
-	virtual FVector2D ComputeDesiredSize() const override;
+	virtual FVector2D ComputeDesiredSize(float) const override;
 	virtual int32 OnPaint( const FPaintArgs& Args, const FGeometry& AllottedGeometry, const FSlateRect& MyClippingRect, FSlateWindowElementList& OutDrawElements, int32 LayerId, const FWidgetStyle& InWidgetStyle, bool bParentEnabled ) const override;
 	virtual FReply OnMouseButtonDown( const FGeometry& MyGeometry, const FPointerEvent& MouseEvent ) override;
 	virtual FReply OnMouseButtonUp( const FGeometry& MyGeometry, const FPointerEvent& MouseEvent ) override;
diff --git a/Engine/Source/Developer/LogVisualizer/Private/STimelineBar.cpp b/Engine/Source/Developer/LogVisualizer/Private/STimelineBar.cpp
--- a/Engine/Source/Developer/LogVisualizer/Private/STimelineBar.cpp
+++ b/Engine/Source/Developer/LogVisualizer/Private/STimelineBar.cpp
@@ -257,7 +257,7 @@ void STimelineBar::Construct(const FArguments& InArgs, TSharedPtr<FSequencerTime
 	TRange<float> LocalViewRange = TimeSliderController->GetTimeSliderArgs().ViewRange.Get();
 }
 
-FVector2D STimelineBar::ComputeDesiredSize() const
+FVector2D STimelineBar::ComputeDesiredSize(float) const
 {
 	return FVector2D(5000.0f, 20.0f);
 }
diff --git a/Engine/Source/Developer/LogVisualizer/Private/STimelineBar.h b/Engine/Source/Developer/LogVisualizer/Private/STimelineBar.h
--- a/Engine/Source/Developer/LogVisualizer/Private/STimelineBar.h
+++ b/Engine/Source/Developer/LogVisualizer/Private/STimelineBar.h
@@ -25,7 +25,7 @@ public:
 
 	void Construct(const FArguments& InArgs, TSharedPtr<FSequencerTimeSliderController> InTimeSliderController, TSharedPtr<class STimeline> TimelineOwner);
 	virtual int32 OnPaint(const FPaintArgs& Args, const FGeometry& AllottedGeometry, const FSlateRect& MyClippingRect, FSlateWindowElementList& OutDrawElements, int32 LayerId, const FWidgetStyle& InWidgetStyle, bool bParentEnabled) const override;
-	FVector2D ComputeDesiredSize() const;
+	FVector2D ComputeDesiredSize(float) const;
 	void SnapScrubPosition(float ScrubPosition);
 	void SnapScrubPosition(int32 NewItemIndex);
 	int32 GetClosestItem(float Time) const;
diff --git a/Engine/Source/Developer/Profiler/Private/Widgets/SDataGraph.h b/Engine/Source/Developer/Profiler/Private/Widgets/SDataGraph.h
--- a/Engine/Source/Developer/Profiler/Private/Widgets/SDataGraph.h
+++ b/Engine/Source/Developer/Profiler/Private/Widgets/SDataGraph.h
@@ -324,7 +324,7 @@ public:
 	
 
 protected:
-	virtual FVector2D ComputeDesiredSize() const override
+	virtual FVector2D ComputeDesiredSize(float) const override
 	{
 		return FVector2D(16.0f,16.0f);
 	}
diff --git a/Engine/Source/Developer/Profiler/Private/Widgets/SHistogram.h b/Engine/Source/Developer/Profiler/Private/Widgets/SHistogram.h
--- a/Engine/Source/Developer/Profiler/Private/Widgets/SHistogram.h
+++ b/Engine/Source/Developer/Profiler/Private/Widgets/SHistogram.h
@@ -105,7 +105,7 @@ public:
 
 	virtual int32 OnPaint( const FPaintArgs& Args, const FGeometry& AllottedGeometry, const FSlateRect& MyClippingRect, FSlateWindowElementList& OutDrawElements, int32 LayerId, const FWidgetStyle& InWidgetStyle, bool bParentEnabled ) const override;
 
-	virtual FVector2D ComputeDesiredSize() const override
+	virtual FVector2D ComputeDesiredSize(float) const override
 	{
 		//return SCompoundWidget::ComputeDesiredSize();
 		return FVector2D(16.0f,16.0f);
diff --git a/Engine/Source/Developer/RealtimeProfiler/Private/SRealtimeProfilerLineGraph.cpp b/Engine/Source/Developer/RealtimeProfiler/Private/SRealtimeProfilerLineGraph.cpp
--- a/Engine/Source/Developer/RealtimeProfiler/Private/SRealtimeProfilerLineGraph.cpp
+++ b/Engine/Source/Developer/RealtimeProfiler/Private/SRealtimeProfilerLineGraph.cpp
@@ -89,7 +89,7 @@ void SRealtimeProfilerLineGraph::Construct(const FArguments& InArgs)
 	
 }
 
-FVector2D SRealtimeProfilerLineGraph::ComputeDesiredSize() const
+FVector2D SRealtimeProfilerLineGraph::ComputeDesiredSize( float ) const
 {
 	return FVector2D(128,64);
 }
diff --git a/Engine/Source/Developer/RealtimeProfiler/Private/SRealtimeProfilerLineGraph.h b/Engine/Source/Developer/RealtimeProfiler/Private/SRealtimeProfilerLineGraph.h
--- a/Engine/Source/Developer/RealtimeProfiler/Private/SRealtimeProfilerLineGraph.h
+++ b/Engine/Source/Developer/RealtimeProfiler/Private/SRealtimeProfilerLineGraph.h
@@ -95,7 +95,7 @@ private:
 
 	// SWidget overrides
 	virtual void Tick( const FGeometry& AllottedGeometry, const double InCurrentTime, const float InDeltaTime ) override;
-	virtual FVector2D ComputeDesiredSize() const override;
+	virtual FVector2D ComputeDesiredSize(float) const override;
 	virtual int32 OnPaint( const FPaintArgs& Args, const FGeometry& AllottedGeometry, const FSlateRect& MyClippingRect, FSlateWindowElementList& OutDrawElements, int32 LayerId, const FWidgetStyle& InWidgetStyle, bool bParentEnabled ) const override;
 	virtual FReply OnMouseButtonDown( const FGeometry& MyGeometry, const FPointerEvent& MouseEvent ) override;
 	virtual FReply OnMouseButtonUp( const FGeometry& MyGeometry, const FPointerEvent& MouseEvent ) override;
diff --git a/Engine/Source/Developer/TaskGraph/Private/SGraphBar.cpp b/Engine/Source/Developer/TaskGraph/Private/SGraphBar.cpp
--- a/Engine/Source/Developer/TaskGraph/Private/SGraphBar.cpp
+++ b/Engine/Source/Developer/TaskGraph/Private/SGraphBar.cpp
@@ -209,7 +209,7 @@ FReply SGraphBar::OnMouseMove( const FGeometry& MyGeometry, const FPointerEvent&
 	return SLeafWidget::OnMouseMove( MyGeometry, MouseEvent );
 }
 
-FVector2D SGraphBar::ComputeDesiredSize() const
+FVector2D SGraphBar::ComputeDesiredSize( float ) const
 {
 	return FVector2D( 500.0f, 50.0f );
 }
diff --git a/Engine/Source/Developer/TaskGraph/Private/STimeline.cpp b/Engine/Source/Developer/TaskGraph/Private/STimeline.cpp
--- a/Engine/Source/Developer/TaskGraph/Private/STimeline.cpp
+++ b/Engine/Source/Developer/TaskGraph/Private/STimeline.cpp
@@ -177,7 +177,7 @@ FReply STimeline::OnMouseMove( const FGeometry& MyGeometry, const FPointerEvent&
 	return SCompoundWidget::OnMouseMove( MyGeometry, MouseEvent );
 }
 
-FVector2D STimeline::ComputeDesiredSize() const
+FVector2D STimeline::ComputeDesiredSize( float ) const
 {
 	return FVector2D(500.0f, 50.0f);
 }
diff --git a/Engine/Source/Developer/TaskGraph/Public/SGraphBar.h b/Engine/Source/Developer/TaskGraph/Public/SGraphBar.h
--- a/Engine/Source/Developer/TaskGraph/Public/SGraphBar.h
+++ b/Engine/Source/Developer/TaskGraph/Public/SGraphBar.h
@@ -37,7 +37,7 @@ public:
 	virtual void Tick( const FGeometry& AllottedGeometry, const double InCurrentTime, const float InDeltaTime ) override;
 	virtual FReply OnMouseButtonDown( const FGeometry& MyGeometry, const FPointerEvent& MouseEvent ) override;
 	virtual FReply OnMouseMove( const FGeometry& MyGeometry, const FPointerEvent& MouseEvent ) override;
-	FVector2D ComputeDesiredSize() const override;
+	FVector2D ComputeDesiredSize(float) const override;
 	// End of SWidget interface
 
 	/**
diff --git a/Engine/Source/Developer/TaskGraph/Public/STimeline.h b/Engine/Source/Developer/TaskGraph/Public/STimeline.h
--- a/Engine/Source/Developer/TaskGraph/Public/STimeline.h
+++ b/Engine/Source/Developer/TaskGraph/Public/STimeline.h
@@ -35,7 +35,7 @@ public:
 	virtual void Tick( const FGeometry& AllottedGeometry, const double InCurrentTime, const float InDeltaTime ) override;
 	virtual FReply OnMouseButtonDown( const FGeometry& MyGeometry, const FPointerEvent& MouseEvent ) override;
 	virtual FReply OnMouseMove( const FGeometry& MyGeometry, const FPointerEvent& MouseEvent ) override;
-	FVector2D ComputeDesiredSize() const override;
+	FVector2D ComputeDesiredSize(float) const override;
 	// End of SWidget interface
 
 	/**
diff --git a/Engine/Source/Editor/EnvironmentQueryEditor/Private/STestFunctionWidget.cpp b/Engine/Source/Editor/EnvironmentQueryEditor/Private/STestFunctionWidget.cpp
--- a/Engine/Source/Editor/EnvironmentQueryEditor/Private/STestFunctionWidget.cpp
+++ b/Engine/Source/Editor/EnvironmentQueryEditor/Private/STestFunctionWidget.cpp
@@ -187,7 +187,7 @@ int32 STestFunctionWidget::OnPaint(const FPaintArgs& Args, const FGeometry& Allo
 	return LayerId;
 }
 
-FVector2D STestFunctionWidget::ComputeDesiredSize() const
+FVector2D STestFunctionWidget::ComputeDesiredSize( float ) const
 {
 	return FVector2D(128, 92);
 }
diff --git a/Engine/Source/Editor/GraphEditor/Private/SGraphNode.cpp b/Engine/Source/Editor/GraphEditor/Private/SGraphNode.cpp
--- a/Engine/Source/Editor/GraphEditor/Private/SGraphNode.cpp
+++ b/Engine/Source/Editor/GraphEditor/Private/SGraphNode.cpp
@@ -428,6 +428,7 @@ bool SGraphNode::IsSelectedExclusively() const
 void SGraphNode::SetOwner( const TSharedRef<SGraphPanel>& OwnerPanel )
 {
 	check( !OwnerGraphPanelPtr.IsValid() );
+	SetParentPanel(OwnerPanel);
 	OwnerGraphPanelPtr = OwnerPanel;
 	GraphNode->NodeWidget = SharedThis(this);
 
diff --git a/Engine/Source/Editor/GraphEditor/Private/SGraphNodeComment.cpp b/Engine/Source/Editor/GraphEditor/Private/SGraphNodeComment.cpp
--- a/Engine/Source/Editor/GraphEditor/Private/SGraphNodeComment.cpp
+++ b/Engine/Source/Editor/GraphEditor/Private/SGraphNodeComment.cpp
@@ -202,7 +202,7 @@ void SGraphNodeComment::UpdateGraphNode()
 	];
 }
 
-FVector2D SGraphNodeComment::ComputeDesiredSize() const
+FVector2D SGraphNodeComment::ComputeDesiredSize( float ) const
 {
 	return UserSize;
 }
diff --git a/Engine/Source/Editor/GraphEditor/Private/SGraphNodeDocumentation.cpp b/Engine/Source/Editor/GraphEditor/Private/SGraphNodeDocumentation.cpp
--- a/Engine/Source/Editor/GraphEditor/Private/SGraphNodeDocumentation.cpp
+++ b/Engine/Source/Editor/GraphEditor/Private/SGraphNodeDocumentation.cpp
@@ -275,7 +275,7 @@ float SGraphNodeDocumentation::GetDocumentationWrapWidth() const
 	return UserSize.X - GraphNodeDocumentationDefs::LineWrapAdjustment;
 }
 
-FVector2D SGraphNodeDocumentation::ComputeDesiredSize() const
+FVector2D SGraphNodeDocumentation::ComputeDesiredSize( float ) const
 {
 	return FVector2D( UserSize.X, UserSize.Y + GetTitleBarHeight() );
 }
diff --git a/Engine/Source/Editor/GraphEditor/Private/SGraphPanel.cpp b/Engine/Source/Editor/GraphEditor/Private/SGraphPanel.cpp
--- a/Engine/Source/Editor/GraphEditor/Private/SGraphPanel.cpp
+++ b/Engine/Source/Editor/GraphEditor/Private/SGraphPanel.cpp
@@ -299,9 +299,9 @@ int32 SGraphPanel::OnPaint( const FPaintArgs& Args, const FGeometry& AllottedGeo
 							if(OverlayInfo.Widget->GetVisibility() == EVisibility::Visible)
 							{
 								// call SlatePrepass as these widgets are not in the 'normal' child hierarchy
-								OverlayInfo.Widget->SlatePrepass();
+								OverlayInfo.Widget->SlatePrepass(AllottedGeometry.GetAccumulatedLayoutTransform().GetScale());
 
-								const FGeometry WidgetGeometry = CurWidget.Geometry.MakeChild(OverlayInfo.OverlayOffset, OverlayInfo.Widget->GetDesiredSize(), 1.f);
+								const FGeometry WidgetGeometry = CurWidget.Geometry.MakeChild(OverlayInfo.OverlayOffset, OverlayInfo.Widget->GetDesiredSize());
 
 								OverlayInfo.Widget->Paint(Args.WithNewParent(this), WidgetGeometry, MyClippingRect, OutDrawElements, CurWidgetsMaxLayerId, InWidgetStyle, bParentEnabled);
 							}
@@ -608,7 +608,7 @@ void SGraphPanel::ArrangeChildrenForContextMenuSummon(const FGeometry& AllottedG
 		const TSharedRef<SNode>& SomeChild = VisibleChildren[ChildIndex];
 		if (!SomeChild->RequiresSecondPassLayout())
 		{
-			ArrangedChildren.AddWidget( AllottedGeometry.MakeChild( SomeChild, SomeChild->GetPosition() - ViewOffset, SomeChild->GetDesiredSizeForMarquee(), GetZoomAmount() ) );
+			ArrangedChildren.AddWidget(AllottedGeometry.MakeChild(SomeChild, SomeChild->GetPosition() - ViewOffset, SomeChild->GetDesiredSizeForMarquee(), GetZoomAmount()));
 		}
 	}
 
@@ -619,7 +619,7 @@ void SGraphPanel::ArrangeChildrenForContextMenuSummon(const FGeometry& AllottedG
 		if (SomeChild->RequiresSecondPassLayout())
 		{
 			SomeChild->PerformSecondPassLayout(NodeToWidgetLookup);
-			ArrangedChildren.AddWidget( AllottedGeometry.MakeChild( SomeChild, SomeChild->GetPosition() - ViewOffset, SomeChild->GetDesiredSizeForMarquee(), GetZoomAmount() ) );
+			ArrangedChildren.AddWidget(AllottedGeometry.MakeChild(SomeChild, SomeChild->GetPosition() - ViewOffset, SomeChild->GetDesiredSizeForMarquee(), GetZoomAmount()));
 		}
 	}
 }
diff --git a/Engine/Source/Editor/GraphEditor/Private/SNodePanel.cpp b/Engine/Source/Editor/GraphEditor/Private/SNodePanel.cpp
--- a/Engine/Source/Editor/GraphEditor/Private/SNodePanel.cpp
+++ b/Engine/Source/Editor/GraphEditor/Private/SNodePanel.cpp
@@ -219,7 +219,7 @@ void SNodePanel::ArrangeChildNodes(const FGeometry& AllottedGeometry, FArrangedC
 		const TSharedRef<SNode>& SomeChild = ChildrenToArrange[ChildIndex];
 		if (!SomeChild->RequiresSecondPassLayout())
 		{
-			ArrangedChildren.AddWidget( AllottedGeometry.MakeChild( SomeChild, SomeChild->GetPosition() - ViewOffset, SomeChild->GetDesiredSize(), GetZoomAmount() ) );
+			ArrangedChildren.AddWidget(AllottedGeometry.MakeChild(SomeChild, SomeChild->GetPosition() - ViewOffset, SomeChild->GetDesiredSize(), GetZoomAmount()));
 		}
 	}
 
@@ -230,12 +230,12 @@ void SNodePanel::ArrangeChildNodes(const FGeometry& AllottedGeometry, FArrangedC
 		if (SomeChild->RequiresSecondPassLayout())
 		{
 			SomeChild->PerformSecondPassLayout(NodeToWidgetLookup);
-			ArrangedChildren.AddWidget( AllottedGeometry.MakeChild( SomeChild, SomeChild->GetPosition() - ViewOffset, SomeChild->GetDesiredSize(), GetZoomAmount() ) );
+			ArrangedChildren.AddWidget(AllottedGeometry.MakeChild(SomeChild, SomeChild->GetPosition() - ViewOffset, SomeChild->GetDesiredSize(), GetZoomAmount()));
 		}
 	}
 }
 
-FVector2D SNodePanel::ComputeDesiredSize() const
+FVector2D SNodePanel::ComputeDesiredSize( float ) const
 {	
 	// In this case, it would be an expensive computation that is not worth performing.
 	// Users prefer to explicitly size canvases just like they do with text documents, browser pages, etc.
@@ -985,6 +985,11 @@ FReply SNodePanel::OnTouchEnded( const FGeometry& MyGeometry, const FPointerEven
 	return FReply::Unhandled();
 }
 
+float SNodePanel::GetRelativeLayoutScale(const FSlotBase& Child) const
+{
+	return GetZoomAmount();
+}
+
 void SNodePanel::FindNodesAffectedByMarquee( FGraphPanelSelectionSet& OutAffectedNodes ) const
 {
 	OutAffectedNodes.Empty();
diff --git a/Engine/Source/Editor/GraphEditor/Public/SNodePanel.h b/Engine/Source/Editor/GraphEditor/Public/SNodePanel.h
--- a/Engine/Source/Editor/GraphEditor/Public/SNodePanel.h
+++ b/Engine/Source/Editor/GraphEditor/Public/SNodePanel.h
@@ -308,7 +308,7 @@ public:
 			return &Children;
 		}
 
-		virtual FVector2D ComputeDesiredSize() const override
+		virtual FVector2D ComputeDesiredSize(float) const override
 		{
 			for( int32 ChildIndex = 0; ChildIndex < Children.Num(); ++ChildIndex )
 			{
@@ -325,8 +325,28 @@ public:
 			return FVector2D::ZeroVector;
 		}
 
+		virtual float GetRelativeLayoutScale(const FSlotBase& Child) const
+		{
+			const FNodeSlot& ThisSlot = static_cast<const FNodeSlot&>(Child);
+			if ( !ThisSlot.AllowScale.Get() )
+			{
+				// Child slots that do not allow zooming should scale themselves to negate the node panel's zoom.
+				TSharedPtr<SNodePanel> ParentPanel = GetParentPanel();
+				if (ParentPanel.IsValid())
+				{					
+					return 1.0f/ParentPanel->GetZoomAmount();
+				}
+			}
+
+			return 1.0f;
+		}
+
 		virtual void OnArrangeChildren( const FGeometry& AllottedGeometry, FArrangedChildren& ArrangedChildren ) const override
 		{
+			// Child slots that do not allow zooming should scale themselves to negate the node panel's zoom.
+			TSharedPtr<SNodePanel> ParentPanel = GetParentPanel();
+			const float ZoomInverse = ParentPanel.IsValid() ? 1.0f / ParentPanel->GetZoomAmount() : 1.0f;
+
 			for( int32 ChildIndex = 0; ChildIndex < Children.Num(); ++ChildIndex )
 			{
 				const FNodeSlot& CurChild = Children[ChildIndex];
@@ -334,7 +354,7 @@ public:
 				if ( ArrangedChildren.Accepts(ChildVisibility) )
 				{
 					const FMargin SlotPadding(CurChild.SlotPadding.Get());
-					const float GeometryScale = CurChild.AllowScale.Get() ? 1.f : 1.f / AllottedGeometry.Scale;
+					// If this child is not allowed to scale, its scale relative to its parent should undo the parent widget's scaling.
 					FVector2D Size;
 
 					if( CurChild.Size.IsSet() )
@@ -347,12 +367,12 @@ public:
 						AlignmentArrangeResult YResult = AlignChild<Orient_Vertical>(AllottedGeometry.Size.Y, CurChild, SlotPadding);
 						Size = FVector2D( XResult.Size, YResult.Size );
 					}
-					const FArrangedWidget ChildGeom = 
-					AllottedGeometry.MakeChild(
+					const FArrangedWidget ChildGeom =
+						AllottedGeometry.MakeChild(
 						CurChild.GetWidget(),
 						CurChild.Offset.Get(),
 						Size,
-						GeometryScale
+						GetRelativeLayoutScale(CurChild)
 					);
 					ArrangedChildren.AddWidget( ChildVisibility, ChildGeom );
 				}
@@ -494,6 +514,11 @@ public:
 			return GetSortDepth() < NodeIn.GetSortDepth();
 		}
 
+		void SetParentPanel(const TSharedPtr<SNodePanel>& InParent)
+		{
+			ParentPanelPtr = InParent;
+		}
+
 	protected:
 		SNode()
 		: BorderImage( FCoreStyle::Get().GetBrush( "NoBorder" ) )
@@ -528,7 +553,13 @@ public:
 
 	private:
 
+		TSharedPtr<SNodePanel> GetParentPanel() const
+		{
+			return ParentPanelPtr.Pin();
+		}
+
 		TPanelChildren<FNodeSlot> Children;
+		TWeakPtr<SNodePanel> ParentPanelPtr;
 
 	};
 
@@ -536,7 +567,7 @@ public:
 
 	// SPanel interface
 	virtual void OnArrangeChildren(const FGeometry& AllottedGeometry, FArrangedChildren& ArrangedChildren) const override;
-	virtual FVector2D ComputeDesiredSize() const override;
+	virtual FVector2D ComputeDesiredSize(float) const override;
 	virtual FChildren* GetChildren() override;
 	// End of SPanel interface
 
@@ -552,7 +583,7 @@ public:
 	virtual void OnFocusLost( const FFocusEvent& InFocusEvent ) override;
 	virtual FReply OnTouchGesture( const FGeometry& MyGeometry, const FPointerEvent& GestureEvent ) override;
 	virtual FReply OnTouchEnded( const FGeometry& MyGeometry, const FPointerEvent& InTouchEvent ) override;
-
+	virtual float GetRelativeLayoutScale(const FSlotBase& Child) const override;
 	// End of SWidget interface
 public:
 	/**
diff --git a/Engine/Source/Editor/KismetWidgets/Private/SScrubWidget.cpp b/Engine/Source/Editor/KismetWidgets/Private/SScrubWidget.cpp
--- a/Engine/Source/Editor/KismetWidgets/Private/SScrubWidget.cpp
+++ b/Engine/Source/Editor/KismetWidgets/Private/SScrubWidget.cpp
@@ -373,7 +373,7 @@ void SScrubWidget::CommitValue( float NewValue, bool bSliderClamp, bool bCommitt
 	OnValueChanged.ExecuteIfBound( NewValue );
 }
 
-FVector2D SScrubWidget::ComputeDesiredSize() const 
+FVector2D SScrubWidget::ComputeDesiredSize( float ) const
 {
 	return FVector2D(100, 30);
 }
diff --git a/Engine/Source/Editor/Matinee/Private/SMatineeRecorder.cpp b/Engine/Source/Editor/Matinee/Private/SMatineeRecorder.cpp
--- a/Engine/Source/Editor/Matinee/Private/SMatineeRecorder.cpp
+++ b/Engine/Source/Editor/Matinee/Private/SMatineeRecorder.cpp
@@ -109,8 +109,6 @@ void SMatineeRecorder::Construct(const FArguments& InArgs)
 		]
 	];
 
-	CameraModeComboBox->ComputeDesiredSize();
-
 	// Create an animation viewport client
 	LevelViewportClient = MakeShareable( new FLevelEditorViewportClient(nullptr) );
 
diff --git a/Engine/Source/Editor/NewLevelDialog/Private/NewLevelDialogModule.cpp b/Engine/Source/Editor/NewLevelDialog/Private/NewLevelDialogModule.cpp
--- a/Engine/Source/Editor/NewLevelDialog/Private/NewLevelDialogModule.cpp
+++ b/Engine/Source/Editor/NewLevelDialog/Private/NewLevelDialogModule.cpp
@@ -39,7 +39,7 @@ public:
 		return LayerId;
 	}
 
-	virtual FVector2D ComputeDesiredSize() const
+	virtual FVector2D ComputeDesiredSize(float) const
 	{
 		return FVector2D(Size.X, Size.Y);
 	}
diff --git a/Engine/Source/Editor/Persona/Private/SAnimNotifyPanel.cpp b/Engine/Source/Editor/Persona/Private/SAnimNotifyPanel.cpp
--- a/Engine/Source/Editor/Persona/Private/SAnimNotifyPanel.cpp
+++ b/Engine/Source/Editor/Persona/Private/SAnimNotifyPanel.cpp
@@ -90,7 +90,7 @@ public:
 	UObject* GetObjectBeingDisplayed() const;
 	// End of SNodePanel::SNode
 
-	virtual FVector2D ComputeDesiredSize() const  override;
+	virtual FVector2D ComputeDesiredSize(float) const override;
 	virtual int32 OnPaint(const FPaintArgs& Args, const FGeometry& AllottedGeometry, const FSlateRect& MyClippingRect, FSlateWindowElementList& OutDrawElements, int32 LayerId, const FWidgetStyle& InWidgetStyle, bool bParentEnabled) const override;
 
 	/** Helpers to draw scrub handles and snap offsets */
@@ -372,7 +372,7 @@ protected:
 	int32 GetHitNotifyNode(const FGeometry& MyGeometry, const FVector2D& Position);
 
 	TSharedPtr<SWidget> SummonContextMenu(const FGeometry& MyGeometry, const FPointerEvent& MouseEvent);
-	virtual FVector2D ComputeDesiredSize() const override;
+	virtual FVector2D ComputeDesiredSize(float) const override;
 	virtual FReply OnMouseButtonUp(const FGeometry& MyGeometry, const FPointerEvent& MouseEvent) override;
 
 	virtual FReply OnDrop( const FGeometry& MyGeometry, const FDragDropEvent& DragDropEvent ) override;
@@ -1042,7 +1042,7 @@ void SAnimNotifyNode::DropCancelled()
 	bBeingDragged = false;
 }
 
-FVector2D SAnimNotifyNode::ComputeDesiredSize() const 
+FVector2D SAnimNotifyNode::ComputeDesiredSize( float ) const
 {
 	return GetSize();
 }
@@ -1566,7 +1566,7 @@ void SAnimNotifyTrack::Construct(const FArguments& InArgs)
 	Update();
 }
 
-FVector2D SAnimNotifyTrack::ComputeDesiredSize() const 
+FVector2D SAnimNotifyTrack::ComputeDesiredSize( float ) const
 {
 	FVector2D Size;
 	Size.X = 200;
diff --git a/Engine/Source/Editor/Persona/Private/SAnimationBlendSpace1D.cpp b/Engine/Source/Editor/Persona/Private/SAnimationBlendSpace1D.cpp
--- a/Engine/Source/Editor/Persona/Private/SAnimationBlendSpace1D.cpp
+++ b/Engine/Source/Editor/Persona/Private/SAnimationBlendSpace1D.cpp
@@ -239,7 +239,7 @@ void SBlendSpace1DWidget::Construct(const FArguments& InArgs)
 	check (BlendSpace);
 }
 
-FVector2D SBlendSpace1DWidget::ComputeDesiredSize() const
+FVector2D SBlendSpace1DWidget::ComputeDesiredSize( float ) const
 {
 	return FVector2D( 150.0f, 150.0f );
 }
diff --git a/Engine/Source/Editor/Persona/Private/STrack.cpp b/Engine/Source/Editor/Persona/Private/STrack.cpp
--- a/Engine/Source/Editor/Persona/Private/STrack.cpp
+++ b/Engine/Source/Editor/Persona/Private/STrack.cpp
@@ -303,7 +303,7 @@ FVector2D STrackNode::GetSize() const
 	return LastSize;
 }
 
-FVector2D STrackNode::ComputeDesiredSize() const
+FVector2D STrackNode::ComputeDesiredSize( float ) const
 {
 	return GetSize();
 }
@@ -372,7 +372,7 @@ void STrack::OnArrangeChildren( const FGeometry& AllottedGeometry, FArrangedChil
 	}
 }
 
-FVector2D STrack::ComputeDesiredSize() const 
+FVector2D STrack::ComputeDesiredSize( float ) const
 {
 	FVector2D Size;
 	Size.X = 5000; // FIXME: Setting this really high so that the track takes all the width it can in the montage tool (since its placed in an HBox next to a fixed with column
diff --git a/Engine/Source/Editor/Persona/Private/STrack.h b/Engine/Source/Editor/Persona/Private/STrack.h
--- a/Engine/Source/Editor/Persona/Private/STrack.h
+++ b/Engine/Source/Editor/Persona/Private/STrack.h
@@ -132,7 +132,7 @@ public:
 	bool HitTest(const FGeometry& AllottedGeometry, FVector2D MouseLocalPose) const;
 
 	virtual FVector2D GetSize() const;
-	virtual FVector2D ComputeDesiredSize() const;
+	virtual FVector2D ComputeDesiredSize(float) const;
 
 	float GetDataStartPos() const;
 
@@ -255,7 +255,7 @@ public:
 	virtual FCursorReply OnCursorQuery( const FGeometry& MyGeometry, const FPointerEvent& CursorEvent ) const override;
 
 	void				AddTrackNode( TSharedRef<STrackNode> Node );
-	FVector2D			ComputeDesiredSize() const override;
+	FVector2D			ComputeDesiredSize(float) const override;
 	virtual FChildren*	GetChildren() override;
 
 protected:
diff --git a/Engine/Source/Editor/PropertyEditor/Private/SDetailSingleItemRow.cpp b/Engine/Source/Editor/PropertyEditor/Private/SDetailSingleItemRow.cpp
--- a/Engine/Source/Editor/PropertyEditor/Private/SDetailSingleItemRow.cpp
+++ b/Engine/Source/Editor/PropertyEditor/Private/SDetailSingleItemRow.cpp
@@ -36,14 +36,14 @@ public:
 		];
 	}
 
-	virtual FVector2D ComputeDesiredSize() const override
+	virtual FVector2D ComputeDesiredSize(float LayoutScaleMultiplier) const override
 	{
 		const float MinWidthVal = MinWidth.Get().Get(0.0f);
 		const float MaxWidthVal = MaxWidth.Get().Get(0.0f);
 
 		if ( MinWidthVal == 0.0f && MaxWidthVal == 0.0f )
 		{
-			return SCompoundWidget::ComputeDesiredSize();
+			return SCompoundWidget::ComputeDesiredSize(LayoutScaleMultiplier);
 		}
 		else
 		{
diff --git a/Engine/Source/Editor/Sequencer/Private/SAnimationOutlinerView.cpp b/Engine/Source/Editor/Sequencer/Private/SAnimationOutlinerView.cpp
--- a/Engine/Source/Editor/Sequencer/Private/SAnimationOutlinerView.cpp
+++ b/Engine/Source/Editor/Sequencer/Private/SAnimationOutlinerView.cpp
@@ -202,7 +202,6 @@ void SAnimationOutlinerView::OnArrangeChildren( const FGeometry& AllottedGeometr
 	const float Padding = SequencerLayoutConstants::NodePadding;
 	const float IndentAmount = SequencerLayoutConstants::IndentAmount;
 
-	const float Scale = 1.0f;
 	float CurrentHeight = 0;
 	for (int32 WidgetIndex = 0; WidgetIndex < Children.Num(); ++WidgetIndex)
 	{
@@ -219,7 +218,7 @@ void SAnimationOutlinerView::OnArrangeChildren( const FGeometry& AllottedGeometr
 			// Place the widget at the current height, at the nodes desired size
 			ArrangedChildren.AddWidget( 
 				Visibility, 
-				AllottedGeometry.MakeChild( Widget, FVector2D( WidgetIndentOffset, CurrentHeight ), FVector2D( AllottedGeometry.GetDrawSize().X-WidgetIndentOffset, HeightIncrement ), Scale ) 
+				AllottedGeometry.MakeChild( Widget, FVector2D( WidgetIndentOffset, CurrentHeight ), FVector2D( AllottedGeometry.GetDrawSize().X-WidgetIndentOffset, HeightIncrement ) ) 
 				);
 		
 			// Compute the start height for the next widget
diff --git a/Engine/Source/Editor/Sequencer/Private/SequencerDisplayNode.cpp b/Engine/Source/Editor/Sequencer/Private/SequencerDisplayNode.cpp
--- a/Engine/Source/Editor/Sequencer/Private/SequencerDisplayNode.cpp
+++ b/Engine/Source/Editor/Sequencer/Private/SequencerDisplayNode.cpp
@@ -24,7 +24,7 @@ public:
 
 	/** SLeafWidget Interface */
 	virtual int32 OnPaint( const FPaintArgs& Args, const FGeometry& AllottedGeometry, const FSlateRect& MyClippingRect, FSlateWindowElementList& OutDrawElements, int32 LayerId, const FWidgetStyle& InWidgetStyle, bool bParentEnabled ) const override;
-	virtual FVector2D ComputeDesiredSize() const override;
+	virtual FVector2D ComputeDesiredSize(float) const override;
 	
 	void Construct( const FArguments& InArgs, TSharedRef<FSequencerDisplayNode> InRootNode )
 	{
@@ -86,7 +86,7 @@ int32 SSequencerObjectTrack::OnPaint(const FPaintArgs& Args, const FGeometry& Al
 	return LayerId+1;
 }
 
-FVector2D SSequencerObjectTrack::ComputeDesiredSize() const
+FVector2D SSequencerObjectTrack::ComputeDesiredSize( float ) const
 {
 	// Note: X Size is not used
 	return FVector2D( 100.0f, RootNode->GetNodeHeight() );
diff --git a/Engine/Source/Editor/SequencerWidgets/Private/STimeSlider.cpp b/Engine/Source/Editor/SequencerWidgets/Private/STimeSlider.cpp
--- a/Engine/Source/Editor/SequencerWidgets/Private/STimeSlider.cpp
+++ b/Engine/Source/Editor/SequencerWidgets/Private/STimeSlider.cpp
@@ -34,7 +34,7 @@ FReply STimeSlider::OnMouseMove( const FGeometry& MyGeometry, const FPointerEven
 	return TimeSliderController->OnMouseMove( SharedThis(this), MyGeometry, MouseEvent );
 }
 
-FVector2D STimeSlider::ComputeDesiredSize() const 
+FVector2D STimeSlider::ComputeDesiredSize( float ) const
 {
 	return FVector2D(100, 22);
 }
diff --git a/Engine/Source/Editor/UMGEditor/Private/Designer/SRuler.cpp b/Engine/Source/Editor/UMGEditor/Private/Designer/SRuler.cpp
--- a/Engine/Source/Editor/UMGEditor/Private/Designer/SRuler.cpp
+++ b/Engine/Source/Editor/UMGEditor/Private/Designer/SRuler.cpp
@@ -366,7 +366,7 @@ FReply SRuler::OnMouseMove(const FGeometry& MyGeometry, const FPointerEvent& Mou
 	return FReply::Unhandled();
 }
 
-FVector2D SRuler::ComputeDesiredSize() const
+FVector2D SRuler::ComputeDesiredSize( float ) const
 {
 	return Orientation == Orient_Horizontal ? FVector2D(100, 18) : FVector2D(18, 100);
 }
diff --git a/Engine/Source/Editor/UMGEditor/Private/Designer/SZoomPan.cpp b/Engine/Source/Editor/UMGEditor/Private/Designer/SZoomPan.cpp
--- a/Engine/Source/Editor/UMGEditor/Private/Designer/SZoomPan.cpp
+++ b/Engine/Source/Editor/UMGEditor/Private/Designer/SZoomPan.cpp
@@ -4,7 +4,6 @@
 
 #include "SZoomPan.h"
 
-#define LOCTEXT_NAMESPACE "UMG"
 
 /////////////////////////////////////////////////////
 // SZoomPan
@@ -20,6 +19,14 @@ void SZoomPan::Construct(const FArguments& InArgs)
 	];
 }
 
+void SZoomPan::SetContent(const TSharedRef< SWidget >& InContent)
+{
+	ChildSlot
+		[
+			InContent
+		];
+}
+
 void SZoomPan::OnArrangeChildren(const FGeometry& AllottedGeometry, FArrangedChildren& ArrangedChildren) const
 {
 	const EVisibility ChildVisibility = ChildSlot.GetWidget()->GetVisibility();
@@ -38,12 +45,8 @@ void SZoomPan::OnArrangeChildren(const FGeometry& AllottedGeometry, FArrangedChi
 	}
 }
 
-void SZoomPan::SetContent(const TSharedRef< SWidget >& InContent)
+float SZoomPan::GetRelativeLayoutScale(const FSlotBase& Child) const
 {
-	ChildSlot
-	[
-		InContent
-	];
+	return ZoomAmount.Get();
 }
 
-#undef LOCTEXT_NAMESPACE
diff --git a/Engine/Source/Editor/UMGEditor/Private/Designer/SZoomPan.h b/Engine/Source/Editor/UMGEditor/Private/Designer/SZoomPan.h
--- a/Engine/Source/Editor/UMGEditor/Private/Designer/SZoomPan.h
+++ b/Engine/Source/Editor/UMGEditor/Private/Designer/SZoomPan.h
@@ -19,8 +19,6 @@ public:
 
 	void Construct(const FArguments& InArgs);
 	
-	void OnArrangeChildren(const FGeometry& AllottedGeometry, FArrangedChildren& ArrangedChildren) const;
-
 	/**
 	 * Sets the content for this border
 	 *
@@ -29,6 +27,10 @@ public:
 	void SetContent( const TSharedRef< SWidget >& InContent );
 
 protected:
+	void OnArrangeChildren(const FGeometry& AllottedGeometry, FArrangedChildren& ArrangedChildren) const;
+
+	virtual float GetRelativeLayoutScale( const FSlotBase& Child ) const override;
+
 	/** The position within the panel at which the user is looking */
 	TAttribute<FVector2D> ViewOffset;
 
diff --git a/Engine/Source/Editor/UnrealEd/Private/SColorGradientEditor.cpp b/Engine/Source/Editor/UnrealEd/Private/SColorGradientEditor.cpp
--- a/Engine/Source/Editor/UnrealEd/Private/SColorGradientEditor.cpp
+++ b/Engine/Source/Editor/UnrealEd/Private/SColorGradientEditor.cpp
@@ -437,7 +437,7 @@ FReply SColorGradientEditor::OnKeyDown( const FGeometry& MyGeometry, const FKeyE
 	return FReply::Unhandled();
 }
 
-FVector2D SColorGradientEditor::ComputeDesiredSize() const
+FVector2D SColorGradientEditor::ComputeDesiredSize( float ) const
 {
 	return FVector2D( 1000, 75 );
 }
diff --git a/Engine/Source/Editor/UnrealEd/Private/SCurveEditor.cpp b/Engine/Source/Editor/UnrealEd/Private/SCurveEditor.cpp
--- a/Engine/Source/Editor/UnrealEd/Private/SCurveEditor.cpp
+++ b/Engine/Source/Editor/UnrealEd/Private/SCurveEditor.cpp
@@ -488,7 +488,7 @@ void SCurveEditor::PushKeyMenu(const FGeometry& InMyGeometry, FVector2D Position
 }
 
 
-FVector2D SCurveEditor::ComputeDesiredSize() const
+FVector2D SCurveEditor::ComputeDesiredSize( float ) const
 {
 	return DesiredSize.Get();
 }
diff --git a/Engine/Source/Editor/UnrealEd/Public/SColorGradientEditor.h b/Engine/Source/Editor/UnrealEd/Public/SColorGradientEditor.h
--- a/Engine/Source/Editor/UnrealEd/Public/SColorGradientEditor.h
+++ b/Engine/Source/Editor/UnrealEd/Public/SColorGradientEditor.h
@@ -51,7 +51,7 @@ public:
 	virtual FReply OnMouseMove( const FGeometry& MyGeometry, const FPointerEvent& MouseEvent ) override;
 	virtual FReply OnMouseButtonUp( const FGeometry& MyGeometry, const FPointerEvent& MouseEvent ) override;
 	virtual FReply OnKeyDown( const FGeometry& MyGeometry, const FKeyEvent& InKeyEvent ) override;
-	virtual FVector2D ComputeDesiredSize() const;
+	virtual FVector2D ComputeDesiredSize(float) const;
 
 	/**
 	 * Sets the curve that is being edited by this editor
diff --git a/Engine/Source/Editor/UnrealEd/Public/SCurveEditor.h b/Engine/Source/Editor/UnrealEd/Public/SCurveEditor.h
--- a/Engine/Source/Editor/UnrealEd/Public/SCurveEditor.h
+++ b/Engine/Source/Editor/UnrealEd/Public/SCurveEditor.h
@@ -357,7 +357,7 @@ private:
 	FSimpleDelegate OnCreateAsset;
 
 	// SWidget interface
-	UNREALED_API virtual FVector2D ComputeDesiredSize() const override;
+	UNREALED_API virtual FVector2D ComputeDesiredSize(float) const override;
 
 	/** Paint a curve */
 	void PaintCurve(TSharedPtr<FCurveViewModel> CurveViewModel, const FGeometry &AllottedGeometry, FTrackScaleInfo &ScaleInfo, FSlateWindowElementList &OutDrawElements, 
diff --git a/Engine/Source/Editor/UserFeedback/Private/UserFeedback.cpp b/Engine/Source/Editor/UserFeedback/Private/UserFeedback.cpp
--- a/Engine/Source/Editor/UserFeedback/Private/UserFeedback.cpp
+++ b/Engine/Source/Editor/UserFeedback/Private/UserFeedback.cpp
@@ -477,13 +477,13 @@ public:
 	}
 
 	/** Compute the desired size for this widget */
-	virtual FVector2D ComputeDesiredSize() const override
+	virtual FVector2D ComputeDesiredSize(float LayoutScaleMultiplier) const override
 	{
 		const float MinWidthVal = MinWidth.Get();
 
 		if (MinWidthVal == 0.0f)
 		{
-			return SCompoundWidget::ComputeDesiredSize();
+			return SCompoundWidget::ComputeDesiredSize(LayoutScaleMultiplier);
 		}
 		else
 		{
diff --git a/Engine/Source/Editor/WorldBrowser/Private/Tiles/SWorldTileItem.cpp b/Engine/Source/Editor/WorldBrowser/Private/Tiles/SWorldTileItem.cpp
--- a/Engine/Source/Editor/WorldBrowser/Private/Tiles/SWorldTileItem.cpp
+++ b/Engine/Source/Editor/WorldBrowser/Private/Tiles/SWorldTileItem.cpp
@@ -356,7 +356,7 @@ FVector2D SWorldTileItem::GetDesiredSizeForMarquee() const
 	return SNodePanel::SNode::GetDesiredSizeForMarquee();
 }
 
-FVector2D SWorldTileItem::ComputeDesiredSize() const
+FVector2D SWorldTileItem::ComputeDesiredSize( float ) const
 {
 	return TileModel->GetLevelSize2D();
 }
diff --git a/Engine/Source/Runtime/AppFramework/Private/Framework/Testing/SWidgetGallery.cpp b/Engine/Source/Runtime/AppFramework/Private/Framework/Testing/SWidgetGallery.cpp
--- a/Engine/Source/Runtime/AppFramework/Private/Framework/Testing/SWidgetGallery.cpp
+++ b/Engine/Source/Runtime/AppFramework/Private/Framework/Testing/SWidgetGallery.cpp
@@ -414,7 +414,7 @@ public:
 							.Padding(0.0f, 5.0f)
 							[
 								SNew(SHyperlink)
-									.Text(LOCTEXT("SHyperlinkText", "Hyperlink"))
+									.Text(LOCTEXT("SHyperlinkText", "Text to appear in the Hyperlink widget."))
 			
 							]
 
diff --git a/Engine/Source/Runtime/AppFramework/Private/Widgets/Colors/SColorThemes.cpp b/Engine/Source/Runtime/AppFramework/Private/Widgets/Colors/SColorThemes.cpp
--- a/Engine/Source/Runtime/AppFramework/Private/Widgets/Colors/SColorThemes.cpp
+++ b/Engine/Source/Runtime/AppFramework/Private/Widgets/Colors/SColorThemes.cpp
@@ -498,7 +498,7 @@ void SThemeColorBlocksBar::OnArrangeChildren( const FGeometry& AllottedGeometry,
 *
 * @return The desired size.
 */
-FVector2D SThemeColorBlocksBar::ComputeDesiredSize() const
+FVector2D SThemeColorBlocksBar::ComputeDesiredSize( float ) const
 {
 	return FVector2D(64, 16);
 }
diff --git a/Engine/Source/Runtime/AppFramework/Private/Widgets/Colors/SColorValueSlider.h b/Engine/Source/Runtime/AppFramework/Private/Widgets/Colors/SColorValueSlider.h
--- a/Engine/Source/Runtime/AppFramework/Private/Widgets/Colors/SColorValueSlider.h
+++ b/Engine/Source/Runtime/AppFramework/Private/Widgets/Colors/SColorValueSlider.h
@@ -132,7 +132,7 @@ public:
 		return FReply::Handled();
 	}
 
-	virtual FVector2D ComputeDesiredSize() const
+	virtual FVector2D ComputeDesiredSize(float) const
 	{
 		return SelectorImage->ImageSize * 2.f;
 	}
diff --git a/Engine/Source/Runtime/AppFramework/Private/Widgets/Workflow/SWizard.cpp b/Engine/Source/Runtime/AppFramework/Private/Widgets/Workflow/SWizard.cpp
--- a/Engine/Source/Runtime/AppFramework/Private/Widgets/Workflow/SWizard.cpp
+++ b/Engine/Source/Runtime/AppFramework/Private/Widgets/Workflow/SWizard.cpp
@@ -232,11 +232,11 @@ void SWizard::ShowPage( int32 PageIndex )
 /* SCompoundWidget overrides
  *****************************************************************************/
 
-FVector2D SWizard::ComputeDesiredSize() const
+FVector2D SWizard::ComputeDesiredSize(float LayoutScaleMultiplier) const
 {
 	if (DesiredSize.IsZero())
 	{
-		return SCompoundWidget::ComputeDesiredSize();
+		return SCompoundWidget::ComputeDesiredSize(LayoutScaleMultiplier);
 	}
 
 	return DesiredSize;
diff --git a/Engine/Source/Runtime/Slate/Private/Framework/Application/MenuStack.cpp b/Engine/Source/Runtime/Slate/Private/Framework/Application/MenuStack.cpp
--- a/Engine/Source/Runtime/Slate/Private/Framework/Application/MenuStack.cpp
+++ b/Engine/Source/Runtime/Slate/Private/Framework/Application/MenuStack.cpp
@@ -102,8 +102,9 @@ TSharedRef<SWindow> FMenuStack::PushMenu( const TSharedRef<SWindow>& ParentWindo
 	// Adjust the position of popup windows so they do not go out of the visible area of the monitor(s)	
 	if( bShouldAutoSize )
 	{
-		// @todo: Doesn't take into account potential window border size
-		WrappedContent->SlatePrepass();
+		// @todo slate: Assumes that popup is not Scaled up or down from application scale.
+		WrappedContent->SlatePrepass( FSlateApplication::Get().GetApplicationScale() );
+		// @todo slate: Doesn't take into account potential window border size
 		ExpectedSize = WrappedContent->GetDesiredSize();
 	}
 	EOrientation Orientation = (TransitionEffect.SlideDirection == FPopupTransitionEffect::SubMenu) ? Orient_Horizontal : Orient_Vertical;
diff --git a/Engine/Source/Runtime/Slate/Private/Framework/Application/SlateApplication.cpp b/Engine/Source/Runtime/Slate/Private/Framework/Application/SlateApplication.cpp
--- a/Engine/Source/Runtime/Slate/Private/Framework/Application/SlateApplication.cpp
+++ b/Engine/Source/Runtime/Slate/Private/Framework/Application/SlateApplication.cpp
@@ -884,7 +884,7 @@ void FSlateApplication::TickWindowAndChildren( TSharedRef<SWindow> WindowToTick
 		
 		{
 			SCOPE_CYCLE_COUNTER( STAT_SlateCacheDesiredSize );
-			WindowToTick->SlatePrepass();
+			WindowToTick->SlatePrepass(GetApplicationScale());
 		}
 
 		if (WindowToTick->IsAutosized())
@@ -956,7 +956,7 @@ void FSlateApplication::DrawWindowAndChildren( const TSharedRef<SWindow>& Window
 				
 				if (CursorWidget.IsValid())
 				{
-					CursorWidget->SlatePrepass();
+					CursorWidget->SlatePrepass(GetApplicationScale());
 
 					FVector2D CursorPosInWindowSpace = WindowToDraw->GetWindowGeometryInScreen().AbsoluteToLocal(GetCursorPos());
 					CursorPosInWindowSpace += (CursorWidget->GetDesiredSize() * -0.5);
@@ -1003,7 +1003,7 @@ static void PrepassWindowAndChildren( TSharedRef<SWindow> WindowToPrepass )
 
 	{
 		SCOPE_CYCLE_COUNTER(STAT_SlateCacheDesiredSize);
-		WindowToPrepass->SlatePrepass();
+		WindowToPrepass->SlatePrepass(FSlateApplication::Get().GetApplicationScale());
 	}
 
 	if (WindowToPrepass->IsAutosized())
@@ -2632,7 +2632,7 @@ void FSlateApplication::SpawnToolTip( const TSharedRef<IToolTip>& InToolTip, con
 		DesiredToolTipLocation = InSpawnLocation;
 		{
 			// Make sure the desired size is valid
-			NewToolTipWindow->SlatePrepass();
+			NewToolTipWindow->SlatePrepass(FSlateApplication::Get().GetApplicationScale());
 
 			FSlateRect Anchor(DesiredToolTipLocation.X, DesiredToolTipLocation.Y, DesiredToolTipLocation.X, DesiredToolTipLocation.Y);
 			DesiredToolTipLocation = CalculatePopupWindowPosition( Anchor, NewToolTipWindow->GetDesiredSize() );
diff --git a/Engine/Source/Runtime/Slate/Private/Framework/Docking/SDockingCross.cpp b/Engine/Source/Runtime/Slate/Private/Framework/Docking/SDockingCross.cpp
--- a/Engine/Source/Runtime/Slate/Private/Framework/Docking/SDockingCross.cpp
+++ b/Engine/Source/Runtime/Slate/Private/Framework/Docking/SDockingCross.cpp
@@ -149,7 +149,7 @@ int32 SDockingCross::OnPaint( const FPaintArgs& Args, const FGeometry& AllottedG
 
 }
 
-FVector2D SDockingCross::ComputeDesiredSize() const
+FVector2D SDockingCross::ComputeDesiredSize( float ) const
 {
 	// DockCross does not really have a desired size.
 	return FVector2D( 16, 16 );
diff --git a/Engine/Source/Runtime/Slate/Private/Framework/Docking/SDockingTabWell.cpp b/Engine/Source/Runtime/Slate/Private/Framework/Docking/SDockingTabWell.cpp
--- a/Engine/Source/Runtime/Slate/Private/Framework/Docking/SDockingTabWell.cpp
+++ b/Engine/Source/Runtime/Slate/Private/Framework/Docking/SDockingTabWell.cpp
@@ -157,7 +157,7 @@ int32 SDockingTabWell::OnPaint( const FPaintArgs& Args, const FGeometry& Allotte
 }
 
 
-FVector2D SDockingTabWell::ComputeDesiredSize() const
+FVector2D SDockingTabWell::ComputeDesiredSize( float ) const
 {
 	FVector2D DesiredSizeResult(0,0);
 
diff --git a/Engine/Source/Runtime/Slate/Private/Framework/MultiBox/SClippingHorizontalBox.cpp b/Engine/Source/Runtime/Slate/Private/Framework/MultiBox/SClippingHorizontalBox.cpp
--- a/Engine/Source/Runtime/Slate/Private/Framework/MultiBox/SClippingHorizontalBox.cpp
+++ b/Engine/Source/Runtime/Slate/Private/Framework/MultiBox/SClippingHorizontalBox.cpp
@@ -110,9 +110,9 @@ int32 SClippingHorizontalBox::OnPaint( const FPaintArgs& Args, const FGeometry&
 	return SHorizontalBox::OnPaint(Args, AllottedGeometry, MyClippingRect, OutDrawElements, LayerId, InWidgetStyle, bParentEnabled);
 }
 
-FVector2D SClippingHorizontalBox::ComputeDesiredSize() const
+FVector2D SClippingHorizontalBox::ComputeDesiredSize(float LayoutScaleMultiplier) const
 {
-	FVector2D Size = SBoxPanel::ComputeDesiredSize();
+	FVector2D Size = SBoxPanel::ComputeDesiredSize(LayoutScaleMultiplier);
 	if (ClippedIdx == (Children.Num() - 2))
 	{
 		// If the wrap button isn't being shown, subtract it's size from the total desired size
diff --git a/Engine/Source/Runtime/Slate/Private/Widgets/Colors/SColorBlock.cpp b/Engine/Source/Runtime/Slate/Private/Widgets/Colors/SColorBlock.cpp
--- a/Engine/Source/Runtime/Slate/Private/Widgets/Colors/SColorBlock.cpp
+++ b/Engine/Source/Runtime/Slate/Private/Widgets/Colors/SColorBlock.cpp
@@ -93,7 +93,7 @@ FReply SColorBlock::OnMouseButtonDown( const FGeometry& MyGeometry, const FPoint
 	}
 }
 
-FVector2D SColorBlock::ComputeDesiredSize() const
+FVector2D SColorBlock::ComputeDesiredSize( float ) const
 {
 	return ColorBlockSize.Get();
 }
diff --git a/Engine/Source/Runtime/Slate/Private/Widgets/Colors/SColorSpectrum.cpp b/Engine/Source/Runtime/Slate/Private/Widgets/Colors/SColorSpectrum.cpp
--- a/Engine/Source/Runtime/Slate/Private/Widgets/Colors/SColorSpectrum.cpp
+++ b/Engine/Source/Runtime/Slate/Private/Widgets/Colors/SColorSpectrum.cpp
@@ -22,7 +22,7 @@ void SColorSpectrum::Construct( const FArguments& InArgs )
 /* SWidget overrides
  *****************************************************************************/
 
-FVector2D SColorSpectrum::ComputeDesiredSize( ) const
+FVector2D SColorSpectrum::ComputeDesiredSize( float ) const
 {
 	return Image->ImageSize;
 }
diff --git a/Engine/Source/Runtime/Slate/Private/Widgets/Colors/SColorWheel.cpp b/Engine/Source/Runtime/Slate/Private/Widgets/Colors/SColorWheel.cpp
--- a/Engine/Source/Runtime/Slate/Private/Widgets/Colors/SColorWheel.cpp
+++ b/Engine/Source/Runtime/Slate/Private/Widgets/Colors/SColorWheel.cpp
@@ -22,7 +22,7 @@ void SColorWheel::Construct( const FArguments& InArgs )
 /* SWidget overrides
  *****************************************************************************/
 
-FVector2D SColorWheel::ComputeDesiredSize( ) const
+FVector2D SColorWheel::ComputeDesiredSize( float ) const
 {
 	return Image->ImageSize + SelectorImage->ImageSize;
 }
diff --git a/Engine/Source/Runtime/Slate/Private/Widgets/Images/SImage.cpp b/Engine/Source/Runtime/Slate/Private/Widgets/Images/SImage.cpp
--- a/Engine/Source/Runtime/Slate/Private/Widgets/Images/SImage.cpp
+++ b/Engine/Source/Runtime/Slate/Private/Widgets/Images/SImage.cpp
@@ -61,7 +61,7 @@ FReply SImage::OnMouseButtonDown( const FGeometry& MyGeometry, const FPointerEve
 /**
  * An Image's desired size is whatever the image looks best as. This is decided on a case by case basis via the Width and Height properties.
  */
-FVector2D SImage::ComputeDesiredSize() const
+FVector2D SImage::ComputeDesiredSize( float ) const
 {
 	const FSlateBrush* ImageBrush = Image.Get();
 	if (ImageBrush != nullptr)
diff --git a/Engine/Source/Runtime/Slate/Private/Widgets/Images/SThrobber.cpp b/Engine/Source/Runtime/Slate/Private/Widgets/Images/SThrobber.cpp
--- a/Engine/Source/Runtime/Slate/Private/Widgets/Images/SThrobber.cpp
+++ b/Engine/Source/Runtime/Slate/Private/Widgets/Images/SThrobber.cpp
@@ -158,7 +158,7 @@ int32 SCircularThrobber::OnPaint( const FPaintArgs& Args, const FGeometry& Allot
 	return LayerId;
 }
 
-FVector2D SCircularThrobber::ComputeDesiredSize() const
+FVector2D SCircularThrobber::ComputeDesiredSize( float ) const
 {
 	return FVector2D(Radius, Radius) * 2;
 }
diff --git a/Engine/Source/Runtime/Slate/Private/Widgets/Input/SButton.cpp b/Engine/Source/Runtime/Slate/Private/Widgets/Input/SButton.cpp
--- a/Engine/Source/Runtime/Slate/Private/Widgets/Input/SButton.cpp
+++ b/Engine/Source/Runtime/Slate/Private/Widgets/Input/SButton.cpp
@@ -320,7 +320,7 @@ void SButton::PlayHoverSound() const
 	FSlateApplication::Get().PlaySound( HoveredSound );
 }
 
-FVector2D SButton::ComputeDesiredSize() const
+FVector2D SButton::ComputeDesiredSize(float LayoutScaleMultiplier) const
 {
 	// When there is no widget in the button, it sizes itself based on
 	// the border image specified by the style.
@@ -330,7 +330,7 @@ FVector2D SButton::ComputeDesiredSize() const
 	}
 	else
 	{
-		return SBorder::ComputeDesiredSize();
+		return SBorder::ComputeDesiredSize(LayoutScaleMultiplier);
 	}
 }
 
diff --git a/Engine/Source/Runtime/Slate/Private/Widgets/Input/SEditableText.cpp b/Engine/Source/Runtime/Slate/Private/Widgets/Input/SEditableText.cpp
--- a/Engine/Source/Runtime/Slate/Private/Widgets/Input/SEditableText.cpp
+++ b/Engine/Source/Runtime/Slate/Private/Widgets/Input/SEditableText.cpp
@@ -1217,7 +1217,7 @@ TSharedRef< SWidget > SEditableText::GetWidget()
 	return SharedThis( this );
 }
 
-FVector2D SEditableText::ComputeDesiredSize() const
+FVector2D SEditableText::ComputeDesiredSize(float LayoutScaleMultiplier) const
 {
 	const TSharedRef< FSlateFontMeasure > FontMeasureService = FSlateApplication::Get().GetRenderer()->GetFontMeasureService();
 	const FSlateFontInfo& FontInfo = Font.Get();
@@ -1230,11 +1230,11 @@ FVector2D SEditableText::ComputeDesiredSize() const
 	const FString StringToRender = GetStringToRender();
 	if( !StringToRender.IsEmpty() )
 	{
-		TextSize = FontMeasureService->Measure( StringToRender, FontInfo, 1.0f/*CachedGeometry.Scale*/ );
+		TextSize = FontMeasureService->Measure(StringToRender, FontInfo, LayoutScaleMultiplier) / LayoutScaleMultiplier;
 	}
 	else
 	{
-		TextSize = FontMeasureService->Measure( HintText.Get( ).ToString( ), FontInfo, 1.0f/*CachedGeometry.Scale*/ );
+		TextSize = FontMeasureService->Measure(HintText.Get().ToString(), FontInfo, LayoutScaleMultiplier) / LayoutScaleMultiplier;
 	}
 	
 	const float DesiredWidth = FMath::Max(TextSize.X, MinDesiredWidth.Get()) + CaretWidth;
diff --git a/Engine/Source/Runtime/Slate/Private/Widgets/Input/SMenuAnchor.cpp b/Engine/Source/Runtime/Slate/Private/Widgets/Input/SMenuAnchor.cpp
--- a/Engine/Source/Runtime/Slate/Private/Widgets/Input/SMenuAnchor.cpp
+++ b/Engine/Source/Runtime/Slate/Private/Widgets/Input/SMenuAnchor.cpp
@@ -113,7 +113,7 @@ void SMenuAnchor::Tick( const FGeometry& AllottedGeometry, const double InCurren
 
 		// We made a window for showing the popup.
 		// Update the window's position!
-		if( PopupWindow->IsMorphing() )
+		if (false /*PopupWindow->IsMorphing()*/ )
 		{
 			if( NewShape != PopupWindow->GetMorphTargetShape() )
 			{
@@ -174,7 +174,7 @@ void SMenuAnchor::OnArrangeChildren( const FGeometry& AllottedGeometry, FArrange
 	}
 }
 
-FVector2D SMenuAnchor::ComputeDesiredSize() const
+FVector2D SMenuAnchor::ComputeDesiredSize( float ) const
 {
 	return Children[0].GetWidget()->GetDesiredSize();
 }
@@ -316,29 +316,25 @@ void SMenuAnchor::SetIsOpen( bool InIsOpen, const bool bFocusMenu )
 					OnMenuOpenChanged.Execute(true);
 				}
 
-				// This can be called at any time so we use the push menu override that explicitly allows us to specify our parent
-
 				// Figure out where the menu anchor is on the screen, so we can set the initial position of our pop-up window
-				SlatePrepass();
-
+				// This can be called at any time so we use the push menu override that explicitly allows us to specify our parent
 				// NOTE: Careful, GeneratePathToWidget can be reentrant in that it can call visibility delegates and such
 				FWidgetPath MyWidgetPath;
 				FSlateApplication::Get().GeneratePathToWidgetUnchecked(AsShared(), MyWidgetPath);
 				if (ensure(MyWidgetPath.IsValid()))
 				{
 					const FGeometry& MyGeometry = MyWidgetPath.Widgets.Last().Geometry;
+					const float LayoutScaleMultiplier = MyGeometry.GetAccumulatedLayoutTransform().GetScale();
 
-					// @todo Slate: This code does not properly propagate the layout scale of the widget we are creating the popup for.
-					// The popup instead has a scale of one, but a computed size as if the contents were scaled. This code should ideally
-					// wrap the contents with an SFxWidget that applies the necessary layout scale. This is a very rare case right now.
+					SlatePrepass(LayoutScaleMultiplier);
 
 					// Figure out how big the content widget is so we can set the window's initial size properly
 					TSharedRef< SWidget > MenuContentRef = MenuContentPtr.ToSharedRef();
-					MenuContentRef->SlatePrepass();
+					MenuContentRef->SlatePrepass(LayoutScaleMultiplier);
 
 					// Combo-boxes never size down smaller than the widget that spawned them, but all
 					// other pop-up menus are currently auto-sized
-					const FVector2D DesiredContentSize = MenuContentRef->GetDesiredSize();  // @todo: This is ignoring any window border size!
+					const FVector2D DesiredContentSize = MenuContentRef->GetDesiredSize();  // @todo slate: This is ignoring any window border size!
 					const EMenuPlacement PlacementMode = Placement.Get();
 
 					const FVector2D NewPosition = MyGeometry.AbsolutePosition;
diff --git a/Engine/Source/Runtime/Slate/Private/Widgets/Input/SSlider.cpp b/Engine/Source/Runtime/Slate/Private/Widgets/Input/SSlider.cpp
--- a/Engine/Source/Runtime/Slate/Private/Widgets/Input/SSlider.cpp
+++ b/Engine/Source/Runtime/Slate/Private/Widgets/Input/SSlider.cpp
@@ -109,7 +109,7 @@ int32 SSlider::OnPaint( const FPaintArgs& Args, const FGeometry& AllottedGeometr
 	return LayerId;
 }
 
-FVector2D SSlider::ComputeDesiredSize() const
+FVector2D SSlider::ComputeDesiredSize( float ) const
 {
 	static const FVector2D SSliderDesiredSize(16.0f, 16.0f);
 
diff --git a/Engine/Source/Runtime/Slate/Private/Widgets/Input/SVirtualJoystick.cpp b/Engine/Source/Runtime/Slate/Private/Widgets/Input/SVirtualJoystick.cpp
--- a/Engine/Source/Runtime/Slate/Private/Widgets/Input/SVirtualJoystick.cpp
+++ b/Engine/Source/Runtime/Slate/Private/Widgets/Input/SVirtualJoystick.cpp
@@ -144,7 +144,7 @@ int32 SVirtualJoystick::OnPaint( const FPaintArgs& Args, const FGeometry& Allott
 	return RetLayerId;
 }
 
-FVector2D SVirtualJoystick::ComputeDesiredSize() const
+FVector2D SVirtualJoystick::ComputeDesiredSize( float ) const
 {
 	return FVector2D(100, 100);
 }
diff --git a/Engine/Source/Runtime/Slate/Private/Widgets/Input/SVirtualKeyboardEntry.cpp b/Engine/Source/Runtime/Slate/Private/Widgets/Input/SVirtualKeyboardEntry.cpp
--- a/Engine/Source/Runtime/Slate/Private/Widgets/Input/SVirtualKeyboardEntry.cpp
+++ b/Engine/Source/Runtime/Slate/Private/Widgets/Input/SVirtualKeyboardEntry.cpp
@@ -134,7 +134,7 @@ void SVirtualKeyboardEntry::Tick( const FGeometry& AllottedGeometry, const doubl
  *
  * @return  The widget's desired size
  */
-FVector2D SVirtualKeyboardEntry::ComputeDesiredSize() const
+FVector2D SVirtualKeyboardEntry::ComputeDesiredSize( float ) const
 {
 	const FSlateFontInfo& FontInfo = Font.Get();
 	const TSharedRef< FSlateFontMeasure > FontMeasureService = FSlateApplication::Get().GetRenderer()->GetFontMeasureService();
diff --git a/Engine/Source/Runtime/Slate/Private/Widgets/LayerManager/STooltipPresenter.cpp b/Engine/Source/Runtime/Slate/Private/Widgets/LayerManager/STooltipPresenter.cpp
--- a/Engine/Source/Runtime/Slate/Private/Widgets/LayerManager/STooltipPresenter.cpp
+++ b/Engine/Source/Runtime/Slate/Private/Widgets/LayerManager/STooltipPresenter.cpp
@@ -42,7 +42,7 @@ void STooltipPresenter::OnArrangeChildren(const FGeometry& AllottedGeometry, FAr
 	));
 }
 
-FVector2D STooltipPresenter::ComputeDesiredSize() const
+FVector2D STooltipPresenter::ComputeDesiredSize( float ) const
 {
 	return ChildSlot.GetWidget()->GetDesiredSize();
 }
diff --git a/Engine/Source/Runtime/Slate/Private/Widgets/Layout/SBorder.cpp b/Engine/Source/Runtime/Slate/Private/Widgets/Layout/SBorder.cpp
--- a/Engine/Source/Runtime/Slate/Private/Widgets/Layout/SBorder.cpp
+++ b/Engine/Source/Runtime/Slate/Private/Widgets/Layout/SBorder.cpp
@@ -190,9 +190,9 @@ FReply SBorder::OnMouseButtonDoubleClick( const FGeometry& MyGeometry, const FPo
 	return FReply::Unhandled();
 }
 
-FVector2D SBorder::ComputeDesiredSize() const
+FVector2D SBorder::ComputeDesiredSize(float LayoutScaleMultiplier) const
 {
-	return DesiredSizeScale.Get() * SCompoundWidget::ComputeDesiredSize();
+	return DesiredSizeScale.Get() * SCompoundWidget::ComputeDesiredSize(LayoutScaleMultiplier);
 }
 
 void SBorder::SetBorderBackgroundColor(const TAttribute<FSlateColor>& InColorAndOpacity)
diff --git a/Engine/Source/Runtime/Slate/Private/Widgets/Layout/SBox.cpp b/Engine/Source/Runtime/Slate/Private/Widgets/Layout/SBox.cpp
--- a/Engine/Source/Runtime/Slate/Private/Widgets/Layout/SBox.cpp
+++ b/Engine/Source/Runtime/Slate/Private/Widgets/Layout/SBox.cpp
@@ -82,7 +82,7 @@ void SBox::SetMaxDesiredHeight(TAttribute<FOptionalSize> InMaxDesiredHeight)
 	MaxDesiredHeight = InMaxDesiredHeight;
 }
 
-FVector2D SBox::ComputeDesiredSize() const
+FVector2D SBox::ComputeDesiredSize( float ) const
 {
 	EVisibility ChildVisibility = ChildSlot.GetWidget()->GetVisibility();
 
diff --git a/Engine/Source/Runtime/Slate/Private/Widgets/Layout/SDPIScaler.cpp b/Engine/Source/Runtime/Slate/Private/Widgets/Layout/SDPIScaler.cpp
--- a/Engine/Source/Runtime/Slate/Private/Widgets/Layout/SDPIScaler.cpp
+++ b/Engine/Source/Runtime/Slate/Private/Widgets/Layout/SDPIScaler.cpp
@@ -35,7 +35,7 @@ void SDPIScaler::OnArrangeChildren( const FGeometry& AllottedGeometry, FArranged
 	}
 }
 	
-FVector2D SDPIScaler::ComputeDesiredSize() const
+FVector2D SDPIScaler::ComputeDesiredSize( float ) const
 {
 	return DPIScale.Get() * ChildSlot.GetWidget()->GetDesiredSize();
 }
@@ -57,3 +57,8 @@ void SDPIScaler::SetDPIScale(TAttribute<float> InDPIScale)
 {
 	DPIScale = InDPIScale;
 }
+
+float SDPIScaler::GetRelativeLayoutScale(const FSlotBase& Child) const
+{
+	return DPIScale.Get();
+}
\ No newline at end of file
diff --git a/Engine/Source/Runtime/Slate/Private/Widgets/Layout/SExpandableArea.cpp b/Engine/Source/Runtime/Slate/Private/Widgets/Layout/SExpandableArea.cpp
--- a/Engine/Source/Runtime/Slate/Private/Widgets/Layout/SExpandableArea.cpp
+++ b/Engine/Source/Runtime/Slate/Private/Widgets/Layout/SExpandableArea.cpp
@@ -170,7 +170,7 @@ FVector2D SExpandableArea::GetSectionScale() const
 	return FVector2D( 1.0f, Scale );
 }
 
-FVector2D SExpandableArea::ComputeDesiredSize() const
+FVector2D SExpandableArea::ComputeDesiredSize( float ) const
 {
 	EVisibility ChildVisibility = ChildSlot.GetWidget()->GetVisibility();
 	if ( ChildVisibility != EVisibility::Collapsed )
diff --git a/Engine/Source/Runtime/Slate/Private/Widgets/Layout/SFxWidget.cpp b/Engine/Source/Runtime/Slate/Private/Widgets/Layout/SFxWidget.cpp
--- a/Engine/Source/Runtime/Slate/Private/Widgets/Layout/SFxWidget.cpp
+++ b/Engine/Source/Runtime/Slate/Private/Widgets/Layout/SFxWidget.cpp
@@ -90,7 +90,7 @@ int32 SFxWidget::OnPaint( const FPaintArgs& Args, const FGeometry& AllottedGeome
 
 }
 
-FVector2D SFxWidget::ComputeDesiredSize() const
+FVector2D SFxWidget::ComputeDesiredSize( float ) const
 {
 	// Layout scale affects out desired size.
 	return TransformVector(LayoutScale.Get(), ChildSlot.GetWidget()->GetDesiredSize());
diff --git a/Engine/Source/Runtime/Slate/Private/Widgets/Layout/SGridPanel.cpp b/Engine/Source/Runtime/Slate/Private/Widgets/Layout/SGridPanel.cpp
--- a/Engine/Source/Runtime/Slate/Private/Widgets/Layout/SGridPanel.cpp
+++ b/Engine/Source/Runtime/Slate/Private/Widgets/Layout/SGridPanel.cpp
@@ -215,7 +215,7 @@ void SGridPanel::OnArrangeChildren( const FGeometry& AllottedGeometry, FArranged
 }
 
 
-void SGridPanel::CacheDesiredSize()
+void SGridPanel::CacheDesiredSize(float LayoutScaleMultiplier)
 {
 	// The desired size of the grid is the sum of the desires sizes for every row and column.
 	ComputeDesiredCellSizes( Columns, Rows );
@@ -231,11 +231,11 @@ void SGridPanel::CacheDesiredSize()
 		TotalDesiredSizes.Y += Rows[RowId];
 	}
 	
-	SPanel::CacheDesiredSize();
+	SPanel::CacheDesiredSize(LayoutScaleMultiplier);
 }
 
 
-FVector2D SGridPanel::ComputeDesiredSize() const
+FVector2D SGridPanel::ComputeDesiredSize( float ) const
 {
 	return TotalDesiredSizes;
 }
diff --git a/Engine/Source/Runtime/Slate/Private/Widgets/Layout/SResponsiveGridPanel.cpp b/Engine/Source/Runtime/Slate/Private/Widgets/Layout/SResponsiveGridPanel.cpp
--- a/Engine/Source/Runtime/Slate/Private/Widgets/Layout/SResponsiveGridPanel.cpp
+++ b/Engine/Source/Runtime/Slate/Private/Widgets/Layout/SResponsiveGridPanel.cpp
@@ -271,7 +271,7 @@ void SResponsiveGridPanel::OnArrangeChildren( const FGeometry& AllottedGeometry,
 	}
 }
 
-void SResponsiveGridPanel::CacheDesiredSize()
+void SResponsiveGridPanel::CacheDesiredSize(float LayoutScaleMultiplier)
 {
 	// The desired size of the grid is the sum of the desires sizes for every row and column.
 	TArray<float> Columns;
@@ -296,10 +296,10 @@ void SResponsiveGridPanel::CacheDesiredSize()
 		TotalDesiredSizes.Y += ( RowGutter * 2 )  * ( Slots[Slots.Num() - 1].RowParam );
 	}
 
-	SPanel::CacheDesiredSize();
+	SPanel::CacheDesiredSize(LayoutScaleMultiplier);
 }
 
-FVector2D SResponsiveGridPanel::ComputeDesiredSize() const
+FVector2D SResponsiveGridPanel::ComputeDesiredSize( float ) const
 {
 	return TotalDesiredSizes;
 }
diff --git a/Engine/Source/Runtime/Slate/Private/Widgets/Layout/SScissorRectBox.cpp b/Engine/Source/Runtime/Slate/Private/Widgets/Layout/SScissorRectBox.cpp
--- a/Engine/Source/Runtime/Slate/Private/Widgets/Layout/SScissorRectBox.cpp
+++ b/Engine/Source/Runtime/Slate/Private/Widgets/Layout/SScissorRectBox.cpp
@@ -16,7 +16,7 @@ void SScissorRectBox::Construct(const FArguments& InArgs)
 }
 
 
-FVector2D SScissorRectBox::ComputeDesiredSize() const
+FVector2D SScissorRectBox::ComputeDesiredSize( float ) const
 {
 	EVisibility ChildVisibility = ChildSlot.GetWidget()->GetVisibility();
 
diff --git a/Engine/Source/Runtime/Slate/Private/Widgets/Layout/SSeparator.cpp b/Engine/Source/Runtime/Slate/Private/Widgets/Layout/SSeparator.cpp
--- a/Engine/Source/Runtime/Slate/Private/Widgets/Layout/SSeparator.cpp
+++ b/Engine/Source/Runtime/Slate/Private/Widgets/Layout/SSeparator.cpp
@@ -18,7 +18,7 @@ void SSeparator::Construct( const FArguments& InArgs )
 	Thickness = InArgs._Thickness;
 }
 
-FVector2D SSeparator::ComputeDesiredSize() const
+FVector2D SSeparator::ComputeDesiredSize( float ) const
 {
 	const float Length = 16.0f;
 	return (Orientation == Orient_Horizontal) ? FVector2D(Length, Thickness) : FVector2D(Thickness, Length);
diff --git a/Engine/Source/Runtime/Slate/Private/Widgets/Layout/SSpacer.cpp b/Engine/Source/Runtime/Slate/Private/Widgets/Layout/SSpacer.cpp
--- a/Engine/Source/Runtime/Slate/Private/Widgets/Layout/SSpacer.cpp
+++ b/Engine/Source/Runtime/Slate/Private/Widgets/Layout/SSpacer.cpp
@@ -31,7 +31,7 @@ int32 SSpacer::OnPaint( const FPaintArgs& Args, const FGeometry& AllottedGeometr
 	return LayerId;
 }
 
-FVector2D SSpacer::ComputeDesiredSize() const
+FVector2D SSpacer::ComputeDesiredSize( float ) const
 {
 	return SpacerSize.Get();
 }
diff --git a/Engine/Source/Runtime/Slate/Private/Widgets/Layout/SSplitter.cpp b/Engine/Source/Runtime/Slate/Private/Widgets/Layout/SSplitter.cpp
--- a/Engine/Source/Runtime/Slate/Private/Widgets/Layout/SSplitter.cpp
+++ b/Engine/Source/Runtime/Slate/Private/Widgets/Layout/SSplitter.cpp
@@ -262,7 +262,7 @@ static FVector2D ComputeDesiredSizeForSplitter( const float PhysicalSplitterHand
 * A Panel's desired size in the space required to arrange of its children on the screen while respecting all of
 * the children's desired sizes and any layout-related options specified by the user. See StackPanel for an example.
 */
-FVector2D SSplitter::ComputeDesiredSize() const
+FVector2D SSplitter::ComputeDesiredSize( float ) const
 {
 	FVector2D MyDesiredSize = (Orientation == Orient_Horizontal)
 		? ComputeDesiredSizeForSplitter<Orient_Horizontal>( PhysicalSplitterHandleSize, Children )
@@ -732,7 +732,7 @@ void SSplitter2x2::OnArrangeChildren( const FGeometry& AllottedGeometry, FArrang
 
 
 
-FVector2D SSplitter2x2::ComputeDesiredSize() const
+FVector2D SSplitter2x2::ComputeDesiredSize( float ) const
 {
 	return FVector2D(100,100);
 }
diff --git a/Engine/Source/Runtime/Slate/Private/Widgets/Layout/SUniformGridPanel.cpp b/Engine/Source/Runtime/Slate/Private/Widgets/Layout/SUniformGridPanel.cpp
--- a/Engine/Source/Runtime/Slate/Private/Widgets/Layout/SUniformGridPanel.cpp
+++ b/Engine/Source/Runtime/Slate/Private/Widgets/Layout/SUniformGridPanel.cpp
@@ -52,7 +52,7 @@ void SUniformGridPanel::OnArrangeChildren( const FGeometry& AllottedGeometry, FA
 	}
 }
 
-FVector2D SUniformGridPanel::ComputeDesiredSize() const
+FVector2D SUniformGridPanel::ComputeDesiredSize( float ) const
 {
 	FVector2D MaxChildDesiredSize = FVector2D::ZeroVector;
 	const FVector2D SlotPaddingDesiredSize = SlotPadding.Get().GetDesiredSize();
diff --git a/Engine/Source/Runtime/Slate/Private/Widgets/Layout/SWidgetSwitcher.cpp b/Engine/Source/Runtime/Slate/Private/Widgets/Layout/SWidgetSwitcher.cpp
--- a/Engine/Source/Runtime/Slate/Private/Widgets/Layout/SWidgetSwitcher.cpp
+++ b/Engine/Source/Runtime/Slate/Private/Widgets/Layout/SWidgetSwitcher.cpp
@@ -120,7 +120,7 @@ void SWidgetSwitcher::OnArrangeChildren( const FGeometry& AllottedGeometry, FArr
 	}
 }
 	
-FVector2D SWidgetSwitcher::ComputeDesiredSize( ) const
+FVector2D SWidgetSwitcher::ComputeDesiredSize( float ) const
 {
 	return AllChildren.Num() > 0
 		? AllChildren[WidgetIndex.Get()].GetWidget()->GetDesiredSize()
diff --git a/Engine/Source/Runtime/Slate/Private/Widgets/Layout/SWrapBox.cpp b/Engine/Source/Runtime/Slate/Private/Widgets/Layout/SWrapBox.cpp
--- a/Engine/Source/Runtime/Slate/Private/Widgets/Layout/SWrapBox.cpp
+++ b/Engine/Source/Runtime/Slate/Private/Widgets/Layout/SWrapBox.cpp
@@ -172,7 +172,7 @@ void SWrapBox::ClearChildren()
 	Slots.Empty();
 }
 
-FVector2D SWrapBox::ComputeDesiredSize() const
+FVector2D SWrapBox::ComputeDesiredSize( float ) const
 {
 	FVector2D MyDesiredSize = FVector2D::ZeroVector;
 	FVector2D RowDesiredSize = FVector2D::ZeroVector;
diff --git a/Engine/Source/Runtime/Slate/Private/Widgets/Notifications/SProgressBar.cpp b/Engine/Source/Runtime/Slate/Private/Widgets/Notifications/SProgressBar.cpp
--- a/Engine/Source/Runtime/Slate/Private/Widgets/Notifications/SProgressBar.cpp
+++ b/Engine/Source/Runtime/Slate/Private/Widgets/Notifications/SProgressBar.cpp
@@ -254,7 +254,7 @@ int32 SProgressBar::OnPaint( const FPaintArgs& Args, const FGeometry& AllottedGe
  *
  * @return  The widget's desired size
  */
-FVector2D SProgressBar::ComputeDesiredSize() const
+FVector2D SProgressBar::ComputeDesiredSize( float ) const
 {
 	return GetMarqueeImage()->ImageSize;
 }
diff --git a/Engine/Source/Runtime/Slate/Private/Widgets/SCanvas.cpp b/Engine/Source/Runtime/Slate/Private/Widgets/SCanvas.cpp
--- a/Engine/Source/Runtime/Slate/Private/Widgets/SCanvas.cpp
+++ b/Engine/Source/Runtime/Slate/Private/Widgets/SCanvas.cpp
@@ -122,7 +122,7 @@ int32 SCanvas::OnPaint( const FPaintArgs& Args, const FGeometry& AllottedGeometr
 }
 
 
-FVector2D SCanvas::ComputeDesiredSize() const
+FVector2D SCanvas::ComputeDesiredSize( float ) const
 {
 	// Canvas widgets have no desired size -- their size is always determined by their container
 	return FVector2D::ZeroVector;
diff --git a/Engine/Source/Runtime/Slate/Private/Widgets/SWeakWidget.cpp b/Engine/Source/Runtime/Slate/Private/Widgets/SWeakWidget.cpp
--- a/Engine/Source/Runtime/Slate/Private/Widgets/SWeakWidget.cpp
+++ b/Engine/Source/Runtime/Slate/Private/Widgets/SWeakWidget.cpp
@@ -14,7 +14,7 @@ void SWeakWidget::Construct(const FArguments& InArgs)
 }
 
 
-FVector2D SWeakWidget::ComputeDesiredSize() const
+FVector2D SWeakWidget::ComputeDesiredSize( float ) const
 {	
 	TSharedRef<SWidget> ReferencedWidget = WeakChild.GetWidget();
 
diff --git a/Engine/Source/Runtime/Slate/Private/Widgets/Text/SMultiLineEditableText.cpp b/Engine/Source/Runtime/Slate/Private/Widgets/Text/SMultiLineEditableText.cpp
--- a/Engine/Source/Runtime/Slate/Private/Widgets/Text/SMultiLineEditableText.cpp
+++ b/Engine/Source/Runtime/Slate/Private/Widgets/Text/SMultiLineEditableText.cpp
@@ -2236,7 +2236,7 @@ int32 SMultiLineEditableText::OnPaint( const FPaintArgs& Args, const FGeometry&
 	return LayerId;
 }
 
-void SMultiLineEditableText::CacheDesiredSize()
+void SMultiLineEditableText::CacheDesiredSize(float LayoutScaleMultiplier)
 {
 	// Get the wrapping width and font to see if they have changed
 	float WrappingWidth = WrapTextAt.Get();
@@ -2257,10 +2257,10 @@ void SMultiLineEditableText::CacheDesiredSize()
 	TextLayout->SetVisibleRegion( CachedSize, ScrollOffset );
 	TextLayout->UpdateIfNeeded();
 
-	SWidget::CacheDesiredSize();
+	SWidget::CacheDesiredSize(LayoutScaleMultiplier);
 }
 
-FVector2D SMultiLineEditableText::ComputeDesiredSize() const
+FVector2D SMultiLineEditableText::ComputeDesiredSize( float ) const
 {
 	const float FontMaxCharHeight = FTextEditHelper::GetFontHeight(TextStyle.Font);
 	const float CaretWidth = FTextEditHelper::CalculateCaretWidth(FontMaxCharHeight);
diff --git a/Engine/Source/Runtime/Slate/Private/Widgets/Text/SRichTextBlock.cpp b/Engine/Source/Runtime/Slate/Private/Widgets/Text/SRichTextBlock.cpp
--- a/Engine/Source/Runtime/Slate/Private/Widgets/Text/SRichTextBlock.cpp
+++ b/Engine/Source/Runtime/Slate/Private/Widgets/Text/SRichTextBlock.cpp
@@ -51,7 +51,7 @@ int32 SRichTextBlock::OnPaint( const FPaintArgs& Args, const FGeometry& Allotted
 	return LayerId;
 }
 
-FVector2D SRichTextBlock::ComputeDesiredSize() const
+FVector2D SRichTextBlock::ComputeDesiredSize(float LayoutScaleMultiplier) const
 {
 	// todo: jdale - The scale needs to be passed to ComputeDesiredSize
 	//const float Scale = CachedScale;
@@ -59,7 +59,7 @@ FVector2D SRichTextBlock::ComputeDesiredSize() const
 	// ComputeDesiredSize will also update the text layout cache if required
 	const FVector2D TextSize = TextLayoutCache->ComputeDesiredSize(
 		FTextBlockLayout::FWidgetArgs(BoundText, HighlightText, WrapTextAt, AutoWrapText, Margin, LineHeightPercentage, Justification), 
-		/*Scale, */TextStyle
+		LayoutScaleMultiplier, TextStyle
 		);
 
 	return TextSize;
diff --git a/Engine/Source/Runtime/Slate/Private/Widgets/Text/STextBlock.cpp b/Engine/Source/Runtime/Slate/Private/Widgets/Text/STextBlock.cpp
--- a/Engine/Source/Runtime/Slate/Private/Widgets/Text/STextBlock.cpp
+++ b/Engine/Source/Runtime/Slate/Private/Widgets/Text/STextBlock.cpp
@@ -168,7 +168,7 @@ FReply STextBlock::OnMouseButtonDoubleClick( const FGeometry& InMyGeometry, cons
 	return FReply::Unhandled();
 }
 
-FVector2D STextBlock::ComputeDesiredSize() const
+FVector2D STextBlock::ComputeDesiredSize(float LayoutScaleMultiplier) const
 {
 #if SLATE_HD_STATS
 	SCOPE_CYCLE_COUNTER( STAT_SlateComputeDesiredSize_STextBlock );
@@ -182,7 +182,7 @@ FVector2D STextBlock::ComputeDesiredSize() const
 	// ComputeDesiredSize will also update the text layout cache if required
 	const FVector2D TextSize = TextLayoutCache->ComputeDesiredSize(
 		FTextBlockLayout::FWidgetArgs(BoundText, HighlightText, WrapTextAt, AutoWrapText, Margin, LineHeightPercentage, Justification), 
-		/*Scale, */GetComputedTextStyle()
+		LayoutScaleMultiplier, GetComputedTextStyle()
 		);
 
 #else//WITH_FANCY_TEXT
diff --git a/Engine/Source/Runtime/Slate/Private/Widgets/Text/TextBlockLayout.cpp b/Engine/Source/Runtime/Slate/Private/Widgets/Text/TextBlockLayout.cpp
--- a/Engine/Source/Runtime/Slate/Private/Widgets/Text/TextBlockLayout.cpp
+++ b/Engine/Source/Runtime/Slate/Private/Widgets/Text/TextBlockLayout.cpp
@@ -18,17 +18,17 @@ FTextBlockLayout::FTextBlockLayout(FTextBlockStyle InDefaultTextStyle, TSharedRe
 	, Marshaller(InMarshaller)
 	, TextHighlighter(FSlateTextHighlightRunRenderer::Create())
 	, CachedSize(ForceInitToZero)
-	, CachedDesiredSize(ForceInitToZero)
 {
 	TextLayout->SetLineBreakIterator(InLineBreakPolicy);
 }
 
-FVector2D FTextBlockLayout::ComputeDesiredSize(const FWidgetArgs& InWidgetArgs/*, const float InScale*/, const FTextBlockStyle& InTextStyle)
+FVector2D FTextBlockLayout::ComputeDesiredSize(const FWidgetArgs& InWidgetArgs, const float InScale, const FTextBlockStyle& InTextStyle)
 {
 	TextLayout->SetWrappingWidth(CalculateWrappingWidth(InWidgetArgs));
 	TextLayout->SetMargin(InWidgetArgs.Margin.Get());
 	TextLayout->SetJustification(InWidgetArgs.Justification.Get());
 	TextLayout->SetLineHeightPercentage(InWidgetArgs.LineHeightPercentage.Get());
+	TextLayout->SetScale(InScale);
 
 	// Has the style used for this text block changed?
 	if(!IsStyleUpToDate(InTextStyle))
@@ -74,24 +74,10 @@ FVector2D FTextBlockLayout::ComputeDesiredSize(const FWidgetArgs& InWidgetArgs/*
 		}
 	}
 
-	// We need to update our cached desired size if the text layout has become dirty
-	// todo: jdale - This is a hack until we can perform accurate measuring in ComputeDesiredSize
-	if(TextLayout->IsLayoutDirty())
-	{
-		// The desired size must always have a scale of 1, OnPaint will make sure the scale is set correctly for painting
-		TextLayout->SetScale(1.0f);
-		TextLayout->UpdateIfNeeded();
-
-		CachedDesiredSize = TextLayout->GetSize();
-	}
-	else
-	{
-		// This logic may look odd, but IsLayoutDirty() only checks that we've made a change that might affect the layout (which in turn might affect the 
-		// desired size), however there's also highlight changes, which don't affect the size, but still require a call to UpdateIfNeeded() to be applied
-		TextLayout->UpdateIfNeeded();
-	}
-
-	return CachedDesiredSize;
+	// We need to update our size if the text layout has become dirty
+	TextLayout->UpdateIfNeeded();
+	
+	return TextLayout->GetSize();
 }
 
 int32 FTextBlockLayout::OnPaint(const FPaintArgs& InPaintArgs, const FGeometry& InAllottedGeometry, const FSlateRect& InClippingRect, FSlateWindowElementList& OutDrawElements, int32 LayerId, const FWidgetStyle& InWidgetStyle, bool bParentEnabled)
@@ -123,7 +109,6 @@ int32 FTextBlockLayout::OnPaint(const FPaintArgs& InPaintArgs, const FGeometry&
 		}
 	}
 
-	TextLayout->SetScale(InAllottedGeometry.Scale);
 	TextLayout->SetVisibleRegion(InAllottedGeometry.Size, AutoScrollValue);
 
 	TextLayout->UpdateIfNeeded();
diff --git a/Engine/Source/Runtime/Slate/Private/Widgets/Views/SListPanel.cpp b/Engine/Source/Runtime/Slate/Private/Widgets/Views/SListPanel.cpp
--- a/Engine/Source/Runtime/Slate/Private/Widgets/Views/SListPanel.cpp
+++ b/Engine/Source/Runtime/Slate/Private/Widgets/Views/SListPanel.cpp
@@ -158,7 +158,7 @@ void SListPanel::Tick( const FGeometry& AllottedGeometry, const double InCurrent
 /**
  * Simply the sum of all the children (vertically), and the largest width (horizontally).
  */
-FVector2D SListPanel::ComputeDesiredSize() const
+FVector2D SListPanel::ComputeDesiredSize( float ) const
 {
 	float MaxWidth = 0;
 	float TotalHeight = 0;
diff --git a/Engine/Source/Runtime/Slate/Public/Widgets/Colors/SColorSpectrum.h b/Engine/Source/Runtime/Slate/Public/Widgets/Colors/SColorSpectrum.h
--- a/Engine/Source/Runtime/Slate/Public/Widgets/Colors/SColorSpectrum.h
+++ b/Engine/Source/Runtime/Slate/Public/Widgets/Colors/SColorSpectrum.h
@@ -45,7 +45,7 @@ public:
 
 	// SWidget overrides
 
-	virtual FVector2D ComputeDesiredSize( ) const;
+	virtual FVector2D ComputeDesiredSize(float) const;
 	virtual FReply OnMouseButtonDoubleClick( const FGeometry& InMyGeometry, const FPointerEvent& InMouseEvent );
 	virtual FReply OnMouseButtonDown( const FGeometry& MyGeometry, const FPointerEvent& MouseEvent );
 	virtual FReply OnMouseButtonUp( const FGeometry& MyGeometry, const FPointerEvent& MouseEvent );
diff --git a/Engine/Source/Runtime/Slate/Public/Widgets/Colors/SColorWheel.h b/Engine/Source/Runtime/Slate/Public/Widgets/Colors/SColorWheel.h
--- a/Engine/Source/Runtime/Slate/Public/Widgets/Colors/SColorWheel.h
+++ b/Engine/Source/Runtime/Slate/Public/Widgets/Colors/SColorWheel.h
@@ -45,7 +45,7 @@ public:
 
 	// SWidget overrides
 
-	virtual FVector2D ComputeDesiredSize( ) const;
+	virtual FVector2D ComputeDesiredSize(float) const;
 	virtual FReply OnMouseButtonDoubleClick( const FGeometry& InMyGeometry, const FPointerEvent& InMouseEvent );
 	virtual FReply OnMouseButtonDown( const FGeometry& MyGeometry, const FPointerEvent& MouseEvent );
 	virtual FReply OnMouseButtonUp( const FGeometry& MyGeometry, const FPointerEvent& MouseEvent );
diff --git a/Engine/Source/Runtime/Slate/Public/Widgets/Input/SSlider.h b/Engine/Source/Runtime/Slate/Public/Widgets/Input/SSlider.h
--- a/Engine/Source/Runtime/Slate/Public/Widgets/Input/SSlider.h
+++ b/Engine/Source/Runtime/Slate/Public/Widgets/Input/SSlider.h
@@ -89,7 +89,7 @@ public:
 	// SWidget overrides
 
 	virtual int32 OnPaint( const FPaintArgs& Args, const FGeometry& AllottedGeometry, const FSlateRect& MyClippingRect, FSlateWindowElementList& OutDrawElements, int32 LayerId, const FWidgetStyle& InWidgetStyle, bool bParentEnabled ) const override;
-	virtual FVector2D ComputeDesiredSize() const;
+	virtual FVector2D ComputeDesiredSize(float) const;
 	virtual FReply OnMouseButtonDown( const FGeometry& MyGeometry, const FPointerEvent& MouseEvent ) override;
 	virtual FReply OnMouseButtonUp( const FGeometry& MyGeometry, const FPointerEvent& MouseEvent ) override;
 	virtual FReply OnMouseMove( const FGeometry& MyGeometry, const FPointerEvent& MouseEvent ) override;
diff --git a/Engine/Source/Runtime/Slate/Public/Widgets/Layout/SDPIScaler.h b/Engine/Source/Runtime/Slate/Public/Widgets/Layout/SDPIScaler.h
--- a/Engine/Source/Runtime/Slate/Public/Widgets/Layout/SDPIScaler.h
+++ b/Engine/Source/Runtime/Slate/Public/Widgets/Layout/SDPIScaler.h
@@ -32,7 +32,7 @@ public:
 
 	virtual void OnArrangeChildren( const FGeometry& AllottedGeometry, FArrangedChildren& ArrangedChildren ) const override;
 	
-	virtual FVector2D ComputeDesiredSize() const override;
+	virtual FVector2D ComputeDesiredSize(float) const override;
 
 	virtual FChildren* GetChildren() override;
 
@@ -44,6 +44,8 @@ public:
 
 protected:
 
+	virtual float GetRelativeLayoutScale(const FSlotBase& Child) const override;
+
 	/**
 	 * The scaling factor from 1:1 pixel to slate unit (SU) ratio.
 	 * e.g. Value of 2 would mean that there are two pixels for every slate unit in every dimension.
diff --git a/Engine/Source/Runtime/Slate/Public/Widgets/Layout/SGridPanel.h b/Engine/Source/Runtime/Slate/Public/Widgets/Layout/SGridPanel.h
--- a/Engine/Source/Runtime/Slate/Public/Widgets/Layout/SGridPanel.h
+++ b/Engine/Source/Runtime/Slate/Public/Widgets/Layout/SGridPanel.h
@@ -205,8 +205,8 @@ public:
 
 	virtual int32 OnPaint( const FPaintArgs& Args, const FGeometry& AllottedGeometry, const FSlateRect& MyClippingRect, FSlateWindowElementList& OutDrawElements, int32 LayerId, const FWidgetStyle& InWidgetStyle, bool bParentEnabled ) const override;
 	virtual void OnArrangeChildren( const FGeometry& AllottedGeometry, FArrangedChildren& ArrangedChildren ) const override;
-	virtual void CacheDesiredSize() override;
-	virtual FVector2D ComputeDesiredSize() const override;
+	virtual void CacheDesiredSize(float) override;
+	virtual FVector2D ComputeDesiredSize(float) const override;
 	virtual FChildren* GetChildren() override;
 
 private:
diff --git a/Engine/Source/Runtime/Slate/Public/Widgets/Layout/SResponsiveGridPanel.h b/Engine/Source/Runtime/Slate/Public/Widgets/Layout/SResponsiveGridPanel.h
--- a/Engine/Source/Runtime/Slate/Public/Widgets/Layout/SResponsiveGridPanel.h
+++ b/Engine/Source/Runtime/Slate/Public/Widgets/Layout/SResponsiveGridPanel.h
@@ -143,8 +143,8 @@ public:
 
 	virtual int32 OnPaint( const FPaintArgs& Args, const FGeometry& AllottedGeometry, const FSlateRect& MyClippingRect, FSlateWindowElementList& OutDrawElements, int32 LayerId, const FWidgetStyle& InWidgetStyle, bool bParentEnabled ) const override;
 	virtual void OnArrangeChildren( const FGeometry& AllottedGeometry, FArrangedChildren& ArrangedChildren ) const override;
-	virtual void CacheDesiredSize() override;
-	virtual FVector2D ComputeDesiredSize() const override;
+	virtual void CacheDesiredSize(float) override;
+	virtual FVector2D ComputeDesiredSize(float) const override;
 	virtual FChildren* GetChildren() override;
 
 private:
diff --git a/Engine/Source/Runtime/Slate/Public/Widgets/Layout/SSplitter.h b/Engine/Source/Runtime/Slate/Public/Widgets/Layout/SSplitter.h
--- a/Engine/Source/Runtime/Slate/Public/Widgets/Layout/SSplitter.h
+++ b/Engine/Source/Runtime/Slate/Public/Widgets/Layout/SSplitter.h
@@ -150,7 +150,7 @@ public:
 	 * A Panel's desired size in the space required to arrange of its children on the screen while respecting all of
 	 * the children's desired sizes and any layout-related options specified by the user. See StackPanel for an example.
 	 */
-	virtual FVector2D ComputeDesiredSize() const override;
+	virtual FVector2D ComputeDesiredSize(float) const override;
 
 	/**
 	 * All widgets must provide a way to access their children in a layout-agnostic way.
@@ -397,7 +397,7 @@ private:
 	 * A Panel's desired size in the space required to arrange of its children on the screen while respecting all of
 	 * the children's desired sizes and any layout-related options specified by the user. See StackPanel for an example.
 	 */
-	virtual FVector2D ComputeDesiredSize() const override;
+	virtual FVector2D ComputeDesiredSize(float) const override;
 
 	/**
 	 * All widgets must provide a way to access their children in a layout-agnostic way.
diff --git a/Engine/Source/Runtime/Slate/Public/Widgets/Layout/SWidgetSwitcher.h b/Engine/Source/Runtime/Slate/Public/Widgets/Layout/SWidgetSwitcher.h
--- a/Engine/Source/Runtime/Slate/Public/Widgets/Layout/SWidgetSwitcher.h
+++ b/Engine/Source/Runtime/Slate/Public/Widgets/Layout/SWidgetSwitcher.h
@@ -139,7 +139,7 @@ protected:
 	// SCompoundWidget interface
 
 	virtual void OnArrangeChildren( const FGeometry& AllottedGeometry, FArrangedChildren& ArrangedChildren ) const override;
-	virtual FVector2D ComputeDesiredSize( ) const override;
+	virtual FVector2D ComputeDesiredSize( float ) const override;
 	virtual FChildren* GetChildren( ) override;
 
 private:
@@ -163,6 +163,8 @@ private:
 		
 	private:
 
+		virtual const FSlotBase& GetSlotAt(int32 ChildIndex) const override { return (*AllChildren)[ChildIndex]; }
+
 		TPanelChildren<FSlot>* AllChildren;
 		const TAttribute<int32>* WidgetIndex;
 
diff --git a/Engine/Source/Runtime/Slate/Public/Widgets/Notifications/SProgressBar.h b/Engine/Source/Runtime/Slate/Public/Widgets/Notifications/SProgressBar.h
--- a/Engine/Source/Runtime/Slate/Public/Widgets/Notifications/SProgressBar.h
+++ b/Engine/Source/Runtime/Slate/Public/Widgets/Notifications/SProgressBar.h
@@ -76,7 +76,7 @@ public:
 
 	virtual int32 OnPaint( const FPaintArgs& Args, const FGeometry& AllottedGeometry, const FSlateRect& MyClippingRect, FSlateWindowElementList& OutDrawElements, int32 LayerId, const FWidgetStyle& InWidgetStyle, bool bParentEnabled ) const override;
 
-	FVector2D ComputeDesiredSize() const;
+	virtual FVector2D ComputeDesiredSize(float) const override;
 
 	/** See attribute Percent */
 	void SetPercent(TAttribute< TOptional<float> > InPercent);
diff --git a/Engine/Source/Runtime/Slate/Public/Widgets/Text/TextBlockLayout.h b/Engine/Source/Runtime/Slate/Public/Widgets/Text/TextBlockLayout.h
--- a/Engine/Source/Runtime/Slate/Public/Widgets/Text/TextBlockLayout.h
+++ b/Engine/Source/Runtime/Slate/Public/Widgets/Text/TextBlockLayout.h
@@ -47,7 +47,7 @@ public:
 	/**
 	 * Get the computed desired size for this layout, updating the internal cache as required
 	 */
-	FVector2D ComputeDesiredSize(const FWidgetArgs& InWidgetArgs/*, const float InScale*/, const FTextBlockStyle& InTextStyle);
+	FVector2D ComputeDesiredSize(const FWidgetArgs& InWidgetArgs, const float InScale, const FTextBlockStyle& InTextStyle);
 
 	/**
 	 * Paint this layout, updating the internal cache as required
@@ -96,9 +96,6 @@ private:
 	/** The last known size of the layout from the previous OnPaint, used to guess at an auto-wrapping width in ComputeDesiredSize */
 	FVector2D CachedSize;
 
-	/** The cached desired size - this should be updated using a text layout scale of 1 whenever the text layout becomes dirty */
-	FVector2D CachedDesiredSize;
-
 	/** The state of the text the last time it was updated (used to allow updates when the text is changed) */
 	FTextSnapshot TextLastUpdate;
 
diff --git a/Engine/Source/Runtime/Slate/Public/Widgets/Views/SListView.h b/Engine/Source/Runtime/Slate/Public/Widgets/Views/SListView.h
--- a/Engine/Source/Runtime/Slate/Public/Widgets/Views/SListView.h
+++ b/Engine/Source/Runtime/Slate/Public/Widgets/Views/SListView.h
@@ -791,11 +791,12 @@ public:
 			bool bGeneratedEnoughForSmoothScrolling = false;
 			bool AtEndOfList = false;
 
+			const float LayoutScaleMultiplier = MyGeometry.GetAccumulatedLayoutTransform().GetScale();
 			for( int32 ItemIndex = StartIndex; !bGeneratedEnoughForSmoothScrolling && ItemIndex < SourceItems->Num(); ++ItemIndex )
 			{
 				const ItemType& CurItem = (*SourceItems)[ItemIndex];
 
-				const float ItemHeight = GenerateWidgetForItem( CurItem, ItemIndex, StartIndex );
+				const float ItemHeight = GenerateWidgetForItem(CurItem, ItemIndex, StartIndex, LayoutScaleMultiplier);
 
 				const bool IsFirstItem = ItemIndex == StartIndex;
 
@@ -846,7 +847,7 @@ public:
 				{
 					const ItemType& CurItem = (*SourceItems)[ItemIndex];
 
-					const float ItemHeight = GenerateWidgetForItem( CurItem, ItemIndex, StartIndex );
+					const float ItemHeight = GenerateWidgetForItem(CurItem, ItemIndex, StartIndex, LayoutScaleMultiplier);
 
 					if (HeightGeneratedSoFar + ItemHeight > MyGeometry.Size.Y)
 					{
@@ -869,7 +870,7 @@ public:
 
 	}
 
-	float GenerateWidgetForItem( const ItemType& CurItem, int32 ItemIndex, int32 StartIndex )
+	float GenerateWidgetForItem( const ItemType& CurItem, int32 ItemIndex, int32 StartIndex, float LayoutScaleMultiplier )
 	{
 		// Find a previously generated Widget for this item, if one exists.
 		TSharedPtr<ITableRow> WidgetForItem = WidgetGenerator.GetWidgetForItem( CurItem );
@@ -889,7 +890,7 @@ public:
 
 		// We rely on the widgets desired size in order to determine how many will fit on screen.
 		const TSharedRef<SWidget> NewlyGeneratedWidget = WidgetForItem->AsWidget();
-		NewlyGeneratedWidget->SlatePrepass();		
+		NewlyGeneratedWidget->SlatePrepass(LayoutScaleMultiplier);
 
 		const bool IsFirstWidgetOnScreen = (ItemIndex == StartIndex);
 		const float ItemHeight = NewlyGeneratedWidget->GetDesiredSize().Y;
@@ -1189,6 +1190,7 @@ protected:
 			if ( SourceItems != nullptr && SourceItems->Num() > 0 )
 			{
 				int ItemIndex = StartingItemIndex;
+				const float LayoutScaleMultiplier = MyGeometry.GetAccumulatedLayoutTransform().GetScale();
 				while( AbsScrollByAmount != 0 && ItemIndex < SourceItems->Num() && ItemIndex >= 0 )
 				{
 					const ItemType CurItem = (*SourceItems)[ ItemIndex ];
@@ -1206,7 +1208,7 @@ protected:
 						// Let the item generator know that we encountered the current Item and associated Widget.
 						WidgetGenerator.OnItemSeen( CurItem, RowWidget.ToSharedRef() );
 
-						RowWidget->AsWidget()->SlatePrepass();
+						RowWidget->AsWidget()->SlatePrepass(LayoutScaleMultiplier);
 					}
 
 					if ( ScrollByAmountInSlateUnits > 0 )
diff --git a/Engine/Source/Runtime/Slate/Public/Widgets/Views/STileView.h b/Engine/Source/Runtime/Slate/Public/Widgets/Views/STileView.h
--- a/Engine/Source/Runtime/Slate/Public/Widgets/Views/STileView.h
+++ b/Engine/Source/Runtime/Slate/Public/Widgets/Views/STileView.h
@@ -191,6 +191,7 @@ public:
 			const double EndOfListOffset = NumItemsPaddedToFillLastRow - NumItemsWide * RowsPerScreen;
 			const double ClampedScrollOffset = FMath::Clamp(STableViewBase::ScrollOffset, 0.0, EndOfListOffset);
 			const bool bAtEndOfList = (STableViewBase::ScrollOffset >= ClampedScrollOffset);
+			const float LayoutScaleMultiplier = MyGeometry.GetAccumulatedLayoutTransform().GetScale();
 			
 			// Once we run out of vertical and horizontal space, we stop generating widgets.
 			float WidthUsedSoFar = 0.0f;
@@ -224,7 +225,7 @@ public:
 					}
 				}
 
-				const float GeneratedItemHeight = SListView<ItemType>::GenerateWidgetForItem( CurItem, ItemIndex, StartIndex );
+				const float GeneratedItemHeight = SListView<ItemType>::GenerateWidgetForItem(CurItem, ItemIndex, StartIndex, LayoutScaleMultiplier);
 
 				// The widget used up some of the available horizontal space.
 				WidthUsedSoFar += ItemWidth;
diff --git a/Engine/Source/Runtime/SlateCore/Private/Widgets/SBoxPanel.cpp b/Engine/Source/Runtime/SlateCore/Private/Widgets/SBoxPanel.cpp
--- a/Engine/Source/Runtime/SlateCore/Private/Widgets/SBoxPanel.cpp
+++ b/Engine/Source/Runtime/SlateCore/Private/Widgets/SBoxPanel.cpp
@@ -233,7 +233,7 @@ static FVector2D ComputeDesiredSizeForBox( const TPanelChildren<SBoxPanel::FSlot
  *
  * @return The desired size.
  */
-FVector2D SBoxPanel::ComputeDesiredSize() const
+FVector2D SBoxPanel::ComputeDesiredSize( float ) const
 {
 	return (Orientation == Orient_Horizontal)
 		? ComputeDesiredSizeForBox<Orient_Horizontal>(this->Children)
diff --git a/Engine/Source/Runtime/SlateCore/Private/Widgets/SCompoundWidget.cpp b/Engine/Source/Runtime/SlateCore/Private/Widgets/SCompoundWidget.cpp
--- a/Engine/Source/Runtime/SlateCore/Private/Widgets/SCompoundWidget.cpp
+++ b/Engine/Source/Runtime/SlateCore/Private/Widgets/SCompoundWidget.cpp
@@ -38,7 +38,7 @@ FChildren* SCompoundWidget::GetChildren()
 }
 
 
-FVector2D SCompoundWidget::ComputeDesiredSize() const
+FVector2D SCompoundWidget::ComputeDesiredSize( float ) const
 {
 	EVisibility ChildVisibility = ChildSlot.GetWidget()->GetVisibility();
 	if ( ChildVisibility != EVisibility::Collapsed )
diff --git a/Engine/Source/Runtime/SlateCore/Private/Widgets/SNullWidget.cpp b/Engine/Source/Runtime/SlateCore/Private/Widgets/SNullWidget.cpp
--- a/Engine/Source/Runtime/SlateCore/Private/Widgets/SNullWidget.cpp
+++ b/Engine/Source/Runtime/SlateCore/Private/Widgets/SNullWidget.cpp
@@ -45,7 +45,7 @@ public:
 		return LayerId;
 	}
 
-	virtual FVector2D ComputeDesiredSize( ) const override final
+	virtual FVector2D ComputeDesiredSize(float) const override final
 	{
 		return FVector2D(0.0f, 0.0f);
 	}
diff --git a/Engine/Source/Runtime/SlateCore/Private/Widgets/SOverlay.cpp b/Engine/Source/Runtime/SlateCore/Private/Widgets/SOverlay.cpp
--- a/Engine/Source/Runtime/SlateCore/Private/Widgets/SOverlay.cpp
+++ b/Engine/Source/Runtime/SlateCore/Private/Widgets/SOverlay.cpp
@@ -54,7 +54,7 @@ void SOverlay::OnArrangeChildren( const FGeometry& AllottedGeometry, FArrangedCh
  *
  * @return The desired size.
  */
-FVector2D SOverlay::ComputeDesiredSize() const
+FVector2D SOverlay::ComputeDesiredSize( float ) const
 {
 	FVector2D MaxSize(0,0);
 	for ( int32 ChildIndex=0; ChildIndex < Children.Num(); ++ChildIndex )
diff --git a/Engine/Source/Runtime/SlateCore/Private/Widgets/SWidget.cpp b/Engine/Source/Runtime/SlateCore/Private/Widgets/SWidget.cpp
--- a/Engine/Source/Runtime/SlateCore/Private/Widgets/SWidget.cpp
+++ b/Engine/Source/Runtime/SlateCore/Private/Widgets/SWidget.cpp
@@ -367,7 +367,12 @@ void SWidget::TickWidgetsRecursively( const FGeometry& AllottedGeometry, const d
 }
 
 
-void SWidget::SlatePrepass()	
+void SWidget::SlatePrepass()
+{
+	SlatePrepass( FSlateApplicationBase::Get().GetApplicationScale() );
+}
+
+void SWidget::SlatePrepass(float LayoutScaleMultiplier)
 {
 	// Cache child desired sizes first. This widget's desired size is
 	// a function of its children's sizes.
@@ -378,20 +383,21 @@ void SWidget::SlatePrepass()
 		const TSharedRef<SWidget>& Child = MyChildren->GetChildAt(ChildIndex);
 		if ( Child->Visibility.Get() != EVisibility::Collapsed )
 		{
+			const float ChildLayoutScaleMultiplier = GetRelativeLayoutScale( MyChildren->GetSlotAt(ChildIndex) );
 			// Recur: Descend down the widget tree.
-			Child->SlatePrepass();
+			Child->SlatePrepass(LayoutScaleMultiplier*ChildLayoutScaleMultiplier);
 		}
 	}
 
 	// Cache this widget's desired size.
-	CacheDesiredSize();
+	CacheDesiredSize(LayoutScaleMultiplier);
 }
 
 
-void SWidget::CacheDesiredSize()
+void SWidget::CacheDesiredSize(float LayoutScaleMultiplier)
 {
 	// Cache this widget's desired size.
-	this->Advanced_SetDesiredSize( this->ComputeDesiredSize() );
+	this->Advanced_SetDesiredSize(this->ComputeDesiredSize(LayoutScaleMultiplier));
 }
 
 
@@ -666,6 +672,11 @@ int32 SWidget::Paint(const FPaintArgs& Args, const FGeometry& AllottedGeometry,
 	return NewLayerID;
 }
 
+float SWidget::GetRelativeLayoutScale(const FSlotBase& Child) const
+{
+	return 1.0f;
+}
+
 void SWidget::ArrangeChildren(const FGeometry& AllottedGeometry, FArrangedChildren& ArrangedChildren) const
 {
 	SCOPE_CYCLE_COUNTER(STAT_SlateArrangeChildren);
diff --git a/Engine/Source/Runtime/SlateCore/Private/Widgets/SWindow.cpp b/Engine/Source/Runtime/SlateCore/Private/Widgets/SWindow.cpp
--- a/Engine/Source/Runtime/SlateCore/Private/Widgets/SWindow.cpp
+++ b/Engine/Source/Runtime/SlateCore/Private/Widgets/SWindow.cpp
@@ -155,7 +155,7 @@ private:
 		}
 	}
 
-	virtual FVector2D ComputeDesiredSize() const override
+	virtual FVector2D ComputeDesiredSize(float) const override
 	{
 		return FVector2D(100,100);
 	}
@@ -709,7 +709,7 @@ void SWindow::MoveWindowTo( FVector2D NewPosition )
 {
 	if (NativeWindow.IsValid())
 	{
-#if PLATFORM_LINUX
+#if 1//PLATFORM_LINUX
 		// Slate code often expects cached screen position to be accurate immediately after the move.
 		// This expectation is generally invalid (see UE-1308) as there may be a delay before the OS reports it back.
 		// This hack sets the position speculatively, keeping Slate happy while also giving the OS chance to report it
@@ -730,7 +730,7 @@ void SWindow::ReshapeWindow( FVector2D NewPosition, FVector2D NewSize )
 {
 	if (NativeWindow.IsValid())
 	{
-#if PLATFORM_LINUX
+#if 1//PLATFORM_LINUX
 		// Slate code often expects cached screen position to be accurate immediately after the move.
 		// This expectation is generally invalid (see UE-1308) as there may be a delay before the OS reports it back.
 		// This hack sets the position speculatively, keeping Slate happy while also giving the OS chance to report it
@@ -1091,7 +1091,7 @@ void SWindow::ShowWindow()
 		// Repositioning the window on show with the new size solves this.
 		if ( SizingRule == ESizingRule::Autosized && AutoCenterRule != EAutoCenter::None )
 		{
-			SlatePrepass();
+			SlatePrepass( FSlateApplicationBase::Get().GetApplicationScale() );
 			ReshapeWindow( InitialDesiredScreenPosition - (GetDesiredSize() * 0.5f), GetDesiredSize() );
 		}
 
@@ -1439,10 +1439,14 @@ FReply SWindow::OnMouseMove( const FGeometry& MyGeometry, const FPointerEvent& M
 
 }
 
-FVector2D SWindow::ComputeDesiredSize() const
+FVector2D SWindow::ComputeDesiredSize(float LayoutScaleMultiplier) const
 {
-	const float Scale = FSlateApplicationBase::Get().GetApplicationScale();
-	return SCompoundWidget::ComputeDesiredSize() * Scale;
+	return SCompoundWidget::ComputeDesiredSize(LayoutScaleMultiplier) * LayoutScaleMultiplier;
+}
+
+float SWindow::GetRelativeLayoutScale(const FSlotBase& Child) const
+{
+	return FSlateApplicationBase::Get().GetApplicationScale();
 }
 
 const TArray< TSharedRef<SWindow> >& SWindow::GetChildWindows() const
diff --git a/Engine/Source/Runtime/SlateCore/Public/Layout/Children.h b/Engine/Source/Runtime/SlateCore/Public/Layout/Children.h
--- a/Engine/Source/Runtime/SlateCore/Public/Layout/Children.h
+++ b/Engine/Source/Runtime/SlateCore/Public/Layout/Children.h
@@ -23,6 +23,11 @@ public:
 	/** @return const pointer to the Widget at the specified Index. */
 	virtual TSharedRef<const SWidget> GetChildAt( int32 Index ) const = 0;
 
+protected:
+	friend class SWidget;
+	/** @return the const reference to the slot at the specified Index */
+	virtual const FSlotBase& GetSlotAt(int32 ChildIndex) const = 0;
+
 protected:
 	virtual ~FChildren(){}
 };
@@ -53,6 +58,16 @@ public:
 		return TSharedPtr<const SWidget>(nullptr).ToSharedRef();
 	}
 
+private:
+	friend class SWidget;
+	virtual const FSlotBase& GetSlotAt(int32 ChildIndex) const override
+	{
+		check(false);
+		static FSlotBase NullSlot;
+		return NullSlot;
+	}
+
+
 };
 
 /**
@@ -67,6 +82,8 @@ public:
 	virtual int32 Num() const { return 1; }
 	virtual TSharedRef<SWidget> GetChildAt( int32 ChildIndex ) { check(ChildIndex == 0); return FSlotBase::GetWidget(); }
 	virtual TSharedRef<const SWidget> GetChildAt( int32 ChildIndex ) const { check(ChildIndex == 0); return FSlotBase::GetWidget(); }
+private:
+	virtual const FSlotBase& GetSlotAt(int32 ChildIndex) const override { check(ChildIndex == 0); return *this; }
 };
 
 /**
@@ -86,7 +103,11 @@ public:
 	virtual int32 Num() const { return WidgetPtr.IsValid() ? 1 : 0 ; }
 	virtual TSharedRef<SWidget> GetChildAt( int32 ChildIndex ) { check(ChildIndex == 0); return WidgetPtr.Pin().ToSharedRef(); }
 	virtual TSharedRef<const SWidget> GetChildAt( int32 ChildIndex ) const { check(ChildIndex == 0); return WidgetPtr.Pin().ToSharedRef(); }
-	
+
+private:
+	virtual const FSlotBase& GetSlotAt(int32 ChildIndex) const override { static FSlotBase NullSlot; check(ChildIndex == 0); return NullSlot; }
+
+public:
 	void AttachWidget(const TSharedPtr<SWidget>& InWidget)
 	{
 		WidgetPtr = InWidget;
@@ -180,11 +201,17 @@ public:
 template<typename SlotType>
 class TPanelChildren : public FChildren, private TIndirectArray< SlotType >
 {
+private:
+	virtual const FSlotBase& GetSlotAt(int32 ChildIndex) const override
+	{
+		return (*this)[ChildIndex];
+	}
+
 public:
 	TPanelChildren()
 	{
 	}
-
+	
 	virtual int32 Num() const override
 	{
 		return TIndirectArray<SlotType>::Num();
@@ -256,6 +283,14 @@ public:
 template<typename ChildType>
 class TSlotlessChildren : public FChildren, private TArray< TSharedRef<ChildType> >
 {
+private:
+	virtual const FSlotBase& GetSlotAt(int32 ChildIndex) const override
+	{
+		// @todo slate : slotless children should be removed altogether; for now they return a fake slot.
+		static FSlotBase NullSlot;
+		return NullSlot;
+	}
+
 public:
 	TSlotlessChildren()
 	{
diff --git a/Engine/Source/Runtime/SlateCore/Public/Widgets/SLeafWidget.h b/Engine/Source/Runtime/SlateCore/Public/Widgets/SLeafWidget.h
--- a/Engine/Source/Runtime/SlateCore/Public/Widgets/SLeafWidget.h
+++ b/Engine/Source/Runtime/SlateCore/Public/Widgets/SLeafWidget.h
@@ -38,7 +38,7 @@ private:
 	 * take child widgets into account as LeafWidgets have none by definition. For example, the TextBlock widget simply
 	 * measures the area necessary to display its text with the given font and font size.
 	 */
-	virtual FVector2D ComputeDesiredSize( ) const = 0;
+	virtual FVector2D ComputeDesiredSize(float) const = 0;
 	
 	/**
 	 * Overwritten from SWidget.
diff --git a/Engine/Source/Runtime/SlateCore/Public/Widgets/SPanel.h b/Engine/Source/Runtime/SlateCore/Public/Widgets/SPanel.h
--- a/Engine/Source/Runtime/SlateCore/Public/Widgets/SPanel.h
+++ b/Engine/Source/Runtime/SlateCore/Public/Widgets/SPanel.h
@@ -30,7 +30,7 @@ public:
 	 *
 	 * @return The desired size.
 	 */
-	virtual FVector2D ComputeDesiredSize( ) const = 0;
+	virtual FVector2D ComputeDesiredSize(float) const = 0;
 
 	/**
 	 * All widgets must provide a way to access their children in a layout-agnostic way.
diff --git a/Engine/Source/Runtime/SlateCore/Public/Widgets/SWidget.h b/Engine/Source/Runtime/SlateCore/Public/Widgets/SWidget.h
--- a/Engine/Source/Runtime/SlateCore/Public/Widgets/SWidget.h
+++ b/Engine/Source/Runtime/SlateCore/Public/Widgets/SWidget.h
@@ -477,21 +477,25 @@ public:
 	// LAYOUT
 	//
 
+	public:
+	
+	/** DEPRECATED version of SlatePrepass that assumes no scaling beyond AppScale*/
+	void SlatePrepass();
+
 	/**
-	 * Applies styles as it descends through the widget hierarchy.
-	 * Gathers desired sizes on the way up.
+	 * Descends to leafmost widgets in the hierarchy and gathers desired sizes on the way up.
 	 * i.e. Caches the desired size of all of this widget's children recursively, then caches desired size for itself.
 	 */
-	void SlatePrepass();
+	void SlatePrepass(float LayoutScaleMultiplier);
 
 	/** @return the DesiredSize that was computed the last time CacheDesiredSize() was called. */
-	const FVector2D& GetDesiredSize() const;
+	public:  const FVector2D& GetDesiredSize() const;
 
 	/**
 	 * The system calls this method. It performs a breadth-first traversal of every visible widget and asks
 	 * each widget to cache how big it needs to be in order to present all of its content.
 	 */
-	protected: virtual void CacheDesiredSize();
+	protected: virtual void CacheDesiredSize(float);
 	
 	/**
 	 * Explicitly set the desired size. This is highly advanced functionality that is meant
@@ -508,12 +512,17 @@ public:
 	 * widget is simulating a bouncing ball, you should just return a reasonable size; e.g. 160x160. Let the programmer set up a reasonable
 	 * rule of resizing the bouncy ball simulation.
 	 *
+	 * @param  LayoutScaleMultiplier    This parameter is safe to ignore for almost all widgets; only really affects text measuring.
+	 *
 	 * @return The desired size.
 	 */
-	private: virtual FVector2D ComputeDesiredSize() const = 0;
+	private: virtual FVector2D ComputeDesiredSize(float LayoutScaleMultiplier) const = 0;
 
 	public:
 
+	/** What is the Child's scale relative to this widget. */
+	virtual float GetRelativeLayoutScale( const FSlotBase& Child ) const;
+
 	/**
 	 * Non-virtual entry point for arrange children. ensures common work is executed before calling the virtual
 	 * ArrangeChildren function.
