We no longer refresh the graph panel when a user deletes or adds a graph node. This is a substantial improvement in responsiveness, especially when the user is editing large graphs.

diff --git a/Engine/Source/Editor/BlueprintGraph/Private/BlueprintNodeSpawner.cpp b/Engine/Source/Editor/BlueprintGraph/Private/BlueprintNodeSpawner.cpp
--- a/Engine/Source/Editor/BlueprintGraph/Private/BlueprintNodeSpawner.cpp
+++ b/Engine/Source/Editor/BlueprintGraph/Private/BlueprintNodeSpawner.cpp
@@ -283,7 +283,10 @@ UEdGraphNode* UBlueprintNodeSpawner::SpawnEdGraphNode(TSubclassOf<UEdGraphNode>
 			NewNode->PostPlacedNewNode();
 
 			ParentGraph->Modify();
-			// the FBlueprintMenuActionItem should do the selecting
+			// the FBlueprintMenuActionItem should do the selecting because we may spawn multiple nodes, and if we
+			// tag them as selected here (via bSelectNewNode) only the last one spawned will be selected. If we exposed
+			// the FEdGraphEditAction's list of nodes we could do creation and selection with a single event, something
+			// to consider..
 			ParentGraph->AddNode(NewNode, /*bFromUI =*/true, /*bSelectNewNode =*/false);
 		}
 
diff --git a/Engine/Source/Editor/GraphEditor/Private/SGraphEditorImpl.cpp b/Engine/Source/Editor/GraphEditor/Private/SGraphEditorImpl.cpp
--- a/Engine/Source/Editor/GraphEditor/Private/SGraphEditorImpl.cpp
+++ b/Engine/Source/Editor/GraphEditor/Private/SGraphEditorImpl.cpp
@@ -87,7 +87,14 @@ void SGraphEditorImpl::NotifyGraphChanged()
 
 void SGraphEditorImpl::OnGraphChanged(const FEdGraphEditAction& InAction)
 {
-	if ( !bIsActiveTimerRegistered )
+	// Actions with a specific purpose are handled by SGraphPanel::OnGraphChanged
+	if ( !bIsActiveTimerRegistered && 
+		 !( 
+			(InAction.Action & GRAPHACTION_AddNode) || 
+			(InAction.Action & GRAPHACTION_SelectNode)  ||
+			(InAction.Action & GRAPHACTION_RemoveNode) 
+		  ) 
+	   )
 	{
 		// Remove the old user interface nodes
 		GraphPanel->PurgeVisualRepresentation();
diff --git a/Engine/Source/Editor/GraphEditor/Private/SGraphPanel.cpp b/Engine/Source/Editor/GraphEditor/Private/SGraphPanel.cpp
--- a/Engine/Source/Editor/GraphEditor/Private/SGraphPanel.cpp
+++ b/Engine/Source/Editor/GraphEditor/Private/SGraphPanel.cpp
@@ -946,12 +946,12 @@ const TSharedRef<SGraphNode> SGraphPanel::GetChild(int32 ChildIndex)
 	return StaticCastSharedRef<SGraphNode>(Children[ChildIndex]);
 }
 
-void SGraphPanel::AddNode(UEdGraphNode* Node)
+void SGraphPanel::AddNode(UEdGraphNode* Node, bool bIsUserAdded, bool bShouldBeSelected)
 {
 	TSharedPtr<SGraphNode> NewNode = FNodeFactory::CreateNodeWidget(Node);
 	check(NewNode.IsValid());
 
-	const bool bWasUserAdded = (UserAddedNodes.Find(Node) != nullptr);
+	const bool bWasUserAdded = bIsUserAdded;
 
 	NewNode->SetIsEditable(IsEditable);
 	NewNode->SetDoubleClickEvent(OnNodeDoubleClicked);
@@ -972,6 +972,33 @@ void SGraphPanel::AddNode(UEdGraphNode* Node)
 		NewNode->PlaySpawnEffect();
 		NewNode->RequestRenameOnSpawn();
 	}
+
+	if( bShouldBeSelected )
+	{
+		DeferredSelectionTargetObjects.Add(Node);
+	}
+}
+
+void SGraphPanel::RemoveNode(UEdGraphNode* Node)
+{
+	for (int32 Iter = 0; Iter != Children.Num(); ++Iter)
+	{
+		TSharedRef<SGraphNode> Child = GetChild(Iter);
+		if( Child->GetNodeObj() == Node )
+		{
+			Children.RemoveAt(Iter);
+			break;
+		}
+	}
+	for (int32 Iter = 0; Iter != Children.Num(); ++Iter)
+	{
+		TSharedRef<SGraphNode> Child = StaticCastSharedRef<SGraphNode>(VisibleChildren[Iter]);
+		if (Child->GetNodeObj() == Node)
+		{
+			VisibleChildren.RemoveAt(Iter);
+			break;
+		}
+	}
 }
 
 TSharedPtr<SGraphNode> SGraphPanel::GetNodeWidgetFromGuid(FGuid Guid) const
@@ -997,30 +1024,12 @@ void SGraphPanel::Update()
 				UE_LOG(LogGraphPanel, Warning, TEXT("Found NULL Node in GraphObj array. A node type has been deleted without creating an ActiveClassRedictor to K2Node_DeadClass."));
 			}
 		}
-
-		// find the last selection action, and execute it
-		for (int32 ActionIndex = UserActions.Num() - 1; ActionIndex >= 0; --ActionIndex)
-		{
-			if (UserActions[ActionIndex].Action & GRAPHACTION_SelectNode)
-			{
-				DeferredSelectionTargetObjects.Empty();
-				for (const UEdGraphNode* Node : UserActions[ActionIndex].Nodes)
-				{
-					DeferredSelectionTargetObjects.Add(Node);
-				}
-				break;
-			}
-		}
 	}
 	else
 	{
 		RemoveAllNodes();
 	}
 
-	// Clean out set of added nodes
-	UserAddedNodes.Empty();
-	UserActions.Empty();
-
 	// Invoke any delegate methods
 	OnUpdateGraphPanel.ExecuteIfBound();
 }
@@ -1117,20 +1126,34 @@ void SGraphPanel::OnEndPIE( const bool bIsSimulating )
 void SGraphPanel::OnGraphChanged(const FEdGraphEditAction& EditAction)
 {
 	if ((EditAction.Graph == GraphObj) &&
-		(EditAction.Nodes.Num() > 0) &&
-		// We do not want to mark it as a UserAddedNode for graphs that do not currently have focus,
-		// this causes each one to want to do the effects and rename, which causes problems.
-		(HasKeyboardFocus() || HasFocusedDescendants()))
+		(EditAction.Nodes.Num() > 0))
 	{
-		int32 ActionIndex = UserActions.Num();
 		if (EditAction.Action & GRAPHACTION_AddNode)
 		{
+			// if we're the focused graph, play spawn effects:
+			for (UEdGraphNode* Node : EditAction.Nodes)
+			{
+				AddNode(Node, HasKeyboardFocus() || HasFocusedDescendants());
+			}
+		}
+		if( ( EditAction.Action & GRAPHACTION_SelectNode ) &&
+			(HasKeyboardFocus() || HasFocusedDescendants() ) )
+		{
+			DeferredSelectionTargetObjects.Empty();
 			for (const UEdGraphNode* Node : EditAction.Nodes)
 			{
-				UserAddedNodes.Add(Node, ActionIndex);
+				DeferredSelectionTargetObjects.Add(Node);
 			}
 		}
-		UserActions.Add(EditAction);
+		if( EditAction.Action & GRAPHACTION_RemoveNode )
+		{
+			for (UEdGraphNode* Node : EditAction.Nodes)
+			{
+				RemoveNode(Node);
+			}
+		}
+
+		// adding logic to here? You also need to update SGraphEditorImpl::OnGraphChanged..
 	}
 }
 
diff --git a/Engine/Source/Editor/GraphEditor/Private/SGraphPanel.h b/Engine/Source/Editor/GraphEditor/Private/SGraphPanel.h
--- a/Engine/Source/Editor/GraphEditor/Private/SGraphPanel.h
+++ b/Engine/Source/Editor/GraphEditor/Private/SGraphPanel.h
@@ -163,7 +163,10 @@ protected:
 	const TSharedRef<SGraphNode> GetChild(int32 ChildIndex);
 
 	/** Helper method to add a new node to the panel */
-	void AddNode (class UEdGraphNode* Node);
+	void AddNode(class UEdGraphNode* Node, bool bIsUserAdded = false, bool bShouldBeSelected = false);
+
+	/** Helper method to remove a node from the panel */
+	void RemoveNode(UEdGraphNode* Node);
 public:
 	/** Pin marked via shift-clicking */
 	TWeakPtr<SGraphPin> MarkedPin;
@@ -237,12 +240,6 @@ protected:
 	TAttribute<bool> ShowGraphStateOverlay;
 
 private:
-	/** Ordered list of user actions, as they came in */
-	TArray<FEdGraphEditAction> UserActions;
-
-	/** Map of recently added nodes for the panel (maps from added nodes to UserActions indices) */
-	TMap<const class UEdGraphNode*, int32> UserAddedNodes;
-
 	FOnGraphChanged::FDelegate MyRegisteredGraphChangedDelegate;
 	FDelegateHandle            MyRegisteredGraphChangedDelegateHandle;
 private:
diff --git a/Engine/Source/Editor/Kismet/Private/BlueprintActionMenuItem.cpp b/Engine/Source/Editor/Kismet/Private/BlueprintActionMenuItem.cpp
--- a/Engine/Source/Editor/Kismet/Private/BlueprintActionMenuItem.cpp
+++ b/Engine/Source/Editor/Kismet/Private/BlueprintActionMenuItem.cpp
@@ -133,7 +133,7 @@ UEdGraphNode* FBlueprintActionMenuItem::PerformAction(UEdGraph* ParentGraph, UEd
 		FromPin->Modify();
 	}
 
-	TSet<const UEdGraphNode*> NodesToFocus;
+	TSet<UEdGraphNode*> NodesToFocus;
 	int32 const PreSpawnNodeCount = ParentGraph->Nodes.Num();
 
 	UEdGraphNode* LastSpawnedNode = nullptr;
diff --git a/Engine/Source/Editor/Kismet/Private/BlueprintEditor.cpp b/Engine/Source/Editor/Kismet/Private/BlueprintEditor.cpp
--- a/Engine/Source/Editor/Kismet/Private/BlueprintEditor.cpp
+++ b/Engine/Source/Editor/Kismet/Private/BlueprintEditor.cpp
@@ -4830,9 +4830,9 @@ private:
 		NewNode->NodePosX = OldNode ? OldNode->NodePosX : NodePosX;
 		NewNode->NodePosY = OldNode ? OldNode->NodePosY : NodePosY;
 		NewNode->SetFlags(RF_Transactional);
-		Graph->AddNode(NewNode, false, false);
 		NewNode->PostPlacedNewNode();
 		NewNode->AllocateDefaultPins();
+		Graph->AddNode(NewNode, false, false);
 	}
 
 	bool Replace(UK2Node_Literal* OldTarget, UK2Node_CallFunctionOnMember* OldCall)
@@ -6948,7 +6948,7 @@ FReply FBlueprintEditor::OnSpawnGraphNodeByShortcut(FInputGesture InGesture, con
 	FBlueprintPaletteListBuilder PaletteBuilder(GetBlueprintObj());
 	FBlueprintSpawnNodeCommands::Get().GetGraphActionByGesture(InGesture, PaletteBuilder, InGraph);
 
-	TSet<const UEdGraphNode*> NodesToSelect;
+	TSet<UEdGraphNode*> NodesToSelect;
 	FVector2D NodeSpawnPos = InPosition;
 
 	bool bPerformedAction = false;
diff --git a/Engine/Source/Editor/Kismet/Private/SBlueprintDiff.cpp b/Engine/Source/Editor/Kismet/Private/SBlueprintDiff.cpp
--- a/Engine/Source/Editor/Kismet/Private/SBlueprintDiff.cpp
+++ b/Engine/Source/Editor/Kismet/Private/SBlueprintDiff.cpp
@@ -319,9 +319,6 @@ private:
 
 public:
 
-	/** Called when the Newer Graph is modified*/
-	void OnGraphChanged(const FEdGraphEditAction& Action);
-
 	/** Generate list of differences*/
 	TSharedRef<SWidget> GenerateDiffListWidget();
 
@@ -340,9 +337,6 @@ private:
 
 	/** ListView of differences */
 	TSharedPtr<SListViewType> DiffList;
-
-	/** Handle to the registered OnGraphChanged delegate. */
-	FDelegateHandle OnGraphChangedDelegateHandle;
 };
 
 TSharedRef<SWidget>	FDiffResultItem::GenerateWidget() const
@@ -366,21 +360,11 @@ FListItemGraphToDiff::FListItemGraphToDiff( class SBlueprintDiff* InDiff, class
 {
 	check(InGraphOld || InGraphNew); //one of them needs to exist
 
-	//need to know when it is modified
-	if(InGraphNew)
-	{
-		OnGraphChangedDelegateHandle = InGraphNew->AddOnGraphChangedHandler( FOnGraphChanged::FDelegate::CreateRaw(this, &FListItemGraphToDiff::OnGraphChanged));
-	}
-
 	BuildDiffSourceArray();
 }
 
 FListItemGraphToDiff::~FListItemGraphToDiff()
 {
-	if(GraphNew)
-	{
-		GraphNew->RemoveOnGraphChangedHandler( OnGraphChangedDelegateHandle);
-	}
 }
 
 TSharedRef<SWidget> FListItemGraphToDiff::GenerateWidget() 
@@ -542,11 +526,6 @@ int32 FListItemGraphToDiff::GetCurrentDiffIndex() const
 	return -1;
 }
 
-void FListItemGraphToDiff::OnGraphChanged( const FEdGraphEditAction& Action )
-{
-	Diff->OnGraphChanged(this);
-}
-
 FDiffPanel::FDiffPanel()
 {
 	Blueprint = NULL;
@@ -883,14 +862,6 @@ void SBlueprintDiff::OnSelectionChanged( FGraphToDiff Item, ESelectInfo::Type Se
 
 }
 
-void SBlueprintDiff::OnGraphChanged(FListItemGraphToDiff* Diff)
-{
-	if(PanelNew.GraphEditor.IsValid() && PanelNew.GraphEditor.Pin()->GetCurrentGraph() == Diff->GetGraphNew())
-	{
-		FocusOnGraphRevisions(Diff);
-	}
-}
-
 TSharedRef<SWidget> SBlueprintDiff::DefaultEmptyPanel()
 {
 	return SNew(SHorizontalBox)
diff --git a/Engine/Source/Editor/Kismet/Public/SBlueprintDiff.h b/Engine/Source/Editor/Kismet/Public/SBlueprintDiff.h
--- a/Engine/Source/Editor/Kismet/Public/SBlueprintDiff.h
+++ b/Engine/Source/Editor/Kismet/Public/SBlueprintDiff.h
@@ -111,9 +111,6 @@ public:
 
 	void Construct(const FArguments& InArgs);
 
-	/** Called when a new Graph is clicked on by user */
-	void OnGraphChanged(struct FListItemGraphToDiff* Diff);
-
 	/** Helper function for generating an empty widget */
 	static TSharedRef<SWidget> DefaultEmptyPanel();
 
diff --git a/Engine/Source/Runtime/Engine/Classes/EdGraph/EdGraph.h b/Engine/Source/Runtime/Engine/Classes/EdGraph/EdGraph.h
--- a/Engine/Source/Runtime/Engine/Classes/EdGraph/EdGraph.h
+++ b/Engine/Source/Runtime/Engine/Classes/EdGraph/EdGraph.h
@@ -139,7 +139,7 @@ public:
 	 * @param  NodeSelection	The group of nodes you want selected
 	 * @param  bFromUI			True if the node was added as the result of a direct user action.
 	 */
-	void SelectNodeSet(TSet<const UEdGraphNode*> NodeSelection, bool bFromUI = false);
+	void SelectNodeSet(TSet<UEdGraphNode*> NodeSelection, bool bFromUI = false);
 
 	/** Remove a node from this graph */
 	bool RemoveNode( UEdGraphNode* NodeToRemove );
diff --git a/Engine/Source/Runtime/Engine/Private/EdGraph/EdGraph.cpp b/Engine/Source/Runtime/Engine/Private/EdGraph/EdGraph.cpp
--- a/Engine/Source/Runtime/Engine/Private/EdGraph/EdGraph.cpp
+++ b/Engine/Source/Runtime/Engine/Private/EdGraph/EdGraph.cpp
@@ -151,7 +151,7 @@ void UEdGraph::AddNode( UEdGraphNode* NodeToAdd, bool bFromUI/* = false*/, bool
 	NotifyGraphChanged( Action );
 }
 
-void UEdGraph::SelectNodeSet(TSet<const UEdGraphNode*> NodeSelection, bool bFromUI/*= false*/) 
+void UEdGraph::SelectNodeSet(TSet<UEdGraphNode*> NodeSelection, bool bFromUI/*= false*/) 
 {
 	FEdGraphEditAction SelectionAction;
 
@@ -175,8 +175,11 @@ bool UEdGraph::RemoveNode( UEdGraphNode* NodeToRemove )
 #if WITH_EDITOR
 	NodeToRemove->BreakAllNodeLinks();
 #endif	//#if WITH_EDITOR
-
-	NotifyGraphChanged();
+	FEdGraphEditAction RemovalAction;
+	RemovalAction.Graph = this;
+	RemovalAction.Action = GRAPHACTION_RemoveNode;
+	RemovalAction.Nodes.Add(NodeToRemove);
+	NotifyGraphChanged(RemovalAction);
 
 	return NumTimesNodeRemoved > 0;
 }
diff --git a/Engine/Source/Runtime/Engine/Public/BlueprintUtilities.h b/Engine/Source/Runtime/Engine/Public/BlueprintUtilities.h
--- a/Engine/Source/Runtime/Engine/Public/BlueprintUtilities.h
+++ b/Engine/Source/Runtime/Engine/Public/BlueprintUtilities.h
@@ -5,6 +5,7 @@
 #include "GraphEditAction.h"
 
 DECLARE_MULTICAST_DELEGATE_OneParam( FOnGraphChanged, const FEdGraphEditAction& );
+DECLARE_MULTICAST_DELEGATE_OneParam( FSingleNodeNotification, class UEdGraphNode*);
 DECLARE_DELEGATE_OneParam( FSingleNodeEvent, class UEdGraphNode* );
 DECLARE_DELEGATE_OneParam( FEdGraphEvent, class UEdGraph* );
 /** Delegate for notification when property changed */
diff --git a/Engine/Source/Runtime/Engine/Public/GraphEditAction.h b/Engine/Source/Runtime/Engine/Public/GraphEditAction.h
--- a/Engine/Source/Runtime/Engine/Public/GraphEditAction.h
+++ b/Engine/Source/Runtime/Engine/Public/GraphEditAction.h
@@ -16,6 +16,7 @@ enum EEdGraphActionType
 	GRAPHACTION_AddNodeUI = GRAPHACTION_AddNode | GRAPHACTION_UserInitiated,
 
 	GRAPHACTION_SelectNode = 0x1 << 2,
+	GRAPHACTION_RemoveNode = 0x1 << 3,
 };
 
 
@@ -27,7 +28,7 @@ struct FEdGraphEditAction
 	/** The graph the action occurred on */
 	class UEdGraph* Graph;
 	/** [Optional] The object the action occurred on */
-	TSet<const class UEdGraphNode*> Nodes;
+	TSet<class UEdGraphNode*> Nodes;
 
 	FEdGraphEditAction()
 		: Action(GRAPHACTION_Default)
