Blueprint editor no longer purges SGraphNodes when a node is removed, disabled for all schemas except 'clean' k2 schema because of risk

diff --git a/Engine/Source/Editor/AnimGraph/Classes/AnimationConduitGraphSchema.h b/Engine/Source/Editor/AnimGraph/Classes/AnimationConduitGraphSchema.h
--- a/Engine/Source/Editor/AnimGraph/Classes/AnimationConduitGraphSchema.h
+++ b/Engine/Source/Editor/AnimGraph/Classes/AnimationConduitGraphSchema.h
@@ -14,6 +14,7 @@ class UAnimationConduitGraphSchema : public UEdGraphSchema_K2
 	virtual void CreateDefaultNodesForGraph(UEdGraph& Graph) const override;
 	virtual bool CanDuplicateGraph(UEdGraph* InSourceGraph) const override { return false; }
 	virtual void GetGraphDisplayInformation(const UEdGraph& Graph, /*out*/ FGraphDisplayInfo& DisplayInfo) const override;
-	virtual bool DoesSupportEventDispatcher() const	override {	return false; }
+	virtual bool DoesSupportEventDispatcher() const	override { return false; }
+	virtual bool ShouldAlwaysPurgeOnModification() const override { return true; }
 	// End of UEdGraphSchema interface
 };
diff --git a/Engine/Source/Editor/AnimGraph/Classes/AnimationGraphSchema.h b/Engine/Source/Editor/AnimGraph/Classes/AnimationGraphSchema.h
--- a/Engine/Source/Editor/AnimGraph/Classes/AnimationGraphSchema.h
+++ b/Engine/Source/Editor/AnimGraph/Classes/AnimationGraphSchema.h
@@ -51,7 +51,8 @@ class UAnimationGraphSchema : public UEdGraphSchema_K2
 	virtual void GetContextMenuActions(const UEdGraph* CurrentGraph, const UEdGraphNode* InGraphNode, const UEdGraphPin* InGraphPin, FMenuBuilder* MenuBuilder, bool bIsDebugging) const override;
 	virtual FText GetPinDisplayName(const UEdGraphPin* Pin) const override;
 	virtual bool CanDuplicateGraph(UEdGraph* InSourceGraph) const override {	return false; }
-	virtual bool DoesSupportEventDispatcher() const	override {	return false; }
+	virtual bool DoesSupportEventDispatcher() const	override { return false; }
+	virtual bool ShouldAlwaysPurgeOnModification() const override { return true; }
 	// End UEdGraphSchema interface.
 
 	// Begin UEdGraphSchema_K2 interface
diff --git a/Engine/Source/Editor/AnimGraph/Classes/AnimationTransitionSchema.h b/Engine/Source/Editor/AnimGraph/Classes/AnimationTransitionSchema.h
--- a/Engine/Source/Editor/AnimGraph/Classes/AnimationTransitionSchema.h
+++ b/Engine/Source/Editor/AnimGraph/Classes/AnimationTransitionSchema.h
@@ -14,7 +14,8 @@ class UAnimationTransitionSchema : public UEdGraphSchema_K2
 	virtual void CreateDefaultNodesForGraph(UEdGraph& Graph) const override;
 	virtual bool CanDuplicateGraph(UEdGraph* InSourceGraph) const override { return false; }
 	virtual void GetGraphDisplayInformation(const UEdGraph& Graph, /*out*/ FGraphDisplayInfo& DisplayInfo) const;
-	virtual bool DoesSupportEventDispatcher() const	override {	return false; }
+	virtual bool DoesSupportEventDispatcher() const	override { return false; }
+	virtual bool ShouldAlwaysPurgeOnModification() const override { return true; }
 	// End UEdGraphSchema interface.
 
 private:
diff --git a/Engine/Source/Editor/BlueprintGraph/Classes/EdGraphSchema_K2.h b/Engine/Source/Editor/BlueprintGraph/Classes/EdGraphSchema_K2.h
--- a/Engine/Source/Editor/BlueprintGraph/Classes/EdGraphSchema_K2.h
+++ b/Engine/Source/Editor/BlueprintGraph/Classes/EdGraphSchema_K2.h
@@ -412,6 +412,7 @@ public:
 	virtual TSharedPtr<FEdGraphSchemaAction> GetCreateDocumentNodeAction() const override;
 	virtual bool FadeNodeWhenDraggingOffPin(const UEdGraphNode* Node, const UEdGraphPin* Pin) const override;
 	virtual void BackwardCompatibilityNodeConversion(UEdGraph* Graph, bool bOnlySafeChanges) const override;
+	virtual bool ShouldAlwaysPurgeOnModification() const override { return false; }
 	virtual void SplitPin(UEdGraphPin* Pin) const override;
 	virtual void RecombinePin(UEdGraphPin* Pin) const override;
 	// End EdGraphSchema Interface
diff --git a/Engine/Source/Editor/GameplayAbilitiesEditor/Classes/GameplayAbilityGraphSchema.h b/Engine/Source/Editor/GameplayAbilitiesEditor/Classes/GameplayAbilityGraphSchema.h
--- a/Engine/Source/Editor/GameplayAbilitiesEditor/Classes/GameplayAbilityGraphSchema.h
+++ b/Engine/Source/Editor/GameplayAbilitiesEditor/Classes/GameplayAbilityGraphSchema.h
@@ -31,4 +31,6 @@ class UGameplayAbilityGraphSchema : public UEdGraphSchema_K2
 	* @return	A pointer to the newly spawned node
 	*/
 	virtual class UK2Node_VariableSet* SpawnVariableSetNode(const FVector2D GraphPosition, class UEdGraph* ParentGraph, FName VariableName, UStruct* Source) const override;
+
+	virtual bool ShouldAlwaysPurgeOnModification() const override { return true; }
 };
\ No newline at end of file
diff --git a/Engine/Source/Editor/GraphEditor/Private/SGraphEditorImpl.cpp b/Engine/Source/Editor/GraphEditor/Private/SGraphEditorImpl.cpp
--- a/Engine/Source/Editor/GraphEditor/Private/SGraphEditorImpl.cpp
+++ b/Engine/Source/Editor/GraphEditor/Private/SGraphEditorImpl.cpp
@@ -90,11 +90,15 @@ void SGraphEditorImpl::OnGraphChanged(const FEdGraphEditAction& InAction)
 	if ( !bIsActiveTimerRegistered )
 	{
 		// Remove the old user interface nodes
-		GraphPanel->PurgeVisualRepresentation();
-
-		// Trigger the refresh
-		bIsActiveTimerRegistered = true;
-		RegisterActiveTimer( 0.f, FWidgetActiveTimerDelegate::CreateSP( this, &SGraphEditorImpl::TriggerRefresh ) );
+		const UEdGraphSchema* Schema = EdGraphObj->GetSchema();
+		const bool bWasRemoveAction = (InAction.Action ^ GRAPHACTION_RemoveNode) == 0;
+		if (!bWasRemoveAction || (bWasRemoveAction && Schema->ShouldAlwaysPurgeOnModification()))
+		{
+			GraphPanel->PurgeVisualRepresentation();
+			// Trigger the refresh
+			bIsActiveTimerRegistered = true;
+			RegisterActiveTimer(0.f, FWidgetActiveTimerDelegate::CreateSP(this, &SGraphEditorImpl::TriggerRefresh));
+		}
 	}
 }
 
diff --git a/Engine/Source/Editor/GraphEditor/Private/SGraphPanel.cpp b/Engine/Source/Editor/GraphEditor/Private/SGraphPanel.cpp
--- a/Engine/Source/Editor/GraphEditor/Private/SGraphPanel.cpp
+++ b/Engine/Source/Editor/GraphEditor/Private/SGraphPanel.cpp
@@ -974,6 +974,28 @@ void SGraphPanel::AddNode(UEdGraphNode* Node)
 	}
 }
 
+void SGraphPanel::RemoveNode(const UEdGraphNode* Node)
+{
+	for (int32 Iter = 0; Iter != Children.Num(); ++Iter)
+	{
+		TSharedRef<SGraphNode> Child = GetChild(Iter);
+		if (Child->GetNodeObj() == Node)
+		{
+			Children.RemoveAt(Iter);
+			break;
+		}
+	}
+	for (int32 Iter = 0; Iter != Children.Num(); ++Iter)
+	{
+		TSharedRef<SGraphNode> Child = StaticCastSharedRef<SGraphNode>(VisibleChildren[Iter]);
+		if (Child->GetNodeObj() == Node)
+		{
+			VisibleChildren.RemoveAt(Iter);
+			break;
+		}
+	}
+}
+
 TSharedPtr<SGraphNode> SGraphPanel::GetNodeWidgetFromGuid(FGuid Guid) const
 {
 	return NodeGuidMap.FindRef(Guid).Pin();
@@ -1116,21 +1138,31 @@ void SGraphPanel::OnEndPIE( const bool bIsSimulating )
 
 void SGraphPanel::OnGraphChanged(const FEdGraphEditAction& EditAction)
 {
+	const bool bShouldRemoveImmediately = !GraphObj->GetSchema()->ShouldAlwaysPurgeOnModification();
 	if ((EditAction.Graph == GraphObj) &&
-		(EditAction.Nodes.Num() > 0) &&
-		// We do not want to mark it as a UserAddedNode for graphs that do not currently have focus,
-		// this causes each one to want to do the effects and rename, which causes problems.
-		(HasKeyboardFocus() || HasFocusedDescendants()))
+		(EditAction.Nodes.Num() > 0) )
 	{
-		int32 ActionIndex = UserActions.Num();
-		if (EditAction.Action & GRAPHACTION_AddNode)
+		if ((EditAction.Action ^ GRAPHACTION_RemoveNode) == 0 && bShouldRemoveImmediately)
 		{
 			for (const UEdGraphNode* Node : EditAction.Nodes)
 			{
-				UserAddedNodes.Add(Node, ActionIndex);
+				RemoveNode(Node);
+			}
+		}
+		// We do not want to mark it as a UserAddedNode for graphs that do not currently have focus,
+		// this causes each one to want to do the effects and rename, which causes problems.
+		else if (HasKeyboardFocus() || HasFocusedDescendants())
+		{
+			int32 ActionIndex = UserActions.Num();
+			if (EditAction.Action & GRAPHACTION_AddNode)
+			{
+				for (const UEdGraphNode* Node : EditAction.Nodes)
+				{
+					UserAddedNodes.Add(Node, ActionIndex);
+				}
 			}
+			UserActions.Add(EditAction);
 		}
-		UserActions.Add(EditAction);
 	}
 }
 
diff --git a/Engine/Source/Editor/GraphEditor/Private/SGraphPanel.h b/Engine/Source/Editor/GraphEditor/Private/SGraphPanel.h
--- a/Engine/Source/Editor/GraphEditor/Private/SGraphPanel.h
+++ b/Engine/Source/Editor/GraphEditor/Private/SGraphPanel.h
@@ -163,7 +163,10 @@ protected:
 	const TSharedRef<SGraphNode> GetChild(int32 ChildIndex);
 
 	/** Helper method to add a new node to the panel */
-	void AddNode (class UEdGraphNode* Node);
+	void AddNode(class UEdGraphNode* Node);
+
+	/** Helper method to remove a node from the panel */
+	void RemoveNode(const UEdGraphNode* Node);
 public:
 	/** Pin marked via shift-clicking */
 	TWeakPtr<SGraphPin> MarkedPin;
diff --git a/Engine/Source/Editor/UMGEditor/Classes/WidgetGraphSchema.h b/Engine/Source/Editor/UMGEditor/Classes/WidgetGraphSchema.h
--- a/Engine/Source/Editor/UMGEditor/Classes/WidgetGraphSchema.h
+++ b/Engine/Source/Editor/UMGEditor/Classes/WidgetGraphSchema.h
@@ -10,6 +10,8 @@ class UWidgetGraphSchema : public UEdGraphSchema_K2
 {
 	GENERATED_UCLASS_BODY()
 
+	virtual bool ShouldAlwaysPurgeOnModification() const override { return true; }
+
 	UPROPERTY()
 	FName NAME_NeverAsPin;
 
diff --git a/Engine/Source/Runtime/Engine/Classes/EdGraph/EdGraphSchema.h b/Engine/Source/Runtime/Engine/Classes/EdGraph/EdGraphSchema.h
--- a/Engine/Source/Runtime/Engine/Classes/EdGraph/EdGraphSchema.h
+++ b/Engine/Source/Runtime/Engine/Classes/EdGraph/EdGraphSchema.h
@@ -733,4 +733,7 @@ class ENGINE_API UEdGraphSchema : public UObject
 	virtual bool FadeNodeWhenDraggingOffPin(const UEdGraphNode* Node, const UEdGraphPin* Pin) const { return false; }
 
 	virtual void BackwardCompatibilityNodeConversion(UEdGraph* Graph, bool bOnlySafeChanges) const { }
+
+	/* When a node is removed, this method determines whether we should remove it immediately or use the old (slower) code path that results in all node being recreated: */
+	virtual bool ShouldAlwaysPurgeOnModification() const { return true; }
 };
diff --git a/Engine/Source/Runtime/Engine/Private/EdGraph/EdGraph.cpp b/Engine/Source/Runtime/Engine/Private/EdGraph/EdGraph.cpp
--- a/Engine/Source/Runtime/Engine/Private/EdGraph/EdGraph.cpp
+++ b/Engine/Source/Runtime/Engine/Private/EdGraph/EdGraph.cpp
@@ -176,7 +176,11 @@ bool UEdGraph::RemoveNode( UEdGraphNode* NodeToRemove )
 	NodeToRemove->BreakAllNodeLinks();
 #endif	//#if WITH_EDITOR
 
-	NotifyGraphChanged();
+	FEdGraphEditAction RemovalAction;
+	RemovalAction.Graph = this;
+	RemovalAction.Action = GRAPHACTION_RemoveNode;
+	RemovalAction.Nodes.Add(NodeToRemove);
+	NotifyGraphChanged(RemovalAction);
 
 	return NumTimesNodeRemoved > 0;
 }
diff --git a/Engine/Source/Runtime/Engine/Public/GraphEditAction.h b/Engine/Source/Runtime/Engine/Public/GraphEditAction.h
--- a/Engine/Source/Runtime/Engine/Public/GraphEditAction.h
+++ b/Engine/Source/Runtime/Engine/Public/GraphEditAction.h
@@ -12,10 +12,13 @@ enum EEdGraphActionType
 
 	/** A node was added to the graph */
 	GRAPHACTION_AddNode = 0x1 << 1,
-	/** A node was added to the graph at the users' request */
+	/** A node was added to the graph at the user's request */
 	GRAPHACTION_AddNodeUI = GRAPHACTION_AddNode | GRAPHACTION_UserInitiated,
 
 	GRAPHACTION_SelectNode = 0x1 << 2,
+
+	/** A node was removed from the graph at the user's request */
+	GRAPHACTION_RemoveNode = 0x1 << 3,
 };
 
 
