No longer purging the visual representation when nodes are added to the vanilla blueprint editor (should not affect anim blueprint, material editor, etc)

diff --git a/Engine/Source/Editor/GraphEditor/Private/SGraphEditorImpl.cpp b/Engine/Source/Editor/GraphEditor/Private/SGraphEditorImpl.cpp
--- a/Engine/Source/Editor/GraphEditor/Private/SGraphEditorImpl.cpp
+++ b/Engine/Source/Editor/GraphEditor/Private/SGraphEditorImpl.cpp
@@ -89,10 +89,18 @@ void SGraphEditorImpl::OnGraphChanged(const FEdGraphEditAction& InAction)
 {
 	if ( !bIsActiveTimerRegistered )
 	{
-		// Remove the old user interface nodes
 		const UEdGraphSchema* Schema = EdGraphObj->GetSchema();
-		const bool bWasRemoveAction = (InAction.Action ^ GRAPHACTION_RemoveNode) == 0;
-		if (!bWasRemoveAction || (bWasRemoveAction && Schema->ShouldAlwaysPurgeOnModification()))
+		const bool bSchemaRequiresFullRefresh = Schema->ShouldAlwaysPurgeOnModification();
+
+		const bool bWasAddAction = (InAction.Action & GRAPHACTION_AddNode) != 0;
+		const bool bWasSelectAction = (InAction.Action & GRAPHACTION_SelectNode) != 0;
+		const bool bWasRemoveAction = (InAction.Action & GRAPHACTION_RemoveNode) != 0;
+
+		// If we did a 'default action' (or some other action not handled by SGraphPanel::OnGraphChanged
+		// or if we're using a schema that always needs a full refresh, then purge the current nodes
+		// and queue an update:
+		if (bSchemaRequiresFullRefresh || 
+			(!bWasAddAction && !bWasSelectAction && !bWasRemoveAction) )
 		{
 			GraphPanel->PurgeVisualRepresentation();
 			// Trigger the refresh
diff --git a/Engine/Source/Editor/GraphEditor/Private/SGraphPanel.cpp b/Engine/Source/Editor/GraphEditor/Private/SGraphPanel.cpp
--- a/Engine/Source/Editor/GraphEditor/Private/SGraphPanel.cpp
+++ b/Engine/Source/Editor/GraphEditor/Private/SGraphPanel.cpp
@@ -946,12 +946,15 @@ const TSharedRef<SGraphNode> SGraphPanel::GetChild(int32 ChildIndex)
 	return StaticCastSharedRef<SGraphNode>(Children[ChildIndex]);
 }
 
-void SGraphPanel::AddNode(UEdGraphNode* Node)
+void SGraphPanel::AddNode(UEdGraphNode* Node, AddNodeBehavior Behavior)
 {
 	TSharedPtr<SGraphNode> NewNode = FNodeFactory::CreateNodeWidget(Node);
 	check(NewNode.IsValid());
 
-	const bool bWasUserAdded = (UserAddedNodes.Find(Node) != nullptr);
+	const bool bWasUserAdded = 
+		Behavior == WasUserAdded ? true : 
+		Behavior == NotUserAdded ? false :
+					(UserAddedNodes.Find(Node) != nullptr);
 
 	NewNode->SetIsEditable(IsEditable);
 	NewNode->SetDoubleClickEvent(OnNodeDoubleClicked);
@@ -1012,7 +1015,7 @@ void SGraphPanel::Update()
 			UEdGraphNode* Node = GraphObj->Nodes[NodeIndex];
 			if (Node)
 			{
-				AddNode(Node);
+				AddNode(Node, CheckUserAddedNodesList);
 			}
 			else
 			{
@@ -1138,20 +1141,14 @@ void SGraphPanel::OnEndPIE( const bool bIsSimulating )
 
 void SGraphPanel::OnGraphChanged(const FEdGraphEditAction& EditAction)
 {
-	const bool bShouldRemoveImmediately = !GraphObj->GetSchema()->ShouldAlwaysPurgeOnModification();
-	if ((EditAction.Graph == GraphObj) &&
-		(EditAction.Nodes.Num() > 0) )
+	const bool bWillPurge = GraphObj->GetSchema()->ShouldAlwaysPurgeOnModification();
+	if (bWillPurge)
 	{
-		if ((EditAction.Action ^ GRAPHACTION_RemoveNode) == 0 && bShouldRemoveImmediately)
-		{
-			for (const UEdGraphNode* Node : EditAction.Nodes)
-			{
-				RemoveNode(Node);
-			}
-		}
-		// We do not want to mark it as a UserAddedNode for graphs that do not currently have focus,
-		// this causes each one to want to do the effects and rename, which causes problems.
-		else if (HasKeyboardFocus() || HasFocusedDescendants())
+		if ((EditAction.Graph == GraphObj) &&
+			(EditAction.Nodes.Num() > 0) &&
+			// We do not want to mark it as a UserAddedNode for graphs that do not currently have focus,
+			// this causes each one to want to do the effects and rename, which causes problems.
+			(HasKeyboardFocus() || HasFocusedDescendants()))
 		{
 			int32 ActionIndex = UserActions.Num();
 			if (EditAction.Action & GRAPHACTION_AddNode)
@@ -1164,6 +1161,61 @@ void SGraphPanel::OnGraphChanged(const FEdGraphEditAction& EditAction)
 			UserActions.Add(EditAction);
 		}
 	}
+	else if ((EditAction.Graph == GraphObj) && (EditAction.Nodes.Num() > 0) )
+	{
+		// Remove action handled immediately by SGraphPanel::OnGraphChanged
+		const bool bWasAddAction = (EditAction.Action & GRAPHACTION_AddNode) != 0;
+		const bool bWasSelectAction = (EditAction.Action & GRAPHACTION_SelectNode) != 0;
+		const bool bWasRemoveAction = (EditAction.Action & GRAPHACTION_RemoveNode) != 0;
+
+		// The *only* reason we defer these actions is because code higher up the call stack
+		// assumes that the node is created later (for example, GenerateBlueprintAPIUtils::AddNodeToGraph
+		// calls AddNode (which calls this function) before calling AllocateDefaultPins, so if we create 
+		// the widget immediately it won't be able to create its pins). There are lots of other examples, 
+		// and I can't be sure that I've found them all.... 
+		
+		// Minor note, the ugly little lambdas are just to deal with the  time values and return values 
+		// that the timer system requires (and we don't leverage):
+		if (bWasRemoveAction)
+		{
+			const auto RemoveNodeDelegateWrapper = [](double, float, SGraphPanel* Parent, const UEdGraphNode* Node) -> EActiveTimerReturnType
+			{
+				Parent->RemoveNode(Node);
+				return EActiveTimerReturnType::Stop;
+			};
+
+			for (auto Node : EditAction.Nodes)
+			{
+				RegisterActiveTimer(0.f, FWidgetActiveTimerDelegate::CreateStatic(RemoveNodeDelegateWrapper, this, Node));
+			}
+		}
+		if (bWasAddAction)
+		{
+			const auto AddNodeDelegateWrapper = [](double, float, SGraphPanel* Parent, UEdGraphNode* Node, bool bForceUserAdded) -> EActiveTimerReturnType
+			{
+				Parent->AddNode(Node, bForceUserAdded ? WasUserAdded : NotUserAdded );
+				return EActiveTimerReturnType::Stop;
+			};
+
+			for (auto Node : EditAction.Nodes)
+			{
+				RegisterActiveTimer(0.f, FWidgetActiveTimerDelegate::CreateStatic(AddNodeDelegateWrapper, this, const_cast<UEdGraphNode*>(Node), (HasKeyboardFocus() || HasFocusedDescendants()) ) );
+			}
+		}
+		if (bWasSelectAction)
+		{
+			const auto SelectNodeDelegateWrapper = [](double, float, SGraphPanel* Parent, TSet<const UEdGraphNode*> Nodes) -> EActiveTimerReturnType
+			{
+				Parent->DeferredSelectionTargetObjects.Empty();
+				for (auto Node : Nodes)
+				{
+					Parent->DeferredSelectionTargetObjects.Add(Node);
+				}
+				return EActiveTimerReturnType::Stop;
+			};
+			RegisterActiveTimer(0.f, FWidgetActiveTimerDelegate::CreateStatic(SelectNodeDelegateWrapper, this, EditAction.Nodes));
+		}
+	}
 }
 
 void SGraphPanel::NotifyGraphChanged ( const FEdGraphEditAction& EditAction)
diff --git a/Engine/Source/Editor/GraphEditor/Private/SGraphPanel.h b/Engine/Source/Editor/GraphEditor/Private/SGraphPanel.h
--- a/Engine/Source/Editor/GraphEditor/Private/SGraphPanel.h
+++ b/Engine/Source/Editor/GraphEditor/Private/SGraphPanel.h
@@ -162,8 +162,16 @@ protected:
 
 	const TSharedRef<SGraphNode> GetChild(int32 ChildIndex);
 
+	/** Flag to control AddNode, more readable than a bool:*/
+	enum AddNodeBehavior
+	{
+		CheckUserAddedNodesList,
+		WasUserAdded,
+		NotUserAdded
+	};
+
 	/** Helper method to add a new node to the panel */
-	void AddNode(class UEdGraphNode* Node);
+	void AddNode(class UEdGraphNode* Node, AddNodeBehavior Behavior);
 
 	/** Helper method to remove a node from the panel */
 	void RemoveNode(const UEdGraphNode* Node);
diff --git a/Engine/Source/Runtime/Engine/Private/EdGraph/EdGraph.cpp b/Engine/Source/Runtime/Engine/Private/EdGraph/EdGraph.cpp
--- a/Engine/Source/Runtime/Engine/Private/EdGraph/EdGraph.cpp
+++ b/Engine/Source/Runtime/Engine/Private/EdGraph/EdGraph.cpp
@@ -135,13 +135,8 @@ void UEdGraph::AddNode( UEdGraphNode* NodeToAdd, bool bFromUI/* = false*/, bool
 
 	// Create the graph
 	EEdGraphActionType AddNodeAction = GRAPHACTION_AddNode;
-	
-	if(bFromUI)
-	{
-		AddNodeAction = (EEdGraphActionType)( ((int32)AddNodeAction) | GRAPHACTION_UserInitiated );
-	}
 
-	if(bSelectNewNode)
+	if (bSelectNewNode)
 	{
 		AddNodeAction = (EEdGraphActionType)( ((int32)AddNodeAction) | GRAPHACTION_SelectNode );
 	}
@@ -156,10 +151,6 @@ void UEdGraph::SelectNodeSet(TSet<const UEdGraphNode*> NodeSelection, bool bFrom
 	FEdGraphEditAction SelectionAction;
 
 	SelectionAction.Action = GRAPHACTION_SelectNode;
-	if (bFromUI)
-	{
-		SelectionAction.Action = (EEdGraphActionType)(((int32)SelectionAction.Action) | GRAPHACTION_UserInitiated);
-	}
 
 	SelectionAction.Graph = this;
 	SelectionAction.Nodes = NodeSelection;
diff --git a/Engine/Source/Runtime/Engine/Public/GraphEditAction.h b/Engine/Source/Runtime/Engine/Public/GraphEditAction.h
--- a/Engine/Source/Runtime/Engine/Public/GraphEditAction.h
+++ b/Engine/Source/Runtime/Engine/Public/GraphEditAction.h
@@ -7,18 +7,15 @@ enum EEdGraphActionType
 {
 	/** A default edit with no information occurred */
 	GRAPHACTION_Default = 0x0,
-	/** An action involving a user initiated action occurred on the graph */
-	GRAPHACTION_UserInitiated = 0x1 << 0,
 
 	/** A node was added to the graph */
-	GRAPHACTION_AddNode = 0x1 << 1,
-	/** A node was added to the graph at the user's request */
-	GRAPHACTION_AddNodeUI = GRAPHACTION_AddNode | GRAPHACTION_UserInitiated,
+	GRAPHACTION_AddNode = 0x1 << 0,
 
-	GRAPHACTION_SelectNode = 0x1 << 2,
+	/** A node was selected */
+	GRAPHACTION_SelectNode = 0x1 << 1,
 
 	/** A node was removed from the graph at the user's request */
-	GRAPHACTION_RemoveNode = 0x1 << 3,
+	GRAPHACTION_RemoveNode = 0x1 << 2,
 };
 
 
