Fix for UE-8199, UE-8202 - intermittent crash when deleting nodes.. copy paste error, could have been prevented by using EraseIf. To reproduce the bug delete a node when there are VisibleChildren

diff --git a/Engine/Source/Editor/GraphEditor/Private/SGraphPanel.cpp b/Engine/Source/Editor/GraphEditor/Private/SGraphPanel.cpp
--- a/Engine/Source/Editor/GraphEditor/Private/SGraphPanel.cpp
+++ b/Engine/Source/Editor/GraphEditor/Private/SGraphPanel.cpp
@@ -982,7 +982,7 @@ void SGraphPanel::RemoveNode(const UEdGraphNode* Node)
 			break;
 		}
 	}
-	for (int32 Iter = 0; Iter != Children.Num(); ++Iter)
+	for (int32 Iter = 0; Iter != VisibleChildren.Num(); ++Iter)
 	{
 		TSharedRef<SGraphNode> Child = StaticCastSharedRef<SGraphNode>(VisibleChildren[Iter]);
 		if (Child->GetNodeObj() == Node)
