[INTEGRATE] Change 2412123 by [REDACTED_PERSON][REDACTED_HANDLE] on 2015/01/20 09:13:57

	Fix Mac build errors.

diff --git a/Engine/Source/Editor/GraphEditor/Public/SGraphActionMenu.h b/Engine/Source/Editor/GraphEditor/Public/SGraphActionMenu.h
--- a/Engine/Source/Editor/GraphEditor/Public/SGraphActionMenu.h
+++ b/Engine/Source/Editor/GraphEditor/Public/SGraphActionMenu.h
@@ -99,10 +99,10 @@ public:
 	DECLARE_DELEGATE_RetVal_TwoParams( bool, FOnActionMatchesName, FEdGraphSchemaAction*, const FName& );
 
 	SLATE_BEGIN_ARGS(SGraphActionMenu)
-		: _UseSectionStyling(false)
-		, _AutoExpandActionMenu(false)
-		, _ShowFilterTextBox(true)
+		: _AutoExpandActionMenu(false)
 		, _AlphaSortItems(true)
+		, _ShowFilterTextBox(true)
+		, _UseSectionStyling(false)
 		{ }
 
 		SLATE_EVENT( FOnActionSelected, OnActionSelected )
