[INTEGRATE] Change 2412449 by [REDACTED_PERSON][REDACTED_HANDLE] on 2015/01/20 13:05:29

	Build - Adding a header fix the mac build.

diff --git a/Engine/Source/Editor/GraphEditor/Private/SGraphActionMenu.cpp b/Engine/Source/Editor/GraphEditor/Private/SGraphActionMenu.cpp
--- a/Engine/Source/Editor/GraphEditor/Private/SGraphActionMenu.cpp
+++ b/Engine/Source/Editor/GraphEditor/Private/SGraphActionMenu.cpp
@@ -9,6 +9,7 @@
 #include "EditorCategoryUtils.h"
 #include "SSearchBox.h"
 #include "SInlineEditableTextBlock.h"
+#include "Widgets/Views/STableRow.h"
 
 #define LOCTEXT_NAMESPACE "GraphActionMenu"
 
