[INTEGRATE] Change 2412451 by [REDACTED_PERSON][REDACTED_HANDLE] on 2015/01/20 13:06:34

	Build - Possible build fix for clang on mac.

diff --git a/Engine/Source/Editor/GraphEditor/Private/SGraphActionMenu.cpp b/Engine/Source/Editor/GraphEditor/Private/SGraphActionMenu.cpp
--- a/Engine/Source/Editor/GraphEditor/Private/SGraphActionMenu.cpp
+++ b/Engine/Source/Editor/GraphEditor/Private/SGraphActionMenu.cpp
@@ -26,7 +26,7 @@ public:
 
 	void Construct(const FArguments& InArgs, const TSharedRef<STableViewBase>& InOwnerTableView)
 	{
-		ChildSlot
+		STableRow<ItemType>::ChildSlot
 		.Padding(0.0f, 2.0f, 0.0f, 0.0f)
 		[
 			SAssignNew(ContentBorder, SBorder)
