[INTEGRATE] Change 2412492 by [REDACTED_PERSON][REDACTED_HANDLE] on 2015/01/20 13:24:30

	Missing qualifiers for IsHovered and IsItemExpanded

diff --git a/Engine/Source/Editor/GraphEditor/Private/SGraphActionMenu.cpp b/Engine/Source/Editor/GraphEditor/Private/SGraphActionMenu.cpp
--- a/Engine/Source/Editor/GraphEditor/Private/SGraphActionMenu.cpp
+++ b/Engine/Source/Editor/GraphEditor/Private/SGraphActionMenu.cpp
@@ -62,13 +62,13 @@ public:
 
 	const FSlateBrush* GetBackgroundImage() const
 	{
-		if ( IsHovered() )
+		if ( STableRow<ItemType>::IsHovered() )
 		{
-			return IsItemExpanded() ? FEditorStyle::GetBrush("DetailsView.CategoryTop_Hovered") : FEditorStyle::GetBrush("DetailsView.CollapsedCategory_Hovered");
+			return STableRow<ItemType>::IsItemExpanded() ? FEditorStyle::GetBrush("DetailsView.CategoryTop_Hovered") : FEditorStyle::GetBrush("DetailsView.CollapsedCategory_Hovered");
 		}
 		else
 		{
-			return IsItemExpanded() ? FEditorStyle::GetBrush("DetailsView.CategoryTop") : FEditorStyle::GetBrush("DetailsView.CollapsedCategory");
+			return STableRow<ItemType>::IsItemExpanded() ? FEditorStyle::GetBrush("DetailsView.CategoryTop") : FEditorStyle::GetBrush("DetailsView.CollapsedCategory");
 		}
 	}
 
