[INTEGRATE] Change 2412828 by [REDACTED_PERSON][REDACTED_HANDLE] on 2015/01/20 16:32:19

	Another fix for Mac build failures.

diff --git a/Engine/Source/Editor/GraphEditor/Private/SGraphActionMenu.cpp b/Engine/Source/Editor/GraphEditor/Private/SGraphActionMenu.cpp
--- a/Engine/Source/Editor/GraphEditor/Private/SGraphActionMenu.cpp
+++ b/Engine/Source/Editor/GraphEditor/Private/SGraphActionMenu.cpp
@@ -41,7 +41,7 @@ public:
 				.Padding(2.0f, 2.0f, 2.0f, 2.0f)
 				.AutoWidth()
 				[
-					SNew(SExpanderArrow, SharedThis(this))
+					SNew(SExpanderArrow, STableRow< ItemType >::SharedThis(this))
 				]
 
 				+ SHorizontalBox::Slot()
@@ -53,7 +53,7 @@ public:
 		];
 
 		STableRow < ItemType >::ConstructInternal(
-			STableRow::FArguments()
+			typename STableRow< ItemType >::FArguments()
 			.Style(FEditorStyle::Get(), "DetailsView.TreeView.TableRow")
 			.ShowSelection(false),
 			InOwnerTableView
