[INTEGRATE] Change 2418803 by [REDACTED_PERSON][REDACTED_HANDLE] on 2015/01/26 12:50:53

	BP - Fixing the sizing on the section widget so that it is auto instead of fill sized.

diff --git a/Engine/Source/Editor/GraphEditor/Private/SGraphActionMenu.cpp b/Engine/Source/Editor/GraphEditor/Private/SGraphActionMenu.cpp
--- a/Engine/Source/Editor/GraphEditor/Private/SGraphActionMenu.cpp
+++ b/Engine/Source/Editor/GraphEditor/Private/SGraphActionMenu.cpp
@@ -1056,6 +1056,7 @@ TSharedRef<ITableRow> SGraphActionMenu::MakeWidget( TSharedPtr<FGraphActionNode>
 			]
 
 			+ SHorizontalBox::Slot()
+			.AutoWidth()
 			.VAlign(VAlign_Center)
 			.HAlign(HAlign_Right)
 			.Padding(FMargin(0,0,2,0))
