Fixed a (presumably) copy-paste derived bug in SGraphActionMenu::SelectItemByName resulting in crashes on adding keys to Blackboard assets #UE4

- addresses UE-8529
- not a broken-build fix, but a critical 4.7 fix (and super safe!)

diff --git a/Engine/Source/Editor/GraphEditor/Private/SGraphActionMenu.cpp b/Engine/Source/Editor/GraphEditor/Private/SGraphActionMenu.cpp
--- a/Engine/Source/Editor/GraphEditor/Private/SGraphActionMenu.cpp
+++ b/Engine/Source/Editor/GraphEditor/Private/SGraphActionMenu.cpp
@@ -498,7 +498,7 @@ bool SGraphActionMenu::SelectItemByName(const FName& ItemName, ESelectInfo::Type
 						{
 							if(ChildGraphAction)
 							{
-								if ((OnActionMatchesName.IsBound() && OnActionMatchesName.Execute(GraphAction, ItemName)) || GraphActionMenuHelpers::ActionMatchesName(ChildGraphAction, ItemName))
+								if ((OnActionMatchesName.IsBound() && OnActionMatchesName.Execute(ChildGraphAction, ItemName)) || GraphActionMenuHelpers::ActionMatchesName(ChildGraphAction, ItemName))
 								{
 									SelectionNode = GraphNodes[i]->Children[ChildIdx];
 
