This changes the behaviour of context menu comment editing in the graphs so that they automatically display comment bubbles when text is entered and hide when text is completely removed ( empty comment )
#UE-8480 - Unable to create comments on nodes
#Branch UE4
#Proj GraphEditor

diff --git a/Engine/Source/Editor/GraphEditor/Private/SCommentBubble.cpp b/Engine/Source/Editor/GraphEditor/Private/SCommentBubble.cpp
--- a/Engine/Source/Editor/GraphEditor/Private/SCommentBubble.cpp
+++ b/Engine/Source/Editor/GraphEditor/Private/SCommentBubble.cpp
@@ -67,18 +67,21 @@ void SCommentBubble::Tick( const FGeometry& AllottedGeometry, const double InCur
 {
 	SCompoundWidget::Tick( AllottedGeometry, InCurrentTime, InDeltaTime );
 
-	bool bIsHovered = IsHovered();
+	const bool bTitleBarBubbleVisible = bEnableTitleBarBubble && IsGraphNodeHovered.IsBound();
 
-	const FLinearColor BubbleColor = ColorAndOpacity.Get().GetSpecifiedColor() * SCommentBubbleDefs::LuminanceCoEff;
-	const float BubbleLuminance = BubbleColor.R + BubbleColor.G + BubbleColor.B;
-	ForegroundColor = BubbleLuminance < 0.5f ? SCommentBubbleDefs::DarkForegroundClr : SCommentBubbleDefs::LightForegroundClr;
-
-	if( bEnableTitleBarBubble && IsGraphNodeHovered.IsBound() )
+	if( bTitleBarBubbleVisible || IsBubbleVisible() )
 	{
-		bIsHovered |= IsGraphNodeHovered.Execute();
+		const FLinearColor BubbleColor = ColorAndOpacity.Get().GetSpecifiedColor() * SCommentBubbleDefs::LuminanceCoEff;
+		const float BubbleLuminance = BubbleColor.R + BubbleColor.G + BubbleColor.B;
+		ForegroundColor = BubbleLuminance < 0.5f ? SCommentBubbleDefs::DarkForegroundClr : SCommentBubbleDefs::LightForegroundClr;
+	}
 
-		if( !GraphNode->bCommentBubbleVisible )
+	if( !GraphNode->bCommentBubbleVisible )
+	{
+		if( bTitleBarBubbleVisible )
 		{
+			const bool bIsHovered = IsHovered() || IsGraphNodeHovered.Execute();
+
 			if( bIsHovered )
 			{
 				if( OpacityValue < 1.f )
@@ -97,6 +100,13 @@ void SCommentBubble::Tick( const FGeometry& AllottedGeometry, const double InCur
 			}
 		}
 	}
+	if( CachedComment != GraphNode->NodeComment )
+	{
+		CachedComment = GraphNode->NodeComment;
+		CachedCommentText = FText::FromString( CachedComment );
+		GraphNode->bCommentBubbleVisible = !CachedComment.IsEmpty();
+		UpdateBubble();
+	}
 }
 
 void SCommentBubble::UpdateBubble()
