[INTEGRATE] Change 2422805 by [REDACTED_PERSON][REDACTED_HANDLE] on 2015/01/28 16:18:44

	Graph editor: Allow long graph drag-drop feedback tooltips to wrap instead of being clipped

diff --git a/Engine/Source/Editor/GraphEditor/Private/GraphEditorDragDropAction.cpp b/Engine/Source/Editor/GraphEditor/Private/GraphEditorDragDropAction.cpp
--- a/Engine/Source/Editor/GraphEditor/Private/GraphEditorDragDropAction.cpp
+++ b/Engine/Source/Editor/GraphEditor/Private/GraphEditorDragDropAction.cpp
@@ -1,7 +1,7 @@
 // Copyright 1998-2015 Epic Games, Inc. All Rights Reserved.
 
-
 #include "GraphEditorCommon.h"
+#include "SScaleBox.h"
 
 UEdGraphPin* FGraphEditorDragDropAction::GetHoveredPin() const
 {
@@ -130,26 +130,35 @@ void FGraphEditorDragDropAction::SetSimpleFeedbackMessage(const FSlateBrush* Ico
 		.AutoWidth()
 		.Padding(3.0f)
 		[
-			SNew(SImage)
-			.Visibility(ErrorIconVisibility)
-			.Image( FEditorStyle::GetBrush( TEXT("Graph.ConnectorFeedback.Error") ))
-			.ColorAndOpacity( FLinearColor::White )
+			SNew(SScaleBox)
+			.Stretch(EStretch::ScaleToFit)
+			[
+				SNew(SImage)
+				.Visibility(ErrorIconVisibility)
+				.Image( FEditorStyle::GetBrush( TEXT("Graph.ConnectorFeedback.Error") ))
+				.ColorAndOpacity( FLinearColor::White )
+			]
 		]
 		+SHorizontalBox::Slot()
 		.AutoWidth()
 		.Padding(3.0f)
 		[
-			SNew(SImage) 
-			.Visibility(IconVisibility)
-			.Image( Icon )
-			.ColorAndOpacity( IconColor )
+			SNew(SScaleBox)
+			.Stretch(EStretch::ScaleToFit)
+			[
+				SNew(SImage)
+				.Visibility(IconVisibility)
+				.Image( Icon )
+				.ColorAndOpacity( IconColor )
+			]
 		]
 		+SHorizontalBox::Slot()
 		.AutoWidth()
 		.MaxWidth(500)
 		.VAlign(VAlign_Center)
 		[
-			SNew(STextBlock) 
+			SNew(STextBlock)
+			.WrapTextAt( 480 )
 			.Text( Message )
 		]
 	);
