Fix for UE-8789, we had two SGraphNodes for a single UEdGraphNode (refresh and add node called on the same frame)

diff --git a/Engine/Source/Editor/GraphEditor/Private/SGraphPanel.cpp b/Engine/Source/Editor/GraphEditor/Private/SGraphPanel.cpp
--- a/Engine/Source/Editor/GraphEditor/Private/SGraphPanel.cpp
+++ b/Engine/Source/Editor/GraphEditor/Private/SGraphPanel.cpp
@@ -1182,6 +1182,7 @@ void SGraphPanel::OnGraphChanged(const FEdGraphEditAction& EditAction)
 		{
 			const auto AddNodeDelegateWrapper = [](double, float, SGraphPanel* Parent, UEdGraphNode* Node, bool bForceUserAdded) -> EActiveTimerReturnType
 			{
+				Parent->RemoveNode(Node);
 				Parent->AddNode(Node, bForceUserAdded ? WasUserAdded : NotUserAdded );
 				return EActiveTimerReturnType::Stop;
 			};
