[INTEGRATE] Change 2430954 by [REDACTED_PERSON][REDACTED_HANDLE] on 2015/02/03 20:36:44

	Standard Macros lists 'ForEach' in it's description, scoring algorithm bug fixed (no longer multiplying succeeding scores, which allowed low scoring entries with lots of text to override good entries with succinct descriptoins), added keywords to UActorComponent::GetOwner. Tested context menu with the following keywords: if, actor, position, transform, rotation, for, foreach, parent, owner. Results for transform and rotation unchanged, results for position changed but not notably improved, all other results improved.

diff --git a/Engine/Source/Editor/BlueprintGraph/Classes/K2Node_MacroInstance.h b/Engine/Source/Editor/BlueprintGraph/Classes/K2Node_MacroInstance.h
--- a/Engine/Source/Editor/BlueprintGraph/Classes/K2Node_MacroInstance.h
+++ b/Engine/Source/Editor/BlueprintGraph/Classes/K2Node_MacroInstance.h
@@ -36,6 +36,7 @@ public:
 	// Begin UEdGraphNode interface
 	virtual void AllocateDefaultPins() override;
 	virtual FText GetTooltipText() const override;
+	virtual FString GetKeywords() const override;
 	virtual void PostPasteNode() override;
 	virtual FLinearColor GetNodeTitleColor() const override;
 	virtual FText GetNodeTitle(ENodeTitleType::Type TitleType) const override;
diff --git a/Engine/Source/Editor/BlueprintGraph/Private/K2Node_MacroInstance.cpp b/Engine/Source/Editor/BlueprintGraph/Private/K2Node_MacroInstance.cpp
--- a/Engine/Source/Editor/BlueprintGraph/Private/K2Node_MacroInstance.cpp
+++ b/Engine/Source/Editor/BlueprintGraph/Private/K2Node_MacroInstance.cpp
@@ -143,6 +143,17 @@ FText UK2Node_MacroInstance::GetTooltipText() const
 	return CachedTooltip;
 }
 
+FString UK2Node_MacroInstance::GetKeywords() const
+{
+	// This might need to be refined, but seems better than providing no keywords:
+	UBlueprint* MacroBP = MacroGraphReference.GetBlueprint();
+	if (MacroBP)
+	{
+		return MacroBP->BlueprintDescription;
+	}
+	return UK2Node_Tunnel::GetKeywords();
+}
+
 FText UK2Node_MacroInstance::GetNodeTitle(ENodeTitleType::Type TitleType) const
 {
 	UEdGraph* MacroGraph = MacroGraphReference.GetGraph();
diff --git a/Engine/Source/Editor/GraphEditor/Private/SGraphActionMenu.cpp b/Engine/Source/Editor/GraphEditor/Private/SGraphActionMenu.cpp
--- a/Engine/Source/Editor/GraphEditor/Private/SGraphActionMenu.cpp
+++ b/Engine/Source/Editor/GraphEditor/Private/SGraphActionMenu.cpp
@@ -735,10 +735,10 @@ int32 SGraphActionMenu::GetActionFilteredWeight( const FGraphActionListBuilderBa
 	int32 TotalWeight = 0;
 
 	// Some simple weight figures to help find the most appropriate match	
-	int32 WholeMatchWeightMultiplier = 2;
-	int32 DescriptionWeight = 5;
-	int32 CategoryWeight = 3;
-	int32 NodeTitleWeight = 3;
+	int32 WholeMatchWeightMultiplier = 1;
+	int32 DescriptionWeight = 1;
+	int32 CategoryWeight = 1;
+	int32 NodeTitleWeight = 1;
 
 	// Helper array
 	struct FArrayWithWeight
@@ -762,7 +762,7 @@ int32 SGraphActionMenu::GetActionFilteredWeight( const FGraphActionListBuilderBa
 
 		// First the keywords
 		InCurrentAction.Actions[Action]->Keywords.ParseIntoArray( &EachEntry.Array, TEXT(" "), true );
-		EachEntry.Weight = 1;
+		EachEntry.Weight = 10;
 		WeightedArrayList.Add( EachEntry );
 
 		// The description
@@ -796,9 +796,9 @@ int32 SGraphActionMenu::GetActionFilteredWeight( const FGraphActionListBuilderBa
 				TotalWeight++;
 			}		
 			// Now check the weighted lists	(We could further improve the hit weight by checking consecutive word matches)
-			int32 WeightPerList = 0;
 			for (int32 iFindCount = 0; iFindCount < WeightedArrayList.Num() ; iFindCount++)
 			{
+				int32 WeightPerList = 0;
 				TArray<FString>& KeywordArray = WeightedArrayList[iFindCount].Array;
 				int32 EachWeight = WeightedArrayList[iFindCount].Weight;
 				int32 WholeMatchCount = 0;
diff --git a/Engine/Source/Runtime/Engine/Classes/Components/ActorComponent.h b/Engine/Source/Runtime/Engine/Classes/Components/ActorComponent.h
--- a/Engine/Source/Runtime/Engine/Classes/Components/ActorComponent.h
+++ b/Engine/Source/Runtime/Engine/Classes/Components/ActorComponent.h
@@ -131,7 +131,7 @@ public:
 	void OnRep_IsActive();
 
 	/** Follow the Outer chain to get the  AActor  that 'Owns' this component */
-	UFUNCTION(BlueprintCallable, Category="Components")
+	UFUNCTION(BlueprintCallable, Category="Components", meta=(Keywords = "Actor Owning Parent"))
 	class AActor* GetOwner() const;
 
 	virtual class UWorld* GetWorld() const override;
