[INTEGRATE] Change 2434763 by [REDACTED_PERSON][REDACTED_HANDLE] on 2015/02/05 22:57:02

	Graph editor: Fixed regression where node comment bubbles (watch values, anim states, timeline/latent delay display, etc...) would be hugely inflated

diff --git a/Engine/Source/Editor/EditorStyle/Private/SlateEditorStyle.cpp b/Engine/Source/Editor/EditorStyle/Private/SlateEditorStyle.cpp
--- a/Engine/Source/Editor/EditorStyle/Private/SlateEditorStyle.cpp
+++ b/Engine/Source/Editor/EditorStyle/Private/SlateEditorStyle.cpp
@@ -3236,7 +3236,8 @@ void FSlateEditorStyle::FStyle::SetupGraphEditorStyles()
 		Set( "Graph.Node.CommentFont", TTF_CORE_FONT( "Fonts/Roboto-Regular", 10 ) );
 		Set( "Graph.Node.Comment.BubbleOffset", FMargin(8,0,0,0) );
 		Set( "Graph.Node.Comment.PinIconPadding", FMargin(0,2,0,0) );
-		Set( "Graph.Node.Comment.BubblePadding", FMargin(4,4) );
+		Set( "Graph.Node.Comment.BubblePadding", FVector2D(3,3) );
+		Set( "Graph.Node.Comment.BubbleWidgetMargin", FMargin(4,4) );
 
 		const FCheckBoxStyle CommentTitleButton = FCheckBoxStyle()
 			.SetCheckBoxType( ESlateCheckBoxType::CheckBox )
diff --git a/Engine/Source/Editor/GraphEditor/Private/SCommentBubble.cpp b/Engine/Source/Editor/GraphEditor/Private/SCommentBubble.cpp
--- a/Engine/Source/Editor/GraphEditor/Private/SCommentBubble.cpp
+++ b/Engine/Source/Editor/GraphEditor/Private/SCommentBubble.cpp
@@ -114,9 +114,9 @@ void SCommentBubble::UpdateBubble()
 	if( GraphNode->bCommentBubbleVisible )
 	{
 		const FSlateBrush* CommentCalloutArrowBrush = FEditorStyle::GetBrush(TEXT("Graph.Node.CommentArrow"));
-		const FMargin BubblePadding = FEditorStyle::GetMargin( TEXT("Graph.Node.Comment.BubblePadding"));
-		const FMargin PinIconPadding = FEditorStyle::GetMargin( TEXT("Graph.Node.Comment.PinIconPadding"));
-		const FMargin BubbleOffset = FEditorStyle::GetMargin( TEXT("Graph.Node.Comment.BubbleOffset"));
+		const FMargin BubblePadding = FEditorStyle::GetMargin(TEXT("Graph.Node.Comment.BubbleWidgetMargin"));
+		const FMargin PinIconPadding = FEditorStyle::GetMargin(TEXT("Graph.Node.Comment.PinIconPadding"));
+		const FMargin BubbleOffset = FEditorStyle::GetMargin(TEXT("Graph.Node.Comment.BubbleOffset"));
 		// Conditionally create bubble controls
 		TSharedPtr<SWidget> BubbleControls = SNullWidget::NullWidget;
 
