[INTEGRATE] Change 2439412 by [REDACTED_PERSON][REDACTED_HANDLE] on 2015/02/10 01:33:06

	Blueprints: Add automatic expansion of parent nodes when programmatically selecting an item in the My Blueprint tree [UE-9169]

diff --git a/Engine/Source/Editor/GraphEditor/Private/SGraphActionMenu.cpp b/Engine/Source/Editor/GraphEditor/Private/SGraphActionMenu.cpp
--- a/Engine/Source/Editor/GraphEditor/Private/SGraphActionMenu.cpp
+++ b/Engine/Source/Editor/GraphEditor/Private/SGraphActionMenu.cpp
@@ -519,6 +519,13 @@ bool SGraphActionMenu::SelectItemByName(const FName& ItemName, ESelectInfo::Type
 
 		if(SelectionNode.IsValid())
 		{
+			// Expand the parent nodes
+			for (TSharedPtr<FGraphActionNode> ParentAction = SelectionNode->GetParentNode().Pin(); ParentAction.IsValid(); ParentAction = ParentAction->GetParentNode().Pin())
+			{
+				TreeView->SetItemExpansion(ParentAction, true);
+			}
+
+			// Select the node
 			TreeView->SetSelection(SelectionNode,SelectInfo);
 			TreeView->RequestScrollIntoView(SelectionNode);
 			return true;
