UE-8834 - Blueprint - The text and string pins now are now multiline capable.  Using Shift+Enter will create a new line.

diff --git a/Engine/Source/Editor/GraphEditor/Private/KismetPins/SGraphPinString.cpp b/Engine/Source/Editor/GraphEditor/Private/KismetPins/SGraphPinString.cpp
--- a/Engine/Source/Editor/GraphEditor/Private/KismetPins/SGraphPinString.cpp
+++ b/Engine/Source/Editor/GraphEditor/Private/KismetPins/SGraphPinString.cpp
@@ -12,15 +12,19 @@ void SGraphPinString::Construct(const FArguments& InArgs, UEdGraphPin* InGraphPi
 
 TSharedRef<SWidget>	SGraphPinString::GetDefaultValueWidget()
 {
-	return SNew(SEditableTextBox)
-		.Style( FEditorStyle::Get(), "Graph.EditableTextBox" )
-		.Text( this, &SGraphPinString::GetTypeInValue )
-		.SelectAllTextWhenFocused(true)
-		.MinDesiredWidth( 18.0f )
-		.Visibility( this, &SGraphPin::GetDefaultValueVisibility )
-		.IsReadOnly(this, &SGraphPinString::GetDefaultValueIsReadOnly )
-		.OnTextCommitted( this, &SGraphPinString::SetTypeInValue )
-		.ForegroundColor( FSlateColor::UseForeground() );
+	return SNew(SBox)
+		.MinDesiredWidth(18.0f)
+		[
+			SNew(SMultiLineEditableTextBox)
+			.Style(FEditorStyle::Get(), "Graph.EditableTextBox")
+			.Text(this, &SGraphPinString::GetTypeInValue)
+			.SelectAllTextWhenFocused(true)
+			.Visibility(this, &SGraphPin::GetDefaultValueVisibility)
+			.IsReadOnly(this, &SGraphPinString::GetDefaultValueIsReadOnly)
+			.OnTextCommitted(this, &SGraphPinString::SetTypeInValue)
+			.ForegroundColor(FSlateColor::UseForeground())
+			.ModiferKeyForNewLine(EModifierKey::Shift)
+		];
 }
 
 FText SGraphPinString::GetTypeInValue() const
diff --git a/Engine/Source/Editor/GraphEditor/Private/KismetPins/SGraphPinText.cpp b/Engine/Source/Editor/GraphEditor/Private/KismetPins/SGraphPinText.cpp
--- a/Engine/Source/Editor/GraphEditor/Private/KismetPins/SGraphPinText.cpp
+++ b/Engine/Source/Editor/GraphEditor/Private/KismetPins/SGraphPinText.cpp
@@ -12,15 +12,19 @@ void SGraphPinText::Construct(const FArguments& InArgs, UEdGraphPin* InGraphPinO
 
 TSharedRef<SWidget>	SGraphPinText::GetDefaultValueWidget()
 {
-	return SNew(SEditableTextBox)
-				.Style( FEditorStyle::Get(), "Graph.EditableTextBox" )
-				.Text( this, &SGraphPinText::GetTypeInValue )
-				.SelectAllTextWhenFocused(true)
-				.MinDesiredWidth( 18.0f )
-				.Visibility( this, &SGraphPin::GetDefaultValueVisibility )
-				.IsReadOnly(this, &SGraphPinText::GetDefaultValueIsReadOnly )
-				.OnTextCommitted( this, &SGraphPinText::SetTypeInValue )
-				.ForegroundColor( FSlateColor::UseForeground() );
+	return SNew(SBox)
+		.MinDesiredWidth(18.0f)
+		[
+			SNew(SMultiLineEditableTextBox)
+			.Style(FEditorStyle::Get(), "Graph.EditableTextBox")
+			.Text(this, &SGraphPinText::GetTypeInValue)
+			.SelectAllTextWhenFocused(true)
+			.Visibility(this, &SGraphPin::GetDefaultValueVisibility)
+			.IsReadOnly(this, &SGraphPinText::GetDefaultValueIsReadOnly)
+			.OnTextCommitted(this, &SGraphPinText::SetTypeInValue)
+			.ForegroundColor(FSlateColor::UseForeground())
+			.ModiferKeyForNewLine(EModifierKey::Shift)
+		];
 }
 
 FText SGraphPinText::GetTypeInValue() const
