[INTEGRATE] Change 2441355 by [REDACTED_PERSON][REDACTED_HANDLE] on 2015/02/11 05:57:54

	This change re-implements info bubble vertical stacking functionality on GraphNodes.
	#UE-9280 - Node comment bubbles draw on top of diagnostic info bubbles (pin watches, anim states / timeline / latent action time displays / etc...)
	#Branch UE4-Releases
	#Proj GraphEditor

	#info Also added a fix for compact nodes not displaying any comments

diff --git a/Engine/Source/Editor/GraphEditor/Private/KismetNodes/SGraphNodeK2Base.cpp b/Engine/Source/Editor/GraphEditor/Private/KismetNodes/SGraphNodeK2Base.cpp
--- a/Engine/Source/Editor/GraphEditor/Private/KismetNodes/SGraphNodeK2Base.cpp
+++ b/Engine/Source/Editor/GraphEditor/Private/KismetNodes/SGraphNodeK2Base.cpp
@@ -138,6 +138,28 @@ void SGraphNodeK2Base::UpdateCompactNode()
 		}
 	}
 
+	// Create comment bubble
+	TSharedPtr<SCommentBubble> CommentBubble;
+
+	SAssignNew( CommentBubble, SCommentBubble )
+	.GraphNode( GraphNode )
+	.Text( this, &SGraphNode::GetNodeComment )
+	.ColorAndOpacity( this, &SGraphNodeK2Base::GetCommentColor )
+	.AllowPinning( true )
+	.EnableTitleBarBubble( true )
+	.EnableBubbleCtrls( true )
+	.GraphLOD( this, &SGraphNode::GetCurrentLOD )
+	.IsGraphNodeHovered( this, &SGraphNodeK2Base::IsHovered );
+
+	GetOrAddSlot( ENodeZone::TopCenter )
+	.SlotOffset( TAttribute<FVector2D>( CommentBubble.Get(), &SCommentBubble::GetOffset ))
+	.SlotSize( TAttribute<FVector2D>( CommentBubble.Get(), &SCommentBubble::GetSize ))
+	.AllowScaling( TAttribute<bool>( CommentBubble.Get(), &SCommentBubble::IsScalingAllowed ))
+	.VAlign( VAlign_Top )
+	[
+		CommentBubble.ToSharedRef()
+	];
+
 	CreateInputSideAddButton(LeftNodeBox);
 	CreateOutputSideAddButton(RightNodeBox);
 }
diff --git a/Engine/Source/Editor/GraphEditor/Private/SCommentBubble.cpp b/Engine/Source/Editor/GraphEditor/Private/SCommentBubble.cpp
--- a/Engine/Source/Editor/GraphEditor/Private/SCommentBubble.cpp
+++ b/Engine/Source/Editor/GraphEditor/Private/SCommentBubble.cpp
@@ -272,7 +272,11 @@ void SCommentBubble::UpdateBubble()
 
 FVector2D SCommentBubble::GetOffset() const
 {
-	return FVector2D( 0.f, -GetDesiredSize().Y );
+	if( GraphNode->bCommentBubbleVisible || OpacityValue > 0.f )
+	{
+		return FVector2D( 0.f, -GetDesiredSize().Y );
+	}
+	return FVector2D::ZeroVector;
 }
 
 FVector2D SCommentBubble::GetSize() const
diff --git a/Engine/Source/Editor/GraphEditor/Private/SGraphPanel.cpp b/Engine/Source/Editor/GraphEditor/Private/SGraphPanel.cpp
--- a/Engine/Source/Editor/GraphEditor/Private/SGraphPanel.cpp
+++ b/Engine/Source/Editor/GraphEditor/Private/SGraphPanel.cpp
@@ -214,7 +214,8 @@ int32 SGraphPanel::OnPaint( const FPaintArgs& Args, const FGeometry& AllottedGeo
 
 				// Draw the comments and information popups for this node, if it has any.
 				{
-					float CommentBubbleY = 0.0f;
+					const SNodePanel::SNode::FNodeSlot& CommentSlot = ChildNode->GetOrAddSlot( ENodeZone::TopCenter );
+					float CommentBubbleY = -CommentSlot.Offset.Get().Y;
 					Context.bSelected = bSelected;
 					TArray<FGraphInformationPopupInfo> Popups;
 
