Hidden comment bubbles should no longer reappear when doing unrelated actions.

diff --git a/Engine/Source/Editor/GraphEditor/Private/SCommentBubble.cpp b/Engine/Source/Editor/GraphEditor/Private/SCommentBubble.cpp
--- a/Engine/Source/Editor/GraphEditor/Private/SCommentBubble.cpp
+++ b/Engine/Source/Editor/GraphEditor/Private/SCommentBubble.cpp
@@ -47,6 +47,10 @@ void SCommentBubble::Construct( const FArguments& InArgs )
 	HintText				= InArgs._HintText.IsSet() ? InArgs._HintText : NSLOCTEXT( "CommentBubble", "EditCommentHint", "Click to edit" );
 	OpacityValue			= SCommentBubbleDefs::FadeDelay;
 
+	// Cache the comment
+	CachedComment = CommentAttribute.Get();
+	CachedCommentText = FText::FromString( CachedComment );
+
 	// Create Widget
 	UpdateBubble();
 }
@@ -308,11 +312,6 @@ FText SCommentBubble::GetCommentText() const
 	{
 		return HintText.Get();
 	}
-	else if( CachedComment != CommentAttribute.Get() )
-	{
-		CachedComment = CommentAttribute.Get();
-		CachedCommentText = FText::FromString( CachedComment );
-	}
 	return CachedCommentText;
 }
 
diff --git a/Engine/Source/Editor/GraphEditor/Private/SGraphNodeComment.cpp b/Engine/Source/Editor/GraphEditor/Private/SGraphNodeComment.cpp
--- a/Engine/Source/Editor/GraphEditor/Private/SGraphNodeComment.cpp
+++ b/Engine/Source/Editor/GraphEditor/Private/SGraphNodeComment.cpp
@@ -58,15 +58,16 @@ void SGraphNodeComment::Tick( const FGeometry& AllottedGeometry, const double In
 {
 	SGraphNode::Tick(AllottedGeometry, InCurrentTime, InDeltaTime);
 
-	const int32 CurrentWidth = static_cast<int32>(UserSize.X);
 	const FString CurrentCommentTitle = GetNodeComment();
-
-	if (( CurrentCommentTitle != CachedCommentTitle ) || ( CurrentWidth != CachedWidth ))
+	if (CurrentCommentTitle != CachedCommentTitle)
 	{
 		CachedCommentTitle = CurrentCommentTitle;
-		CachedWidth = CurrentWidth;
+	}
 
-		UpdateGraphNode();
+	const int32 CurrentWidth = static_cast<int32>(UserSize.X);
+	if (CurrentWidth != CachedWidth)
+	{
+		CachedWidth = CurrentWidth;
 	}
 }
 
