[INTEGRATE] Change 2445658 by [REDACTED_PERSON][REDACTED_HANDLE] on 2015/02/13 18:28:24

	Fix for missing comment text in reference viewer.
	#UE-9833 - Reference viewer doesn't display text in the bubbles
	#Branch UE4-Releases
	#Proj GraphEditor

diff --git a/Engine/Source/Editor/GraphEditor/Private/SCommentBubble.cpp b/Engine/Source/Editor/GraphEditor/Private/SCommentBubble.cpp
--- a/Engine/Source/Editor/GraphEditor/Private/SCommentBubble.cpp
+++ b/Engine/Source/Editor/GraphEditor/Private/SCommentBubble.cpp
@@ -46,6 +46,8 @@ void SCommentBubble::Construct( const FArguments& InArgs )
 	IsGraphNodeHovered		= InArgs._IsGraphNodeHovered;
 	HintText				= InArgs._HintText.IsSet() ? InArgs._HintText : NSLOCTEXT( "CommentBubble", "EditCommentHint", "Click to edit" );
 	OpacityValue			= SCommentBubbleDefs::FadeDelay;
+	// Ensue this value is set to something sensible
+	ForegroundColor = SCommentBubbleDefs::LightForegroundClr;
 
 	// Cache the comment
 	CachedComment = CommentAttribute.Get();
