[INTEGRATE] Change 2448972 by [REDACTED_PERSON][REDACTED_HANDLE] on 2015/02/17 13:49:55

	UE-9950 Creating Custom Event doesn't allow user to name it right away anymore

diff --git a/Engine/Source/Editor/GraphEditor/Private/SGraphPanel.cpp b/Engine/Source/Editor/GraphEditor/Private/SGraphPanel.cpp
--- a/Engine/Source/Editor/GraphEditor/Private/SGraphPanel.cpp
+++ b/Engine/Source/Editor/GraphEditor/Private/SGraphPanel.cpp
@@ -1072,6 +1072,7 @@ bool SGraphPanel::IsNodeTitleVisible(const class UEdGraphNode* Node, bool bReque
 			if( bTitleVisible && bRequestRename )
 			{
 				GraphNode.Pin()->RequestRename();
+				SelectAndCenterObject(Node, false);
 			}
 		}
 	}
