Changing pin values will now register a transaction for Undo/Redo support.

Fixed an issue with FText's not going to or from an empty FText correctly during Undo/Redo.

diff --git a/Engine/Source/Editor/GraphEditor/Private/KismetPins/SGraphPinBool.cpp b/Engine/Source/Editor/GraphEditor/Private/KismetPins/SGraphPinBool.cpp
--- a/Engine/Source/Editor/GraphEditor/Private/KismetPins/SGraphPinBool.cpp
+++ b/Engine/Source/Editor/GraphEditor/Private/KismetPins/SGraphPinBool.cpp
@@ -27,5 +27,11 @@ ECheckBoxState SGraphPinBool::IsDefaultValueChecked() const
 void SGraphPinBool::OnDefaultValueCheckBoxChanged(ECheckBoxState InIsChecked)
 {
 	const FString BoolString = (InIsChecked == ECheckBoxState::Checked) ? TEXT("true") : TEXT("false");
-	GraphPinObj->GetSchema()->TrySetDefaultValue(*GraphPinObj, BoolString);
+	if(GraphPinObj->GetDefaultAsString() != BoolString)
+	{
+		const FScopedTransaction Transaction( NSLOCTEXT("GraphEditor", "ChangeBoolPinValue", "Change Bool Pin Value" ) );
+		GraphPinObj->Modify();
+
+		GraphPinObj->GetSchema()->TrySetDefaultValue(*GraphPinObj, BoolString);
+	}
 }
\ No newline at end of file
diff --git a/Engine/Source/Editor/GraphEditor/Private/KismetPins/SGraphPinClass.cpp b/Engine/Source/Editor/GraphEditor/Private/KismetPins/SGraphPinClass.cpp
--- a/Engine/Source/Editor/GraphEditor/Private/KismetPins/SGraphPinClass.cpp
+++ b/Engine/Source/Editor/GraphEditor/Private/KismetPins/SGraphPinClass.cpp
@@ -102,8 +102,14 @@ FOnClicked SGraphPinClass::GetOnUseButtonDelegate()
 
 void SGraphPinClass::OnPickedNewClass(UClass* ChosenClass)
 {
-	AssetPickerAnchor->SetIsOpen(false);
-	GraphPinObj->GetSchema()->TrySetDefaultObject(*GraphPinObj, ChosenClass);
+	if(GraphPinObj->DefaultObject != ChosenClass)
+	{
+		const FScopedTransaction Transaction( NSLOCTEXT("GraphEditor", "ChangeClassPinValue", "Change Class Pin Value" ) );
+		GraphPinObj->Modify();
+
+		AssetPickerAnchor->SetIsOpen(false);
+		GraphPinObj->GetSchema()->TrySetDefaultObject(*GraphPinObj, ChosenClass);
+	}
 }
 
 FText SGraphPinClass::GetDefaultComboText() const
diff --git a/Engine/Source/Editor/GraphEditor/Private/KismetPins/SGraphPinCollisionProfile.cpp b/Engine/Source/Editor/GraphEditor/Private/KismetPins/SGraphPinCollisionProfile.cpp
--- a/Engine/Source/Editor/GraphEditor/Private/KismetPins/SGraphPinCollisionProfile.cpp
+++ b/Engine/Source/Editor/GraphEditor/Private/KismetPins/SGraphPinCollisionProfile.cpp
@@ -78,14 +78,20 @@ void SGraphPinCollisionProfile::SetPropertyWithName(const FName& Name)
 {
 	check(GraphPinObj);
 	check(GraphPinObj->PinType.PinSubCategoryObject == FCollisionProfileName::StaticStruct());
-
+	
 	FString PinString = TEXT("(Name=\"");
 	PinString += *Name.ToString();
 	PinString += TEXT("\")");
 
-	if (PinString != GraphPinObj->GetDefaultAsString())
+	if(GraphPinObj->GetDefaultAsString() != PinString)
 	{
-		GraphPinObj->GetSchema()->TrySetDefaultValue(*GraphPinObj, PinString);
+		const FScopedTransaction Transaction( NSLOCTEXT("GraphEditor", "ChangeCollisionProfilePinValue", "Change Collision Profile Pin Value" ) );
+		GraphPinObj->Modify();
+
+		if (PinString != GraphPinObj->GetDefaultAsString())
+		{
+			GraphPinObj->GetSchema()->TrySetDefaultValue(*GraphPinObj, PinString);
+		}
 	}
 }
 
diff --git a/Engine/Source/Editor/GraphEditor/Private/KismetPins/SGraphPinColor.cpp b/Engine/Source/Editor/GraphEditor/Private/KismetPins/SGraphPinColor.cpp
--- a/Engine/Source/Editor/GraphEditor/Private/KismetPins/SGraphPinColor.cpp
+++ b/Engine/Source/Editor/GraphEditor/Private/KismetPins/SGraphPinColor.cpp
@@ -74,6 +74,13 @@ void SGraphPinColor::OnColorCommitted(FLinearColor InColor)
 {
 	// Update pin object
 	FString ColorString = InColor.ToString();
-	GraphPinObj->GetSchema()->TrySetDefaultValue(*GraphPinObj, ColorString);
-	OwnerNodePtr.Pin()->UpdateGraphNode();
+
+	if(GraphPinObj->GetDefaultAsString() != ColorString)
+	{
+		const FScopedTransaction Transaction( NSLOCTEXT("GraphEditor", "ChangeColorPinValue", "Change Color Pin Value" ) );
+		GraphPinObj->Modify();
+		
+		GraphPinObj->GetSchema()->TrySetDefaultValue(*GraphPinObj, ColorString);
+		OwnerNodePtr.Pin()->UpdateGraphNode();
+	}
 }
\ No newline at end of file
diff --git a/Engine/Source/Editor/GraphEditor/Private/KismetPins/SGraphPinEnum.cpp b/Engine/Source/Editor/GraphEditor/Private/KismetPins/SGraphPinEnum.cpp
--- a/Engine/Source/Editor/GraphEditor/Private/KismetPins/SGraphPinEnum.cpp
+++ b/Engine/Source/Editor/GraphEditor/Private/KismetPins/SGraphPinEnum.cpp
@@ -122,8 +122,17 @@ void SGraphPinEnum::ComboBoxSelectionChanged( TSharedPtr<int32> NewSelection, ES
 	UEnum* EnumPtr = Cast<UEnum>(GraphPinObj->PinType.PinSubCategoryObject.Get());
 	check(EnumPtr);
 	check(*NewSelection < EnumPtr->NumEnums() - 1);
-	//Set new selection
-	GraphPinObj->GetSchema()->TrySetDefaultValue(*GraphPinObj, EnumPtr->GetEnumName(*NewSelection));
+
+	FString EnumSelectionString = EnumPtr->GetEnumName(*NewSelection);
+
+	if(GraphPinObj->GetDefaultAsString() != EnumSelectionString)
+	{
+		const FScopedTransaction Transaction( NSLOCTEXT("GraphEditor", "ChangeEnumPinValue", "Change Enum Pin Value" ) );
+		GraphPinObj->Modify();
+
+		//Set new selection
+		GraphPinObj->GetSchema()->TrySetDefaultValue(*GraphPinObj, EnumSelectionString);
+	}
 }
 
 FString SGraphPinEnum::OnGetText() const 
diff --git a/Engine/Source/Editor/GraphEditor/Private/KismetPins/SGraphPinKey.cpp b/Engine/Source/Editor/GraphEditor/Private/KismetPins/SGraphPinKey.cpp
--- a/Engine/Source/Editor/GraphEditor/Private/KismetPins/SGraphPinKey.cpp
+++ b/Engine/Source/Editor/GraphEditor/Private/KismetPins/SGraphPinKey.cpp
@@ -39,6 +39,12 @@ FKey SGraphPinKey::GetCurrentKey() const
 
 void SGraphPinKey::OnKeyChanged(TSharedPtr<FKey> InSelectedKey)
 {
-	SelectedKey = *InSelectedKey.Get();
-	GraphPinObj->GetSchema()->TrySetDefaultValue(*GraphPinObj, SelectedKey.ToString());
+	if(SelectedKey != *InSelectedKey.Get())
+	{
+		const FScopedTransaction Transaction( NSLOCTEXT("GraphEditor", "ChangeKeyPinValue", "Change Key Pin Value" ) );
+		GraphPinObj->Modify();
+
+		SelectedKey = *InSelectedKey.Get();
+		GraphPinObj->GetSchema()->TrySetDefaultValue(*GraphPinObj, SelectedKey.ToString());
+	}
 }
\ No newline at end of file
diff --git a/Engine/Source/Editor/GraphEditor/Private/KismetPins/SGraphPinNameList.cpp b/Engine/Source/Editor/GraphEditor/Private/KismetPins/SGraphPinNameList.cpp
--- a/Engine/Source/Editor/GraphEditor/Private/KismetPins/SGraphPinNameList.cpp
+++ b/Engine/Source/Editor/GraphEditor/Private/KismetPins/SGraphPinNameList.cpp
@@ -44,7 +44,14 @@ void SGraphPinNameList::ComboBoxSelectionChanged(TSharedPtr<FName> NameItem, ESe
 	FName Name = NameItem.IsValid() ? *NameItem : NAME_None;
 	if (auto Schema = (GraphPinObj ? GraphPinObj->GetSchema() : NULL))
 	{
-		Schema->TrySetDefaultValue(*GraphPinObj, *Name.ToString());
+		FString NameAsString = Name.ToString();
+		if(GraphPinObj->GetDefaultAsString() != NameAsString)
+		{
+			const FScopedTransaction Transaction( NSLOCTEXT("GraphEditor", "ChangeNameListPinValue", "Change Name List Pin Value" ) );
+			GraphPinObj->Modify();
+
+			Schema->TrySetDefaultValue(*GraphPinObj, NameAsString);
+		}
 	}
 }
 
diff --git a/Engine/Source/Editor/GraphEditor/Private/KismetPins/SGraphPinNum.cpp b/Engine/Source/Editor/GraphEditor/Private/KismetPins/SGraphPinNum.cpp
--- a/Engine/Source/Editor/GraphEditor/Private/KismetPins/SGraphPinNum.cpp
+++ b/Engine/Source/Editor/GraphEditor/Private/KismetPins/SGraphPinNum.cpp
@@ -17,17 +17,23 @@ void SGraphPinNum::SetTypeInValue(const FText& NewTypeInValue, ETextCommit::Type
 	FString TypeValueString = NewTypeInValue.ToString();
 	if (FDefaultValueHelper::IsStringValidFloat(TypeValueString))
 	{
-		//This is rubbish.
-		//TODO: Make not rubbish.
-		const UEdGraphSchema_K2* K2Schema = Cast<const UEdGraphSchema_K2>(GraphPinObj->GetSchema());
-		if ((K2Schema && (GraphPinObj->PinType.PinCategory == K2Schema->PC_Int || GraphPinObj->PinType.PinCategory == K2Schema->PC_Byte)))
+		if(GraphPinObj->GetDefaultAsString() != TypeValueString)
 		{
-			int32 IntValue = FCString::Atoi(*TypeValueString);
-			GraphPinObj->GetSchema()->TrySetDefaultValue(*GraphPinObj, FString::FromInt(IntValue));
-		}
-		else
-		{
-			GraphPinObj->GetSchema()->TrySetDefaultValue(*GraphPinObj, *TypeValueString);
+			const FScopedTransaction Transaction( NSLOCTEXT("GraphEditor", "ChangeNumberPinValue", "Change Number Pin Value" ) );
+			GraphPinObj->Modify();
+
+			//This is rubbish.
+			//TODO: Make not rubbish.
+			const UEdGraphSchema_K2* K2Schema = Cast<const UEdGraphSchema_K2>(GraphPinObj->GetSchema());
+			if ((K2Schema && (GraphPinObj->PinType.PinCategory == K2Schema->PC_Int || GraphPinObj->PinType.PinCategory == K2Schema->PC_Byte)))
+			{
+				int32 IntValue = FCString::Atoi(*TypeValueString);
+				GraphPinObj->GetSchema()->TrySetDefaultValue(*GraphPinObj, FString::FromInt(IntValue));
+			}
+			else
+			{
+				GraphPinObj->GetSchema()->TrySetDefaultValue(*GraphPinObj, *TypeValueString);
+			}
 		}
 	}
 }
diff --git a/Engine/Source/Editor/GraphEditor/Private/KismetPins/SGraphPinObject.cpp b/Engine/Source/Editor/GraphEditor/Private/KismetPins/SGraphPinObject.cpp
--- a/Engine/Source/Editor/GraphEditor/Private/KismetPins/SGraphPinObject.cpp
+++ b/Engine/Source/Editor/GraphEditor/Private/KismetPins/SGraphPinObject.cpp
@@ -199,11 +199,18 @@ TSharedRef<SWidget> SGraphPinObject::GenerateAssetPicker()
 
 void SGraphPinObject::OnAssetSelectedFromPicker(const class FAssetData& AssetData)
 {
-	// Close the asset picker
-	AssetPickerAnchor->SetIsOpen(false);
+	UObject* AssetObject = AssetData.GetAsset();
+	if(GraphPinObj->DefaultObject != AssetObject)
+	{
+		const FScopedTransaction Transaction( NSLOCTEXT("GraphEditor", "ChangeObjectPinValue", "Change Object Pin Value" ) );
+		GraphPinObj->Modify();
+
+		// Close the asset picker
+		AssetPickerAnchor->SetIsOpen(false);
 
-	// Set the object found from the asset picker
-	GraphPinObj->GetSchema()->TrySetDefaultObject(*GraphPinObj, AssetData.GetAsset());
+		// Set the object found from the asset picker
+		GraphPinObj->GetSchema()->TrySetDefaultObject(*GraphPinObj, AssetObject);
+	}
 }
 
 
diff --git a/Engine/Source/Editor/GraphEditor/Private/KismetPins/SGraphPinString.cpp b/Engine/Source/Editor/GraphEditor/Private/KismetPins/SGraphPinString.cpp
--- a/Engine/Source/Editor/GraphEditor/Private/KismetPins/SGraphPinString.cpp
+++ b/Engine/Source/Editor/GraphEditor/Private/KismetPins/SGraphPinString.cpp
@@ -34,7 +34,13 @@ FText SGraphPinString::GetTypeInValue() const
 
 void SGraphPinString::SetTypeInValue(const FText& NewTypeInValue, ETextCommit::Type /*CommitInfo*/)
 {
-	GraphPinObj->GetSchema()->TrySetDefaultValue(*GraphPinObj, NewTypeInValue.ToString());
+	if(GraphPinObj->GetDefaultAsString() != NewTypeInValue.ToString())
+	{
+		const FScopedTransaction Transaction( NSLOCTEXT("GraphEditor", "ChangeStringPinValue", "Change String Pin Value" ) );
+		GraphPinObj->Modify();
+
+		GraphPinObj->GetSchema()->TrySetDefaultValue(*GraphPinObj, NewTypeInValue.ToString());
+	}
 }
 
 bool SGraphPinString::GetDefaultValueIsReadOnly() const
diff --git a/Engine/Source/Editor/GraphEditor/Private/KismetPins/SGraphPinText.cpp b/Engine/Source/Editor/GraphEditor/Private/KismetPins/SGraphPinText.cpp
--- a/Engine/Source/Editor/GraphEditor/Private/KismetPins/SGraphPinText.cpp
+++ b/Engine/Source/Editor/GraphEditor/Private/KismetPins/SGraphPinText.cpp
@@ -34,8 +34,14 @@ FText SGraphPinText::GetTypeInValue() const
 
 void SGraphPinText::SetTypeInValue(const FText& NewTypeInValue, ETextCommit::Type /*CommitInfo*/)
 {
-	// Update the localization string.
-	GraphPinObj->GetSchema()->TrySetDefaultText(*GraphPinObj, NewTypeInValue);
+	if(!GraphPinObj->DefaultTextValue.EqualTo(NewTypeInValue))
+	{
+		const FScopedTransaction Transaction( NSLOCTEXT("GraphEditor", "ChangeTxtPinValue", "Change Text Pin Value" ) );
+		GraphPinObj->Modify();
+
+		// Update the localization string.
+		GraphPinObj->GetSchema()->TrySetDefaultText(*GraphPinObj, NewTypeInValue);
+	}
 }
 
 bool SGraphPinText::GetDefaultValueIsReadOnly() const
diff --git a/Engine/Source/Editor/GraphEditor/Private/KismetPins/SGraphPinVector.cpp b/Engine/Source/Editor/GraphEditor/Private/KismetPins/SGraphPinVector.cpp
--- a/Engine/Source/Editor/GraphEditor/Private/KismetPins/SGraphPinVector.cpp
+++ b/Engine/Source/Editor/GraphEditor/Private/KismetPins/SGraphPinVector.cpp
@@ -222,8 +222,14 @@ void SGraphPinVector::OnChangedValueTextBox_0(float NewValue, ETextCommit::Type
 		DefaultValue = ValueStr + FString(TEXT(",")) + GetValue(TextBox_1) + FString(TEXT(",")) + GetValue(TextBox_2);
 	}
 
-	//Set new default value
-	GraphPinObj->GetSchema()->TrySetDefaultValue(*GraphPinObj, DefaultValue);
+	if(GraphPinObj->GetDefaultAsString() != DefaultValue)
+	{
+		const FScopedTransaction Transaction( NSLOCTEXT("GraphEditor", "ChangeVectorPinValue", "Change Vector Pin Value" ) );
+		GraphPinObj->Modify();
+
+		//Set new default value
+		GraphPinObj->GetSchema()->TrySetDefaultValue(*GraphPinObj, DefaultValue);
+	}
 }
 
 void SGraphPinVector::OnChangedValueTextBox_1(float NewValue, ETextCommit::Type CommitInfo)
@@ -242,8 +248,14 @@ void SGraphPinVector::OnChangedValueTextBox_1(float NewValue, ETextCommit::Type
 		DefaultValue = GetValue(TextBox_0) + FString(TEXT(",")) + ValueStr + FString(TEXT(",")) + GetValue(TextBox_2);
 	}
 
-	//Set new default value
-	GraphPinObj->GetSchema()->TrySetDefaultValue(*GraphPinObj, DefaultValue);
+	if(GraphPinObj->GetDefaultAsString() != DefaultValue)
+	{
+		const FScopedTransaction Transaction( NSLOCTEXT("GraphEditor", "ChangeVectorPinValue", "Change Vector Pin Value" ) );
+		GraphPinObj->Modify();
+
+		//Set new default value
+		GraphPinObj->GetSchema()->TrySetDefaultValue(*GraphPinObj, DefaultValue);
+	}
 }
 
 void SGraphPinVector::OnChangedValueTextBox_2(float NewValue, ETextCommit::Type CommitInfo)
@@ -262,8 +274,14 @@ void SGraphPinVector::OnChangedValueTextBox_2(float NewValue, ETextCommit::Type
 		DefaultValue = GetValue(TextBox_0) + FString(TEXT(",")) + GetValue(TextBox_1) + FString(TEXT(",")) + ValueStr;
 	}
 
-	//Set new default value
-	GraphPinObj->GetSchema()->TrySetDefaultValue(*GraphPinObj, DefaultValue);
+	if(GraphPinObj->GetDefaultAsString() != DefaultValue)
+	{
+		const FScopedTransaction Transaction( NSLOCTEXT("GraphEditor", "ChangeVectorPinValue", "Change Vector Pin Value" ) );
+		GraphPinObj->Modify();
+
+		//Set new default value
+		GraphPinObj->GetSchema()->TrySetDefaultValue(*GraphPinObj, DefaultValue);
+	}
 }
 
 #undef LOCTEXT_NAMESPACE
diff --git a/Engine/Source/Editor/GraphEditor/Private/KismetPins/SGraphPinVector2D.cpp b/Engine/Source/Editor/GraphEditor/Private/KismetPins/SGraphPinVector2D.cpp
--- a/Engine/Source/Editor/GraphEditor/Private/KismetPins/SGraphPinVector2D.cpp
+++ b/Engine/Source/Editor/GraphEditor/Private/KismetPins/SGraphPinVector2D.cpp
@@ -157,17 +157,31 @@ FString MakeVector2DString(const FString& X, const FString& Y)
 void SGraphPinVector2D::OnChangedValueTextBox_X(float NewValue, ETextCommit::Type CommitInfo)
 {
 	const FString ValueStr = FString::Printf( TEXT("%f"), NewValue );
+	const FString Vector2DString = MakeVector2DString(ValueStr, GetValue(TextBox_Y));
 
-	//Set new default value
-	GraphPinObj->GetSchema()->TrySetDefaultValue(*GraphPinObj, MakeVector2DString(ValueStr, GetValue(TextBox_Y)));
+	if(GraphPinObj->GetDefaultAsString() != Vector2DString)
+	{
+		const FScopedTransaction Transaction( NSLOCTEXT("GraphEditor", "ChangeVectorPinValue", "Change Vector Pin Value" ) );
+		GraphPinObj->Modify();
+
+		//Set new default value
+		GraphPinObj->GetSchema()->TrySetDefaultValue(*GraphPinObj, Vector2DString);
+	}
 }
 
 void SGraphPinVector2D::OnChangedValueTextBox_Y(float NewValue, ETextCommit::Type CommitInfo)
 {
 	const FString ValueStr = FString::Printf( TEXT("%f"), NewValue );
+	const FString Vector2DString =MakeVector2DString(GetValue(TextBox_X), ValueStr);
 
-	//Set new default value
-	GraphPinObj->GetSchema()->TrySetDefaultValue(*GraphPinObj, MakeVector2DString(GetValue(TextBox_X), ValueStr));
+	if(GraphPinObj->GetDefaultAsString() != Vector2DString)
+	{
+		const FScopedTransaction Transaction( NSLOCTEXT("GraphEditor", "ChangeVectorPinValue", "Change Vector Pin Value" ) );
+		GraphPinObj->Modify();
+
+		//Set new default value
+		GraphPinObj->GetSchema()->TrySetDefaultValue(*GraphPinObj, Vector2DString);
+	}
 }
 
 #undef LOCTEXT_NAMESPACE
diff --git a/Engine/Source/Runtime/Core/Private/Internationalization/Text.cpp b/Engine/Source/Runtime/Core/Private/Internationalization/Text.cpp
--- a/Engine/Source/Runtime/Core/Private/Internationalization/Text.cpp
+++ b/Engine/Source/Runtime/Core/Private/Internationalization/Text.cpp
@@ -647,6 +647,7 @@ CORE_API FArchive& operator<<( FArchive& Ar, FText& Value )
 			default:
 				{
 					Value.History.Reset();
+					Value.DisplayString = FText::GetEmpty().DisplayString;
 				}
 			}
 		}
