Fix non-unity build issues.

diff --git a/Engine/Source/Editor/GraphEditor/Private/KismetPins/SGraphPinBool.cpp b/Engine/Source/Editor/GraphEditor/Private/KismetPins/SGraphPinBool.cpp
--- a/Engine/Source/Editor/GraphEditor/Private/KismetPins/SGraphPinBool.cpp
+++ b/Engine/Source/Editor/GraphEditor/Private/KismetPins/SGraphPinBool.cpp
@@ -3,6 +3,7 @@
 
 #include "GraphEditorCommon.h"
 #include "SGraphPinBool.h"
+#include "Editor/UnrealEd/Public/ScopedTransaction.h"
 
 void SGraphPinBool::Construct(const FArguments& InArgs, UEdGraphPin* InGraphPinObj)
 {
diff --git a/Engine/Source/Editor/GraphEditor/Private/KismetPins/SGraphPinClass.cpp b/Engine/Source/Editor/GraphEditor/Private/KismetPins/SGraphPinClass.cpp
--- a/Engine/Source/Editor/GraphEditor/Private/KismetPins/SGraphPinClass.cpp
+++ b/Engine/Source/Editor/GraphEditor/Private/KismetPins/SGraphPinClass.cpp
@@ -5,6 +5,7 @@
 #include "SGraphPinClass.h"
 #include "Editor/ClassViewer/Public/ClassViewerModule.h"
 #include "Editor/ClassViewer/Public/ClassViewerFilter.h"
+#include "Editor/UnrealEd/Public/ScopedTransaction.h"
 
 #define LOCTEXT_NAMESPACE "SGraphPinClass"
 
diff --git a/Engine/Source/Editor/GraphEditor/Private/KismetPins/SGraphPinCollisionProfile.cpp b/Engine/Source/Editor/GraphEditor/Private/KismetPins/SGraphPinCollisionProfile.cpp
--- a/Engine/Source/Editor/GraphEditor/Private/KismetPins/SGraphPinCollisionProfile.cpp
+++ b/Engine/Source/Editor/GraphEditor/Private/KismetPins/SGraphPinCollisionProfile.cpp
@@ -3,7 +3,7 @@
 #include "GraphEditorCommon.h"
 #include "SNameComboBox.h"
 #include "SGraphPinCollisionProfile.h"
-
+#include "Editor/UnrealEd/Public/ScopedTransaction.h"
 
 void SGraphPinCollisionProfile::Construct(const FArguments& InArgs, UEdGraphPin* InGraphPinObj)
 {
diff --git a/Engine/Source/Editor/GraphEditor/Private/KismetPins/SGraphPinColor.cpp b/Engine/Source/Editor/GraphEditor/Private/KismetPins/SGraphPinColor.cpp
--- a/Engine/Source/Editor/GraphEditor/Private/KismetPins/SGraphPinColor.cpp
+++ b/Engine/Source/Editor/GraphEditor/Private/KismetPins/SGraphPinColor.cpp
@@ -5,6 +5,7 @@
 #include "NodeFactory.h"
 #include "KismetPins/SGraphPinColor.h"
 #include "SColorPicker.h"
+#include "Editor/UnrealEd/Public/ScopedTransaction.h"
 
 
 void SGraphPinColor::Construct(const FArguments& InArgs, UEdGraphPin* InGraphPinObj)
diff --git a/Engine/Source/Editor/GraphEditor/Private/KismetPins/SGraphPinEnum.cpp b/Engine/Source/Editor/GraphEditor/Private/KismetPins/SGraphPinEnum.cpp
--- a/Engine/Source/Editor/GraphEditor/Private/KismetPins/SGraphPinEnum.cpp
+++ b/Engine/Source/Editor/GraphEditor/Private/KismetPins/SGraphPinEnum.cpp
@@ -4,6 +4,7 @@
 #include "GraphEditorCommon.h"
 #include "SGraphPinEnum.h"
 #include "SGraphPinComboBox.h"
+#include "Editor/UnrealEd/Public/ScopedTransaction.h"
 
 //Construct combo box using combo button and combo list
 void SPinComboBox::Construct( const FArguments& InArgs )
diff --git a/Engine/Source/Editor/GraphEditor/Private/KismetPins/SGraphPinKey.cpp b/Engine/Source/Editor/GraphEditor/Private/KismetPins/SGraphPinKey.cpp
--- a/Engine/Source/Editor/GraphEditor/Private/KismetPins/SGraphPinKey.cpp
+++ b/Engine/Source/Editor/GraphEditor/Private/KismetPins/SGraphPinKey.cpp
@@ -6,6 +6,7 @@
 #include "SGraphPinEnum.h"
 #include "SGraphPinKey.h"
 #include "SKeySelector.h"
+#include "Editor/UnrealEd/Public/ScopedTransaction.h"
 
 void SGraphPinKey::Construct(const FArguments& InArgs, UEdGraphPin* InGraphPinObj)
 {
diff --git a/Engine/Source/Editor/GraphEditor/Private/KismetPins/SGraphPinNameList.cpp b/Engine/Source/Editor/GraphEditor/Private/KismetPins/SGraphPinNameList.cpp
--- a/Engine/Source/Editor/GraphEditor/Private/KismetPins/SGraphPinNameList.cpp
+++ b/Engine/Source/Editor/GraphEditor/Private/KismetPins/SGraphPinNameList.cpp
@@ -4,6 +4,7 @@
 #include "GraphEditorCommon.h"
 #include "SGraphPinComboBox.h"
 #include "SGraphPinNameList.h"
+#include "Editor/UnrealEd/Public/ScopedTransaction.h"
 
 void SGraphPinNameList::Construct(const FArguments& InArgs, UEdGraphPin* InGraphPinObj, const TArray<TSharedPtr<FName>>& InNameList)
 {
diff --git a/Engine/Source/Editor/GraphEditor/Private/KismetPins/SGraphPinNum.cpp b/Engine/Source/Editor/GraphEditor/Private/KismetPins/SGraphPinNum.cpp
--- a/Engine/Source/Editor/GraphEditor/Private/KismetPins/SGraphPinNum.cpp
+++ b/Engine/Source/Editor/GraphEditor/Private/KismetPins/SGraphPinNum.cpp
@@ -6,6 +6,7 @@
 #include "SGraphPinString.h"
 #include "SGraphPinNum.h"
 #include "DefaultValueHelper.h"
+#include "Editor/UnrealEd/Public/ScopedTransaction.h"
 
 void SGraphPinNum::Construct(const FArguments& InArgs, UEdGraphPin* InGraphPinObj)
 {
diff --git a/Engine/Source/Editor/GraphEditor/Private/KismetPins/SGraphPinObject.cpp b/Engine/Source/Editor/GraphEditor/Private/KismetPins/SGraphPinObject.cpp
--- a/Engine/Source/Editor/GraphEditor/Private/KismetPins/SGraphPinObject.cpp
+++ b/Engine/Source/Editor/GraphEditor/Private/KismetPins/SGraphPinObject.cpp
@@ -5,6 +5,7 @@
 #include "SGraphPinClass.h"
 #include "Editor/ContentBrowser/Public/ContentBrowserModule.h"
 #include "AssetData.h"
+#include "Editor/UnrealEd/Public/ScopedTransaction.h"
 
 #define LOCTEXT_NAMESPACE "SGraphPinObject"
 
diff --git a/Engine/Source/Editor/GraphEditor/Private/KismetPins/SGraphPinString.cpp b/Engine/Source/Editor/GraphEditor/Private/KismetPins/SGraphPinString.cpp
--- a/Engine/Source/Editor/GraphEditor/Private/KismetPins/SGraphPinString.cpp
+++ b/Engine/Source/Editor/GraphEditor/Private/KismetPins/SGraphPinString.cpp
@@ -3,7 +3,7 @@
 
 #include "GraphEditorCommon.h"
 #include "SGraphPinString.h"
-
+#include "Editor/UnrealEd/Public/ScopedTransaction.h"
 
 void SGraphPinString::Construct(const FArguments& InArgs, UEdGraphPin* InGraphPinObj)
 {
diff --git a/Engine/Source/Editor/GraphEditor/Private/KismetPins/SGraphPinText.cpp b/Engine/Source/Editor/GraphEditor/Private/KismetPins/SGraphPinText.cpp
--- a/Engine/Source/Editor/GraphEditor/Private/KismetPins/SGraphPinText.cpp
+++ b/Engine/Source/Editor/GraphEditor/Private/KismetPins/SGraphPinText.cpp
@@ -3,7 +3,7 @@
 
 #include "GraphEditorCommon.h"
 #include "SGraphPinText.h"
-
+#include "Editor/UnrealEd/Public/ScopedTransaction.h"
 
 void SGraphPinText::Construct(const FArguments& InArgs, UEdGraphPin* InGraphPinObj)
 {
diff --git a/Engine/Source/Editor/GraphEditor/Private/KismetPins/SGraphPinVector.cpp b/Engine/Source/Editor/GraphEditor/Private/KismetPins/SGraphPinVector.cpp
--- a/Engine/Source/Editor/GraphEditor/Private/KismetPins/SGraphPinVector.cpp
+++ b/Engine/Source/Editor/GraphEditor/Private/KismetPins/SGraphPinVector.cpp
@@ -4,7 +4,7 @@
 #include "GraphEditorCommon.h"
 #include "SGraphPinVector.h"
 #include "SNumericEntryBox.h"
-
+#include "Editor/UnrealEd/Public/ScopedTransaction.h"
 
 #define LOCTEXT_NAMESPACE "VectorTextBox"
 
diff --git a/Engine/Source/Editor/GraphEditor/Private/KismetPins/SGraphPinVector2D.cpp b/Engine/Source/Editor/GraphEditor/Private/KismetPins/SGraphPinVector2D.cpp
--- a/Engine/Source/Editor/GraphEditor/Private/KismetPins/SGraphPinVector2D.cpp
+++ b/Engine/Source/Editor/GraphEditor/Private/KismetPins/SGraphPinVector2D.cpp
@@ -4,6 +4,7 @@
 #include "GraphEditorCommon.h"
 #include "SGraphPinVector2D.h"
 #include "SNumericEntryBox.h"
+#include "Editor/UnrealEd/Public/ScopedTransaction.h"
 
 #define LOCTEXT_NAMESPACE "VectorTextBox"
 
diff --git a/Engine/Source/Editor/GraphEditor/Private/NiagaraPins/SGraphPinVector4.cpp b/Engine/Source/Editor/GraphEditor/Private/NiagaraPins/SGraphPinVector4.cpp
--- a/Engine/Source/Editor/GraphEditor/Private/NiagaraPins/SGraphPinVector4.cpp
+++ b/Engine/Source/Editor/GraphEditor/Private/NiagaraPins/SGraphPinVector4.cpp
@@ -4,6 +4,7 @@
 #include "GraphEditorCommon.h"
 #include "SGraphPinVector4.h"
 #include "SNumericEntryBox.h"
+#include "Editor/UnrealEd/Public/ScopedTransaction.h"
 
 
 #define LOCTEXT_NAMESPACE "VectorTextBox"
@@ -250,8 +251,15 @@ void SGraphPinVector4::OnChangedValueTextBox_0(float NewValue, ETextCommit::Type
 	FString DefaultValue;
 	//Update X value
 	DefaultValue = ValueStr + FString(TEXT(",")) + GetValue(TextBox_1) + FString(TEXT(",")) + GetValue(TextBox_2) + FString(TEXT(",")) + GetValue(TextBox_3);
-	//Set new default value
-	GraphPinObj->GetSchema()->TrySetDefaultValue(*GraphPinObj, DefaultValue);
+
+	if(GraphPinObj->GetDefaultAsString() != DefaultValue)
+	{
+		const FScopedTransaction Transaction( NSLOCTEXT("GraphEditor", "ChangeVector4PinValue", "Change Vector4 Pin Value" ) );
+		GraphPinObj->Modify();
+
+		//Set new default value
+		GraphPinObj->GetSchema()->TrySetDefaultValue(*GraphPinObj, DefaultValue);
+	}
 }
 
 void SGraphPinVector4::OnChangedValueTextBox_1(float NewValue, ETextCommit::Type CommitInfo)
@@ -261,8 +269,15 @@ void SGraphPinVector4::OnChangedValueTextBox_1(float NewValue, ETextCommit::Type
 	FString DefaultValue;
 	//Update Y value
 	DefaultValue = GetValue(TextBox_0) + FString(TEXT(",")) + ValueStr + FString(TEXT(",")) + GetValue(TextBox_2) + FString(TEXT(",")) + GetValue(TextBox_3);
-	//Set new default value
-	GraphPinObj->GetSchema()->TrySetDefaultValue(*GraphPinObj, DefaultValue);
+
+	if(GraphPinObj->GetDefaultAsString() != DefaultValue)
+	{
+		const FScopedTransaction Transaction( NSLOCTEXT("GraphEditor", "ChangeVector4PinValue", "Change Vector4 Pin Value" ) );
+		GraphPinObj->Modify();
+
+		//Set new default value
+		GraphPinObj->GetSchema()->TrySetDefaultValue(*GraphPinObj, DefaultValue);
+	}
 }
 
 void SGraphPinVector4::OnChangedValueTextBox_2(float NewValue, ETextCommit::Type CommitInfo)
@@ -272,8 +287,15 @@ void SGraphPinVector4::OnChangedValueTextBox_2(float NewValue, ETextCommit::Type
 	FString DefaultValue;
 	//Update Z value
 	DefaultValue = GetValue(TextBox_0) + FString(TEXT(",")) + GetValue(TextBox_1) + FString(TEXT(",")) + ValueStr + FString(TEXT(",")) + GetValue(TextBox_3);
-	//Set new default value
-	GraphPinObj->GetSchema()->TrySetDefaultValue(*GraphPinObj, DefaultValue);
+
+	if(GraphPinObj->GetDefaultAsString() != DefaultValue)
+	{
+		const FScopedTransaction Transaction( NSLOCTEXT("GraphEditor", "ChangeVector4PinValue", "Change Vector4 Pin Value" ) );
+		GraphPinObj->Modify();
+
+		//Set new default value
+		GraphPinObj->GetSchema()->TrySetDefaultValue(*GraphPinObj, DefaultValue);
+	}
 }
 
 void SGraphPinVector4::OnChangedValueTextBox_3(float NewValue, ETextCommit::Type CommitInfo)
@@ -284,8 +306,14 @@ void SGraphPinVector4::OnChangedValueTextBox_3(float NewValue, ETextCommit::Type
 	//Update W value
 	DefaultValue = GetValue(TextBox_0) + FString(TEXT(",")) + GetValue(TextBox_1) + FString(TEXT(",")) + GetValue(TextBox_2) + FString(TEXT(",")) + ValueStr;
 
-	//Set new default value
-	GraphPinObj->GetSchema()->TrySetDefaultValue(*GraphPinObj, DefaultValue);
+	if(GraphPinObj->GetDefaultAsString() != DefaultValue)
+	{
+		const FScopedTransaction Transaction( NSLOCTEXT("GraphEditor", "ChangeVector4PinValue", "Change Vector4 Pin Value" ) );
+		GraphPinObj->Modify();
+
+		//Set new default value
+		GraphPinObj->GetSchema()->TrySetDefaultValue(*GraphPinObj, DefaultValue);
+	}
 }
 
 
