Blueprint - Tweaking the min/max sizes on the multiline text pins.  Also enabling auto word wrap.

diff --git a/Engine/Source/Editor/GraphEditor/Private/KismetPins/SGraphPinString.cpp b/Engine/Source/Editor/GraphEditor/Private/KismetPins/SGraphPinString.cpp
--- a/Engine/Source/Editor/GraphEditor/Private/KismetPins/SGraphPinString.cpp
+++ b/Engine/Source/Editor/GraphEditor/Private/KismetPins/SGraphPinString.cpp
@@ -14,6 +14,8 @@ TSharedRef<SWidget>	SGraphPinString::GetDefaultValueWidget()
 {
 	return SNew(SBox)
 		.MinDesiredWidth(18.0f)
+		.MaxDesiredWidth(400)
+		.MaxDesiredHeight(200)
 		[
 			SNew(SMultiLineEditableTextBox)
 			.Style(FEditorStyle::Get(), "Graph.EditableTextBox")
@@ -23,6 +25,7 @@ TSharedRef<SWidget>	SGraphPinString::GetDefaultValueWidget()
 			.IsReadOnly(this, &SGraphPinString::GetDefaultValueIsReadOnly)
 			.OnTextCommitted(this, &SGraphPinString::SetTypeInValue)
 			.ForegroundColor(FSlateColor::UseForeground())
+			.AutoWrapText(true)
 			.ModiferKeyForNewLine(EModifierKey::Shift)
 		];
 }
diff --git a/Engine/Source/Editor/GraphEditor/Private/KismetPins/SGraphPinText.cpp b/Engine/Source/Editor/GraphEditor/Private/KismetPins/SGraphPinText.cpp
--- a/Engine/Source/Editor/GraphEditor/Private/KismetPins/SGraphPinText.cpp
+++ b/Engine/Source/Editor/GraphEditor/Private/KismetPins/SGraphPinText.cpp
@@ -14,6 +14,8 @@ TSharedRef<SWidget>	SGraphPinText::GetDefaultValueWidget()
 {
 	return SNew(SBox)
 		.MinDesiredWidth(18.0f)
+		.MaxDesiredWidth(400)
+		.MaxDesiredHeight(200)
 		[
 			SNew(SMultiLineEditableTextBox)
 			.Style(FEditorStyle::Get(), "Graph.EditableTextBox")
@@ -23,6 +25,7 @@ TSharedRef<SWidget>	SGraphPinText::GetDefaultValueWidget()
 			.IsReadOnly(this, &SGraphPinText::GetDefaultValueIsReadOnly)
 			.OnTextCommitted(this, &SGraphPinText::SetTypeInValue)
 			.ForegroundColor(FSlateColor::UseForeground())
+			.AutoWrapText(true)
 			.ModiferKeyForNewLine(EModifierKey::Shift)
 		];
 }
