Blueprint - Now wrapping the text at the max desired size of the text/string pins instead of relying on auto wrap text, which has bad side effects on auto width textblocks.

diff --git a/Engine/Source/Editor/GraphEditor/Private/KismetPins/SGraphPinString.cpp b/Engine/Source/Editor/GraphEditor/Private/KismetPins/SGraphPinString.cpp
--- a/Engine/Source/Editor/GraphEditor/Private/KismetPins/SGraphPinString.cpp
+++ b/Engine/Source/Editor/GraphEditor/Private/KismetPins/SGraphPinString.cpp
@@ -14,7 +14,6 @@ TSharedRef<SWidget>	SGraphPinString::GetDefaultValueWidget()
 {
 	return SNew(SBox)
 		.MinDesiredWidth(18.0f)
-		.MaxDesiredWidth(400)
 		.MaxDesiredHeight(200)
 		[
 			SNew(SMultiLineEditableTextBox)
@@ -25,7 +24,7 @@ TSharedRef<SWidget>	SGraphPinString::GetDefaultValueWidget()
 			.IsReadOnly(this, &SGraphPinString::GetDefaultValueIsReadOnly)
 			.OnTextCommitted(this, &SGraphPinString::SetTypeInValue)
 			.ForegroundColor(FSlateColor::UseForeground())
-			.AutoWrapText(true)
+			.WrapTextAt(400)
 			.ModiferKeyForNewLine(EModifierKey::Shift)
 		];
 }
diff --git a/Engine/Source/Editor/GraphEditor/Private/KismetPins/SGraphPinText.cpp b/Engine/Source/Editor/GraphEditor/Private/KismetPins/SGraphPinText.cpp
--- a/Engine/Source/Editor/GraphEditor/Private/KismetPins/SGraphPinText.cpp
+++ b/Engine/Source/Editor/GraphEditor/Private/KismetPins/SGraphPinText.cpp
@@ -14,7 +14,6 @@ TSharedRef<SWidget>	SGraphPinText::GetDefaultValueWidget()
 {
 	return SNew(SBox)
 		.MinDesiredWidth(18.0f)
-		.MaxDesiredWidth(400)
 		.MaxDesiredHeight(200)
 		[
 			SNew(SMultiLineEditableTextBox)
@@ -25,7 +24,7 @@ TSharedRef<SWidget>	SGraphPinText::GetDefaultValueWidget()
 			.IsReadOnly(this, &SGraphPinText::GetDefaultValueIsReadOnly)
 			.OnTextCommitted(this, &SGraphPinText::SetTypeInValue)
 			.ForegroundColor(FSlateColor::UseForeground())
-			.AutoWrapText(true)
+			.WrapTextAt(400)
 			.ModiferKeyForNewLine(EModifierKey::Shift)
 		];
 }
