GraphEditor: Add support for direction flipping for reroute nodes that travel backwards
- Move PinGeometries to be a member variable of FConnectionDrawingPolicy
- Known issue: The preview wire doesn't flip directions until it is connected to something

diff --git a/Engine/Source/Editor/GraphEditor/Private/BlueprintConnectionDrawingPolicy.cpp b/Engine/Source/Editor/GraphEditor/Private/BlueprintConnectionDrawingPolicy.cpp
--- a/Engine/Source/Editor/GraphEditor/Private/BlueprintConnectionDrawingPolicy.cpp
+++ b/Engine/Source/Editor/GraphEditor/Private/BlueprintConnectionDrawingPolicy.cpp
@@ -339,6 +339,86 @@ FKismetConnectionDrawingPolicy::FTimePair const* FKismetConnectionDrawingPolicy:
 	return FoundExecPath;
 }
 
+bool FKismetConnectionDrawingPolicy::FindPinCenter(UEdGraphPin* Pin, FVector2D& OutCenter) const
+{
+	if (const TSharedRef<SGraphPin>* pPinWidget = PinToPinWidgetMap.Find(Pin))
+	{
+		if (FArrangedWidget* pPinEntry = PinGeometries->Find(*pPinWidget))
+		{
+			OutCenter = FGeometryHelper::CenterOf(pPinEntry->Geometry);
+			return true;
+		}
+	}
+
+	return false;
+}
+
+bool FKismetConnectionDrawingPolicy::GetAverageConnectedPosition(class UK2Node_Knot* Knot, EEdGraphPinDirection Direction, FVector2D& OutPos) const
+{
+	FVector2D Result = FVector2D::ZeroVector;
+	int32 ResultCount = 0;
+
+	UEdGraphPin* Pin = (Direction == EGPD_Input) ? Knot->GetInputPin() : Knot->GetOutputPin();
+	for (UEdGraphPin* LinkedPin : Pin->LinkedTo)
+	{
+		FVector2D CenterPoint;
+		if (FindPinCenter(LinkedPin, /*out*/ CenterPoint))
+		{
+			Result += CenterPoint;
+			ResultCount++;
+		}
+	}
+
+	if (ResultCount > 0)
+	{
+		OutPos = Result * (1.0f / ResultCount);
+		return true;
+	}
+	else
+	{
+		return false;
+	}
+}
+
+bool FKismetConnectionDrawingPolicy::ShouldChangeTangentForKnot(UK2Node_Knot* Knot)
+{
+	if (bool* pResult = KnotToReversedDirectionMap.Find(Knot))
+	{
+		return *pResult;
+	}
+	else
+	{
+		bool bPinReversed = false;
+
+		FVector2D AverageLeftPin;
+		FVector2D AverageRightPin;
+		FVector2D CenterPin;
+		bool bCenterValid = FindPinCenter(Knot->GetOutputPin(), /*out*/ CenterPin);
+		bool bLeftValid = GetAverageConnectedPosition(Knot, EGPD_Input, /*out*/ AverageLeftPin);
+		bool bRightValid = GetAverageConnectedPosition(Knot, EGPD_Output, /*out*/ AverageRightPin);
+
+		if (bLeftValid && bRightValid)
+		{
+			bPinReversed = AverageRightPin.X < AverageLeftPin.X;
+		}
+		else if (bCenterValid)
+		{
+			if (bLeftValid)
+			{
+				bPinReversed = CenterPin.X < AverageLeftPin.X;
+			}
+			else if (bRightValid)
+			{
+				bPinReversed = AverageRightPin.X < CenterPin.X;
+			}
+		}
+
+		KnotToReversedDirectionMap.Add(Knot, bPinReversed);
+
+		return bPinReversed;
+	}
+}
+
 // Give specific editor modes a chance to highlight this connection or darken non-interesting connections
 void FKismetConnectionDrawingPolicy::DetermineWiringStyle(UEdGraphPin* OutputPin, UEdGraphPin* InputPin, /*inout*/ FConnectionParams& Params)
 {
@@ -358,6 +438,25 @@ void FKismetConnectionDrawingPolicy::DetermineWiringStyle(UEdGraphPin* OutputPin
 	const UEdGraphSchema_K2* K2Schema = Cast<const UEdGraphSchema_K2>(Schema);
 	if (K2Schema != NULL)
 	{
+		// If the output or input connect to a knot that is going backwards, we will flip the direction on values going into them
+		{
+			if (UK2Node_Knot* OutputKnotNode = Cast<UK2Node_Knot>(OutputNode))
+			{
+				if (ShouldChangeTangentForKnot(OutputKnotNode))
+				{
+					Params.StartDirection = EGPD_Input;
+				}
+			}
+
+			if (UK2Node_Knot* InputKnotNode = Cast<UK2Node_Knot>(InputNode))
+			{
+				if (ShouldChangeTangentForKnot(InputKnotNode))
+				{
+					Params.EndDirection = EGPD_Output;
+				}
+			}
+		}
+
 		if (TreatWireAsExecutionPin(InputPin, OutputPin))
 		{
 			if (CanBuildRoadmap())
diff --git a/Engine/Source/Editor/GraphEditor/Private/BlueprintConnectionDrawingPolicy.h b/Engine/Source/Editor/GraphEditor/Private/BlueprintConnectionDrawingPolicy.h
--- a/Engine/Source/Editor/GraphEditor/Private/BlueprintConnectionDrawingPolicy.h
+++ b/Engine/Source/Editor/GraphEditor/Private/BlueprintConnectionDrawingPolicy.h
@@ -54,6 +54,9 @@ protected:
 	float ReleasePeriod;
 	float ReleaseExponent;
 
+	// Each time a knot is encountered, input geometry is compared to output geometry to see if the pins on the knot need to be reversed
+	TMap<class UK2Node_Knot*, bool> KnotToReversedDirectionMap;
+
 protected:
 	// Should this wire be treated as an execution pin, querying the execution trace for connected nodes to draw it differently?
 	virtual bool TreatWireAsExecutionPin(UEdGraphPin* InputPin, UEdGraphPin* OutputPin) const;
@@ -102,4 +105,8 @@ private:
 	 * @return Null if OutputPin did not cause the invocation of the node in question, otherwise a valid time pair denoting when the node was executed.
 	 */
 	FTimePair const* BackTraceExecPath(UEdGraphPin const* const OutputPin, FExecPairingMap const* const NodeExecutionList);
+
+	bool ShouldChangeTangentForKnot(class UK2Node_Knot* Knot);
+	bool GetAverageConnectedPosition(class UK2Node_Knot* Knot, EEdGraphPinDirection Direction, FVector2D& OutPos) const;
+	bool FindPinCenter(UEdGraphPin* Pin, FVector2D& OutCenter) const;
 };
diff --git a/Engine/Source/Editor/GraphEditor/Private/ConnectionDrawingPolicy.cpp b/Engine/Source/Editor/GraphEditor/Private/ConnectionDrawingPolicy.cpp
--- a/Engine/Source/Editor/GraphEditor/Private/ConnectionDrawingPolicy.cpp
+++ b/Engine/Source/Editor/GraphEditor/Private/ConnectionDrawingPolicy.cpp
@@ -228,8 +228,8 @@ void FConnectionDrawingPolicy::DrawConnection(int32 LayerId, const FVector2D& St
 	const FVector2D& P1 = End;
 
 	const FVector2D SplineTangent = ComputeSplineTangent(P0, P1);
-	const FVector2D P0Tangent = SplineTangent;
-	const FVector2D P1Tangent = SplineTangent;
+	const FVector2D P0Tangent = (Params.StartDirection == EGPD_Output) ? SplineTangent : -SplineTangent;
+	const FVector2D P1Tangent = (Params.EndDirection == EGPD_Input) ? SplineTangent : -SplineTangent;
 
 	// Draw the spline itself
 	FSlateDrawElement::MakeDrawSpaceSpline(
@@ -327,8 +327,6 @@ void FConnectionDrawingPolicy::DetermineWiringStyle(UEdGraphPin* OutputPin, UEdG
 }
 
 void FConnectionDrawingPolicy::DetermineLinkGeometry(
-	TMap<TSharedRef<SWidget>,
-	FArrangedWidget>& PinGeometries,
 	FArrangedChildren& ArrangedNodes, 
 	TSharedRef<SWidget>& OutputPinWidget,
 	UEdGraphPin* OutputPin,
@@ -337,19 +335,21 @@ void FConnectionDrawingPolicy::DetermineLinkGeometry(
 	/*out*/ FArrangedWidget*& EndWidgetGeometry
 	)
 {
-	StartWidgetGeometry = PinGeometries.Find(OutputPinWidget);
+	StartWidgetGeometry = PinGeometries->Find(OutputPinWidget);
 	
 	if (TSharedRef<SGraphPin>* pTargetWidget = PinToPinWidgetMap.Find(InputPin))
 	{
 		TSharedRef<SGraphPin> InputWidget = *pTargetWidget;
-		EndWidgetGeometry = PinGeometries.Find(InputWidget);
+		EndWidgetGeometry = PinGeometries->Find(InputWidget);
 	}
 }
 
-void FConnectionDrawingPolicy::Draw(TMap<TSharedRef<SWidget>, FArrangedWidget>& PinGeometries, FArrangedChildren& ArrangedNodes)
+void FConnectionDrawingPolicy::Draw(TMap<TSharedRef<SWidget>, FArrangedWidget>& InPinGeometries, FArrangedChildren& ArrangedNodes)
 {
+	PinGeometries = &InPinGeometries;
+
 	PinToPinWidgetMap.Empty();
-	for (TMap<TSharedRef<SWidget>, FArrangedWidget>::TIterator ConnectorIt(PinGeometries); ConnectorIt; ++ConnectorIt)
+	for (TMap<TSharedRef<SWidget>, FArrangedWidget>::TIterator ConnectorIt(InPinGeometries); ConnectorIt; ++ConnectorIt)
 	{
 		TSharedRef<SWidget> SomePinWidget = ConnectorIt.Key();
 		SGraphPin& PinWidget = static_cast<SGraphPin&>(SomePinWidget.Get());
@@ -357,7 +357,7 @@ void FConnectionDrawingPolicy::Draw(TMap<TSharedRef<SWidget>, FArrangedWidget>&
 		PinToPinWidgetMap.Add(PinWidget.GetPinObj(), StaticCastSharedRef<SGraphPin>(SomePinWidget));
 	}
 
-	for (TMap<TSharedRef<SWidget>, FArrangedWidget>::TIterator ConnectorIt(PinGeometries); ConnectorIt; ++ConnectorIt)
+	for (TMap<TSharedRef<SWidget>, FArrangedWidget>::TIterator ConnectorIt(InPinGeometries); ConnectorIt; ++ConnectorIt)
 	{
 		TSharedRef<SWidget> SomePinWidget = ConnectorIt.Key();
 		SGraphPin& PinWidget = static_cast<SGraphPin&>(SomePinWidget.Get());
@@ -372,7 +372,7 @@ void FConnectionDrawingPolicy::Draw(TMap<TSharedRef<SWidget>, FArrangedWidget>&
 
 				UEdGraphPin* TargetPin = ThePin->LinkedTo[LinkIndex];
 
-				DetermineLinkGeometry(PinGeometries, ArrangedNodes, SomePinWidget, ThePin, TargetPin, /*out*/ LinkStartWidgetGeometry, /*out*/ LinkEndWidgetGeometry);
+				DetermineLinkGeometry(ArrangedNodes, SomePinWidget, ThePin, TargetPin, /*out*/ LinkStartWidgetGeometry, /*out*/ LinkEndWidgetGeometry);
 
 				if (( LinkEndWidgetGeometry && LinkStartWidgetGeometry ) && !IsConnectionCulled( *LinkStartWidgetGeometry, *LinkEndWidgetGeometry ))
 				{
diff --git a/Engine/Source/Editor/GraphEditor/Private/StateMachineConnectionDrawingPolicy.cpp b/Engine/Source/Editor/GraphEditor/Private/StateMachineConnectionDrawingPolicy.cpp
--- a/Engine/Source/Editor/GraphEditor/Private/StateMachineConnectionDrawingPolicy.cpp
+++ b/Engine/Source/Editor/GraphEditor/Private/StateMachineConnectionDrawingPolicy.cpp
@@ -37,7 +37,6 @@ void FStateMachineConnectionDrawingPolicy::DetermineWiringStyle(UEdGraphPin* Out
 }
 
 void FStateMachineConnectionDrawingPolicy::DetermineLinkGeometry(
-	TMap<TSharedRef<SWidget>, FArrangedWidget>& PinGeometries,
 	FArrangedChildren& ArrangedNodes, 
 	TSharedRef<SWidget>& OutputPinWidget,
 	UEdGraphPin* OutputPin,
@@ -48,9 +47,7 @@ void FStateMachineConnectionDrawingPolicy::DetermineLinkGeometry(
 {
 	if (UAnimStateEntryNode* EntryNode = Cast<UAnimStateEntryNode>(OutputPin->GetOwningNode()))
 	{
-		//FConnectionDrawingPolicy::DetermineLinkGeometry(PinGeometries, ArrangedNodes, OutputPinWidget, OutputPin, InputPin, StartWidgetGeometry, EndWidgetGeometry);
-
-		StartWidgetGeometry = PinGeometries.Find(OutputPinWidget);
+		StartWidgetGeometry = PinGeometries->Find(OutputPinWidget);
 
 		UAnimStateNodeBase* State = CastChecked<UAnimStateNodeBase>(InputPin->GetOwningNode());
 		int32 StateIndex = NodeWidgetMap.FindChecked(State);
@@ -73,12 +70,12 @@ void FStateMachineConnectionDrawingPolicy::DetermineLinkGeometry(
 	}
 	else
 	{
-		StartWidgetGeometry = PinGeometries.Find(OutputPinWidget);
+		StartWidgetGeometry = PinGeometries->Find(OutputPinWidget);
 
 		if (TSharedRef<SGraphPin>* pTargetWidget = PinToPinWidgetMap.Find(InputPin))
 		{
 			TSharedRef<SGraphPin> InputWidget = *pTargetWidget;
-			EndWidgetGeometry = PinGeometries.Find(InputWidget);
+			EndWidgetGeometry = PinGeometries->Find(InputWidget);
 		}
 	}
 }
diff --git a/Engine/Source/Editor/GraphEditor/Private/StateMachineConnectionDrawingPolicy.h b/Engine/Source/Editor/GraphEditor/Private/StateMachineConnectionDrawingPolicy.h
--- a/Engine/Source/Editor/GraphEditor/Private/StateMachineConnectionDrawingPolicy.h
+++ b/Engine/Source/Editor/GraphEditor/Private/StateMachineConnectionDrawingPolicy.h
@@ -22,8 +22,6 @@ public:
 	virtual void DetermineWiringStyle(UEdGraphPin* OutputPin, UEdGraphPin* InputPin, /*inout*/ FConnectionParams& Params) override;
 	virtual void Draw(TMap<TSharedRef<SWidget>, FArrangedWidget>& PinGeometries, FArrangedChildren& ArrangedNodes) override;
 	virtual void DetermineLinkGeometry(
-		TMap<TSharedRef<SWidget>,
-		FArrangedWidget>& PinGeometries,
 		FArrangedChildren& ArrangedNodes, 
 		TSharedRef<SWidget>& OutputPinWidget,
 		UEdGraphPin* OutputPin,
diff --git a/Engine/Source/Editor/GraphEditor/Public/ConnectionDrawingPolicy.h b/Engine/Source/Editor/GraphEditor/Public/ConnectionDrawingPolicy.h
--- a/Engine/Source/Editor/GraphEditor/Public/ConnectionDrawingPolicy.h
+++ b/Engine/Source/Editor/GraphEditor/Public/ConnectionDrawingPolicy.h
@@ -75,6 +75,7 @@ protected:
 	FSlateWindowElementList& DrawElementsList;
 	TMap< UEdGraphPin*, TSharedRef<SGraphPin> > PinToPinWidgetMap;
 	TSet< UEdGraphPin* > HoveredPins;
+	TMap<TSharedRef<SWidget>, FArrangedWidget>* PinGeometries;
 	double LastHoverTimeEvent;
 public:
 	virtual ~FConnectionDrawingPolicy() {}
@@ -100,11 +101,9 @@ public:
 	// Give specific editor modes a chance to highlight this connection or darken non-interesting connections
 	virtual void DetermineWiringStyle(UEdGraphPin* OutputPin, UEdGraphPin* InputPin, /*inout*/ FConnectionParams& Params);
 
-	virtual void Draw(TMap<TSharedRef<SWidget>, FArrangedWidget>& PinGeometries, FArrangedChildren& ArrangedNodes);
+	virtual void Draw(TMap<TSharedRef<SWidget>, FArrangedWidget>& InPinGeometries, FArrangedChildren& ArrangedNodes);
 
 	virtual void DetermineLinkGeometry(
-		TMap<TSharedRef<SWidget>,
-		FArrangedWidget>& PinGeometries,
 		FArrangedChildren& ArrangedNodes, 
 		TSharedRef<SWidget>& OutputPinWidget,
 		UEdGraphPin* OutputPin,
