This adds some code to enable/disable add pin buttons determined by the owning graphnode's state.
#UE-9254 - Can add pins to nodes during PIE
#Branch UE4
#Proj GraphEditor

diff --git a/Engine/Source/Editor/GraphEditor/Private/SGraphNode.cpp b/Engine/Source/Editor/GraphEditor/Private/SGraphNode.cpp
--- a/Engine/Source/Editor/GraphEditor/Private/SGraphNode.cpp
+++ b/Engine/Source/Editor/GraphEditor/Private/SGraphNode.cpp
@@ -1363,6 +1363,7 @@ TSharedRef<SWidget> SGraphNode::AddPinButtonContent(FText PinText, FText PinTool
 	.ContentPadding(0.0f)
 	.ButtonStyle( FEditorStyle::Get(), "NoBorder" )
 	.OnClicked( this, &SGraphNode::OnAddPin )
+	.IsEnabled( this, &SGraphNode::IsNodeEditable )
 	.ToolTipText(PinTooltipText)
 	.ToolTip(Tooltip)
 	.Visibility(this, &SGraphNode::IsAddPinButtonVisible)
