SGraphPinObject will now pull the pin property's metadata and use any specified AllowedClasses for filtering.

Added UEdGraphNode::GetPinMetaData for requested pin metadata of a specific key to be returned. Overriden in UK2Node_StructOperation for returning the UProperty's metadata.

diff --git a/Engine/Source/Editor/BlueprintGraph/Classes/K2Node_StructOperation.h b/Engine/Source/Editor/BlueprintGraph/Classes/K2Node_StructOperation.h
--- a/Engine/Source/Editor/BlueprintGraph/Classes/K2Node_StructOperation.h
+++ b/Engine/Source/Editor/BlueprintGraph/Classes/K2Node_StructOperation.h
@@ -15,7 +15,7 @@ class UK2Node_StructOperation : public UK2Node_Variable
 	UScriptStruct* StructType;
 
 	// Begin UEdGraphNode interface
-	//virtual FLinearColor GetNodeTitleColor() const override;
+	virtual FString GetPinMetaData(FString InPinName, FName InKey) override;
 	// End UEdGraphNode interface
 
 	// UK2Node interface
diff --git a/Engine/Source/Editor/BlueprintGraph/Private/K2Node_StructOperation.cpp b/Engine/Source/Editor/BlueprintGraph/Private/K2Node_StructOperation.cpp
--- a/Engine/Source/Editor/BlueprintGraph/Private/K2Node_StructOperation.cpp
+++ b/Engine/Source/Editor/BlueprintGraph/Private/K2Node_StructOperation.cpp
@@ -67,4 +67,20 @@ bool UK2Node_StructOperation::DoRenamedPinsMatch(const UEdGraphPin* NewPin, cons
 		}
 	}
 	return bResult;
+}
+
+FString UK2Node_StructOperation::GetPinMetaData(FString InPinName, FName InKey)
+{
+	FString ReturnValue;
+
+	for (TFieldIterator<UProperty> It(StructType); It; ++It)
+	{
+		const UProperty* Property = *It;
+		if(Property && Property->GetName() == InPinName)
+		{
+			ReturnValue = Property->GetMetaData(InKey);
+			break;
+		}
+	}
+	return ReturnValue;
 }
\ No newline at end of file
diff --git a/Engine/Source/Editor/GraphEditor/Private/KismetPins/SGraphPinObject.cpp b/Engine/Source/Editor/GraphEditor/Private/KismetPins/SGraphPinObject.cpp
--- a/Engine/Source/Editor/GraphEditor/Private/KismetPins/SGraphPinObject.cpp
+++ b/Engine/Source/Editor/GraphEditor/Private/KismetPins/SGraphPinObject.cpp
@@ -185,6 +185,22 @@ TSharedRef<SWidget> SGraphPinObject::GenerateAssetPicker()
 	AssetPickerConfig.InitialAssetViewType = EAssetViewType::List;
 	AssetPickerConfig.bAllowDragging = false;
 
+	// Check with the node to see if there is any "AllowClasses" metadata for the pin
+	FString ClassFilterString = GraphPinObj->GetOwningNode()->GetPinMetaData(GraphPinObj->PinName, FName(TEXT("AllowedClasses")));
+	if( !ClassFilterString.IsEmpty() )
+	{
+		// Clear out the allowed class names and have the pin's metadata override.
+		AssetPickerConfig.Filter.ClassNames.Empty();
+
+		// Parse and add the classes from the metadata
+		TArray<FString> CustomClassFilterNames;
+		ClassFilterString.ParseIntoArray(&CustomClassFilterNames, TEXT(","), true);
+		for(auto It = CustomClassFilterNames.CreateConstIterator(); It; ++It)
+		{
+			AssetPickerConfig.Filter.ClassNames.Add(FName(**It));
+		}
+	}
+
 	return
 		SNew(SBox)
 		.HeightOverride(300)
diff --git a/Engine/Source/Runtime/Engine/Classes/EdGraph/EdGraphNode.h b/Engine/Source/Runtime/Engine/Classes/EdGraph/EdGraphNode.h
--- a/Engine/Source/Runtime/Engine/Classes/EdGraph/EdGraphNode.h
+++ b/Engine/Source/Runtime/Engine/Classes/EdGraph/EdGraphNode.h
@@ -420,6 +420,9 @@ public:
 	 * @param OutTaggedMetaData		Built array of tagged meta data for the node
 	 */
 	virtual void AddSearchMetaDataInfo(TArray<struct FSearchTagDataPair>& OutTaggedMetaData) const;
+
+	/** Return the requested metadata for the pin if there is any */
+	virtual FString GetPinMetaData(FString InPinName, FName InKey) { return FString(); }
 #endif // WITH_EDITOR
 
 };
