Fix for the call out arrow on ghost nodes displaying with the wrong color.
#Branch UE4
#proj GraphEditor

diff --git a/Engine/Source/Editor/GraphEditor/Private/SCommentBubble.cpp b/Engine/Source/Editor/GraphEditor/Private/SCommentBubble.cpp
--- a/Engine/Source/Editor/GraphEditor/Private/SCommentBubble.cpp
+++ b/Engine/Source/Editor/GraphEditor/Private/SCommentBubble.cpp
@@ -238,7 +238,7 @@ void SCommentBubble::UpdateBubble()
 				[
 					SNew(SImage)
 					.Image( CommentCalloutArrowBrush )
-					.ColorAndOpacity( ColorAndOpacity )
+					.ColorAndOpacity( this, &SCommentBubble::GetBubbleColor )
 				]
 			]
 		];
