Comment bubble colors now obey property set on the comment node.

diff --git a/Engine/Source/Editor/GraphEditor/Private/SCommentBubble.cpp b/Engine/Source/Editor/GraphEditor/Private/SCommentBubble.cpp
--- a/Engine/Source/Editor/GraphEditor/Private/SCommentBubble.cpp
+++ b/Engine/Source/Editor/GraphEditor/Private/SCommentBubble.cpp
@@ -77,7 +77,7 @@ void SCommentBubble::Tick( const FGeometry& AllottedGeometry, const double InCur
 
 	if( bTitleBarBubbleVisible || IsBubbleVisible() )
 	{
-		const FLinearColor BubbleColor = ColorAndOpacity.Get().GetSpecifiedColor() * SCommentBubbleDefs::LuminanceCoEff;
+		const FLinearColor BubbleColor = GetBubbleColor().GetSpecifiedColor() * SCommentBubbleDefs::LuminanceCoEff;
 		const float BubbleLuminance = BubbleColor.R + BubbleColor.G + BubbleColor.B;
 		ForegroundColor = BubbleLuminance < 0.5f ? SCommentBubbleDefs::DarkForegroundClr : SCommentBubbleDefs::LightForegroundClr;
 	}
@@ -337,7 +337,19 @@ FSlateColor SCommentBubble::GetToggleButtonColor() const
 
 FSlateColor SCommentBubble::GetBubbleColor() const
 {
-	FLinearColor ReturnColor = ColorAndOpacity.Get().GetSpecifiedColor();
+	bool bColorCommentBubble = true;
+
+	// For comment nodes, check if the color is enabled
+	if (UEdGraphNode_Comment* CommentBubble = Cast<UEdGraphNode_Comment>(GraphNode))
+	{
+		bColorCommentBubble = CommentBubble->bColorCommentBubble;
+	}
+
+	FLinearColor ReturnColor = FLinearColor::White;
+	if (bColorCommentBubble)
+	{
+		ReturnColor = ColorAndOpacity.Get().GetSpecifiedColor();
+	}
 
 	if(!GraphNode->bIsNodeEnabled)
 	{
