Fixed crash on opening SoundCue Editor

Could have got the main audio device when Preview Component is null but not much point when it will kick out soon after.

diff --git a/Engine/Source/Editor/GraphEditor/Private/SoundCueGraphConnectionDrawingPolicy.cpp b/Engine/Source/Editor/GraphEditor/Private/SoundCueGraphConnectionDrawingPolicy.cpp
--- a/Engine/Source/Editor/GraphEditor/Private/SoundCueGraphConnectionDrawingPolicy.cpp
+++ b/Engine/Source/Editor/GraphEditor/Private/SoundCueGraphConnectionDrawingPolicy.cpp
@@ -37,7 +37,7 @@ void FSoundCueGraphConnectionDrawingPolicy::Draw(TMap<TSharedRef<SWidget>, FArra
 void FSoundCueGraphConnectionDrawingPolicy::BuildAudioFlowRoadmap()
 {
 	UAudioComponent* PreviewAudioComponent = GEditor->GetPreviewAudioComponent();
-	FAudioDevice* AudioDevice = PreviewAudioComponent->GetAudioDevice();
+	FAudioDevice* AudioDevice = PreviewAudioComponent ? PreviewAudioComponent->GetAudioDevice() : nullptr;
 
 	if (AudioDevice)
 	{
