Added the ability to pin re-route node comment bubbles to resolve conflicting user requests.
#UE-7910 - Pin-able Reroute Node Comments
#Branch UE4
#Proj GraphEditor, EditorStyle

diff --git a/Engine/Source/Editor/EditorStyle/Private/SlateEditorStyle.cpp b/Engine/Source/Editor/EditorStyle/Private/SlateEditorStyle.cpp
--- a/Engine/Source/Editor/EditorStyle/Private/SlateEditorStyle.cpp
+++ b/Engine/Source/Editor/EditorStyle/Private/SlateEditorStyle.cpp
@@ -3433,9 +3433,9 @@ void FSlateEditorStyle::FStyle::SetupGraphEditorStyles()
 			.SetUncheckedImage( IMAGE_BRUSH( "Icons/icon_Blueprint_CommentBubbleOn_16x", FVector2D(10,10), FLinearColor(1.f, 1.f, 1.f, 0.5f)))
 			.SetUncheckedHoveredImage( IMAGE_BRUSH( "Icons/icon_Blueprint_CommentBubbleOn_16x", FVector2D(10,10), FLinearColor(1.f, 1.f, 1.f, 0.9f)))
 			.SetUncheckedPressedImage( IMAGE_BRUSH( "Icons/icon_Blueprint_CommentBubbleOn_16x", FVector2D(10,10), FLinearColor(1.f, 1.f, 1.f, 1.f)))
-			.SetCheckedImage( IMAGE_BRUSH( "Icons/icon_Blueprint_CommentBubbleOn_16x", FVector2D(10,10), FLinearColor(1.f, 1.f, 1.f, 0.5f)))
+			.SetCheckedImage( IMAGE_BRUSH( "Icons/icon_Blueprint_CommentBubbleOn_16x", FVector2D(10,10), FLinearColor(1.f, 1.f, 1.f, 0.7f)))
 			.SetCheckedHoveredImage( IMAGE_BRUSH( "Icons/icon_Blueprint_CommentBubbleOn_16x", FVector2D(10,10), FLinearColor(1.f, 1.f, 1.f, 1.f)))
-			.SetCheckedPressedImage( IMAGE_BRUSH( "Icons/icon_Blueprint_CommentBubbleOn_16x", FVector2D(10,10), FLinearColor(1.f, 1.f, 1.f, 0.4f)));
+			.SetCheckedPressedImage( IMAGE_BRUSH( "Icons/icon_Blueprint_CommentBubbleOn_16x", FVector2D(10,10), FLinearColor(1.f, 1.f, 1.f, 0.6f)));
 		Set( "CommentBubbleButton", CommentBubbleButton );
 
 		const FCheckBoxStyle CommentBubblePin = FCheckBoxStyle()
diff --git a/Engine/Source/Editor/GraphEditor/Private/SCommentBubble.cpp b/Engine/Source/Editor/GraphEditor/Private/SCommentBubble.cpp
--- a/Engine/Source/Editor/GraphEditor/Private/SCommentBubble.cpp
+++ b/Engine/Source/Editor/GraphEditor/Private/SCommentBubble.cpp
@@ -46,6 +46,11 @@ void SCommentBubble::Construct( const FArguments& InArgs )
 	IsGraphNodeHovered		= InArgs._IsGraphNodeHovered;
 	HintText				= InArgs._HintText.IsSet() ? InArgs._HintText : NSLOCTEXT( "CommentBubble", "EditCommentHint", "Click to edit" );
 	OpacityValue			= SCommentBubbleDefs::FadeDelay;
+	// Create Default toggle and pin handlers
+	CommentBubbleToggleDelegate	= InArgs._OnCommentBubbleToggle.IsBound() ? InArgs._OnCommentBubbleToggle : 
+																			FOnCheckStateChanged::CreateSP( this, &SCommentBubble::OnCommentBubbleToggle );
+	CommentBubblePinnedDelegate = InArgs._OnCommentBubblePinned.IsBound() ? InArgs._OnCommentBubblePinned : 
+																			FOnCheckStateChanged::CreateSP( this, &SCommentBubble::OnPinStateToggle );
 	// Ensue this value is set to something sensible
 	ForegroundColor = SCommentBubbleDefs::LightForegroundClr;
 
@@ -144,7 +149,7 @@ void SCommentBubble::UpdateBubble()
 						SNew( SCheckBox )
 						.Style( &FEditorStyle::Get().GetWidgetStyle<FCheckBoxStyle>( "CommentBubblePin" ))
 						.IsChecked( GraphNode->bCommentBubblePinned ? ECheckBoxState::Checked : ECheckBoxState::Unchecked )
-						.OnCheckStateChanged( this, &SCommentBubble::OnPinStateToggle )
+						.OnCheckStateChanged( CommentBubblePinnedDelegate )
 						.ToolTipText( this, &SCommentBubble::GetScaleButtonTooltip )
 						.Cursor( EMouseCursor::Default )
 						.ForegroundColor( this, &SCommentBubble::GetForegroundColor )
@@ -157,8 +162,8 @@ void SCommentBubble::UpdateBubble()
 				[
 					SNew( SCheckBox )
 					.Style( &FEditorStyle::Get().GetWidgetStyle< FCheckBoxStyle >( "CommentBubbleButton" ))
-					.IsChecked( GraphNode->bCommentBubbleVisible ? ECheckBoxState::Checked : ECheckBoxState::Unchecked )
-					.OnCheckStateChanged( this, &SCommentBubble::OnCommentBubbleToggle )
+					.IsChecked( ECheckBoxState::Checked )
+					.OnCheckStateChanged( CommentBubbleToggleDelegate )
 					.ToolTipText( NSLOCTEXT( "CommentBubble", "ToggleCommentTooltip", "Toggle Comment Bubble" ))
 					.Cursor( EMouseCursor::Default )
 					.ForegroundColor( FLinearColor::White )
@@ -174,8 +179,8 @@ void SCommentBubble::UpdateBubble()
 				[
 					SNew( SCheckBox )
 					.Style( &FEditorStyle::Get().GetWidgetStyle< FCheckBoxStyle >( "CommentBubbleButton" ))
-					.IsChecked( GraphNode->bCommentBubbleVisible ? ECheckBoxState::Checked : ECheckBoxState::Unchecked )
-					.OnCheckStateChanged( this, &SCommentBubble::OnCommentBubbleToggle )
+					.IsChecked( ECheckBoxState::Checked )
+					.OnCheckStateChanged( CommentBubbleToggleDelegate )
 					.ToolTipText( NSLOCTEXT( "CommentBubble", "ToggleCommentTooltip", "Toggle Comment Bubble" ))
 					.Cursor( EMouseCursor::Default )
 					.ForegroundColor( FLinearColor::White )
@@ -262,7 +267,7 @@ void SCommentBubble::UpdateBubble()
 				SNew( SCheckBox )
 				.Style( &FEditorStyle::Get().GetWidgetStyle< FCheckBoxStyle >( "CommentTitleButton" ))
 				.IsChecked( ECheckBoxState::Unchecked )
-				.OnCheckStateChanged( this, &SCommentBubble::OnCommentBubbleToggle )
+				.OnCheckStateChanged( CommentBubbleToggleDelegate )
 				.ToolTipText( NSLOCTEXT( "CommentBubble", "ToggleCommentTooltip", "Toggle Comment Bubble" ))
 				.Cursor( EMouseCursor::Default )
 				.ForegroundColor( this, &SCommentBubble::GetToggleButtonColor )
@@ -280,8 +285,8 @@ FVector2D SCommentBubble::GetOffset() const
 {
 	if( GraphNode->bCommentBubbleVisible || OpacityValue > 0.f )
 	{
-	return FVector2D( 0.f, -GetDesiredSize().Y );
-}
+		return FVector2D( 0.f, -GetDesiredSize().Y );
+	}
 	return FVector2D::ZeroVector;
 }
 
diff --git a/Engine/Source/Editor/GraphEditor/Private/SGraphNodeKnot.cpp b/Engine/Source/Editor/GraphEditor/Private/SGraphNodeKnot.cpp
--- a/Engine/Source/Editor/GraphEditor/Private/SGraphNodeKnot.cpp
+++ b/Engine/Source/Editor/GraphEditor/Private/SGraphNodeKnot.cpp
@@ -254,6 +254,7 @@ FReply SGraphPinKnot::OnPinMouseDown(const FGeometry& SenderGeometry, const FPoi
 void SGraphNodeKnot::Construct(const FArguments& InArgs, UK2Node_Knot* InKnot)
 {
 	SGraphNodeDefault::Construct(SGraphNodeDefault::FArguments().GraphNodeObj(InKnot));
+	bCommentPinnedState = InKnot->bCommentBubbleVisible;
 }
 
 void SGraphNodeKnot::UpdateGraphNode()
@@ -325,9 +326,11 @@ void SGraphNodeKnot::UpdateGraphNode()
 	SAssignNew( CommentBubble, SCommentBubble )
 	.GraphNode( GraphNode )
 	.Text( this, &SGraphNode::GetNodeComment )
+	.EnableBubbleCtrls( true )
 	.ColorAndOpacity( FLinearColor::White )
 	.GraphLOD( this, &SGraphNode::GetCurrentLOD )
-	.IsGraphNodeHovered( this, &SGraphNode::IsHovered );
+	.IsGraphNodeHovered( this, &SGraphNode::IsHovered )
+	.OnCommentBubbleToggle( this, &SGraphNodeKnot::OnCommentBubbleToggle );
 
 	GetOrAddSlot( ENodeZone::TopCenter )
 	.SlotOffset( TAttribute<FVector2D>( this, &SGraphNodeKnot::GetCommentOffset ))
@@ -390,7 +393,7 @@ FVector2D SGraphNodeKnot::GetCommentOffset() const
 
 void SGraphNodeKnot::OnMouseEnter( const FGeometry& MyGeometry, const FPointerEvent& MouseEvent )
 {
-	if( !GraphNode->NodeComment.IsEmpty() )
+	if( !bCommentPinnedState && !GraphNode->NodeComment.IsEmpty() )
 	{
 		GraphNode->bCommentBubbleVisible = true;
 		CommentBubble->UpdateBubble();
@@ -399,9 +402,19 @@ void SGraphNodeKnot::OnMouseEnter( const FGeometry& MyGeometry, const FPointerEv
 
 void SGraphNodeKnot::OnMouseLeave( const FPointerEvent& MouseEvent )
 {
-	if( !GraphNode->NodeComment.IsEmpty() )
+	if( !bCommentPinnedState && !GraphNode->NodeComment.IsEmpty() )
 	{
 		GraphNode->bCommentBubbleVisible = false;
 		CommentBubble->UpdateBubble();
 	}
 }
+
+void SGraphNodeKnot::OnCommentBubbleToggle( ECheckBoxState State )
+{
+	if( CommentBubble.IsValid() )
+	{
+		bCommentPinnedState = !bCommentPinnedState;
+		GraphNode->bCommentBubbleVisible = bCommentPinnedState;
+		CommentBubble->UpdateBubble();
+	}
+}
\ No newline at end of file
diff --git a/Engine/Source/Editor/GraphEditor/Private/SGraphNodeKnot.h b/Engine/Source/Editor/GraphEditor/Private/SGraphNodeKnot.h
--- a/Engine/Source/Editor/GraphEditor/Private/SGraphNodeKnot.h
+++ b/Engine/Source/Editor/GraphEditor/Private/SGraphNodeKnot.h
@@ -31,8 +31,13 @@ public:
 	/** Returns Offset to center comment on the node's only pin */
 	FVector2D GetCommentOffset() const;
 
+	/** Called when the comment bubble visibility/pinned state is changed */
+	void OnCommentBubbleToggle( ECheckBoxState State );
+
 protected:
 
+	/** The pinned state for comment bubble */
+	bool bCommentPinnedState;
 	/** SharedPtr to comment bubble */
 	TSharedPtr<SCommentBubble> CommentBubble;
 
diff --git a/Engine/Source/Editor/GraphEditor/Public/SCommentBubble.h b/Engine/Source/Editor/GraphEditor/Public/SCommentBubble.h
--- a/Engine/Source/Editor/GraphEditor/Public/SCommentBubble.h
+++ b/Engine/Source/Editor/GraphEditor/Public/SCommentBubble.h
@@ -18,6 +18,12 @@ class GRAPHEDITOR_API SCommentBubble : public SCompoundWidget
 		/** the GraphNode this bubble should interact with */
 		SLATE_ARGUMENT( UEdGraphNode*, GraphNode )
 
+		/** Called when the comment bubble toggle button is pressed */
+		SLATE_EVENT( FOnCheckStateChanged, OnCommentBubbleToggle )
+
+		/** Called when the comment bubble pin button is pressed */
+		SLATE_EVENT( FOnCheckStateChanged, OnCommentBubblePinned )
+
 		/** The comment text for the bubble */
 		SLATE_ATTRIBUTE( FString, Text )
 
@@ -121,15 +127,15 @@ protected:
 	TAttribute<FSlateColor> ColorAndOpacity;
 	/** Attribute to query node comment */
 	TAttribute<FString> CommentAttribute;
-	/** Attribute to query node pinned state */
-	TAttribute<bool> PinnedState;
-	/** Attribute to query node visibility state */
-	TAttribute<bool> VisibilityState;
 	/** Attribute to query current LOD */
 	TAttribute<EGraphRenderingLOD::Type> GraphLOD;
 	/** Hint Text */
 	TAttribute<FText> HintText;
 
+	/** Optional delegate to call when the comment is toggled */
+	FOnCheckStateChanged CommentBubbleToggleDelegate;
+	/** Optional delegate to call when the comment is toggled */
+	FOnCheckStateChanged CommentBubblePinnedDelegate;
 	/** Delegate to determine if the graph node is currently hovered */
 	FIsGraphNodeHovered IsGraphNodeHovered;
 
