UE-12073 [CrashReport] Disconnecting/changing blueprint nodes causes crash.

diff --git a/Engine/Source/Editor/GraphEditor/Private/SGraphPin.cpp b/Engine/Source/Editor/GraphEditor/Private/SGraphPin.cpp
--- a/Engine/Source/Editor/GraphEditor/Private/SGraphPin.cpp
+++ b/Engine/Source/Editor/GraphEditor/Private/SGraphPin.cpp
@@ -76,6 +76,7 @@ SGraphPin::SGraphPin()
 	, bIsMovingLinks(false)
 	, PinColorModifier(FLinearColor::White)
 	, CachedNodeOffset(FVector2D::ZeroVector)
+	, GraphPinObj(nullptr)
 {
 	IsEditable = true;
 
@@ -965,8 +966,7 @@ FText SGraphPin::GetTooltip() const
 {
 	FText HoverText = FText::GetEmpty();
 
-	check(GraphPinObj != nullptr);
-	UEdGraphNode* GraphNode = GraphPinObj->GetOwningNodeUnchecked();
+	UEdGraphNode* GraphNode = GraphPinObj ? GraphPinObj->GetOwningNodeUnchecked() : nullptr;
 	if (GraphNode != nullptr)
 	{
 		FString HoverStr;
