Graph Editor: Replace an inclusion guard with #pragma once

diff --git a/Engine/Source/Editor/GraphEditor/Private/GraphEditorCommon.h b/Engine/Source/Editor/GraphEditor/Private/GraphEditorCommon.h
--- a/Engine/Source/Editor/GraphEditor/Private/GraphEditorCommon.h
+++ b/Engine/Source/Editor/GraphEditor/Private/GraphEditorCommon.h
@@ -1,8 +1,6 @@
 // Copyright 1998-2015 Epic Games, Inc. All Rights Reserved.
 
-
-#ifndef __GraphEditorCommon_h__
-#define __GraphEditorCommon_h__
+#pragma once
 
 #include "UnrealEd.h"
 #include "GraphEditor.h"
@@ -23,4 +21,3 @@
 
 #include "AnimGraphDefinitions.h"
 
-#endif//__GraphEditorCommon_h__
