CIS fix

diff --git a/Engine/Source/Editor/GraphEditor/Private/SGraphPin.cpp b/Engine/Source/Editor/GraphEditor/Private/SGraphPin.cpp
--- a/Engine/Source/Editor/GraphEditor/Private/SGraphPin.cpp
+++ b/Engine/Source/Editor/GraphEditor/Private/SGraphPin.cpp
@@ -72,10 +72,10 @@ struct FKnotNetCollector
 
 
 SGraphPin::SGraphPin()
-	: bShowLabel(true)
+	: GraphPinObj(nullptr)
+	, bShowLabel(true)
 	, bIsMovingLinks(false)
 	, PinColorModifier(FLinearColor::White)
-	, GraphPinObj(nullptr)
 	, CachedNodeOffset(FVector2D::ZeroVector)
 {
 	IsEditable = true;
