Graph Editor: Add experimental support for hovering and interacting with splines
- Currently gated by a setting in Editor Preferences..Graph Editors
- Alt+LMB Click on the spline to break the connection
- Ctrl+LMB Click on the spline to move the connection to a different pin
- Hover over the spline to highlight it

diff --git a/Engine/Source/Editor/AIGraph/Private/AIGraphConnectionDrawingPolicy.cpp b/Engine/Source/Editor/AIGraph/Private/AIGraphConnectionDrawingPolicy.cpp
--- a/Engine/Source/Editor/AIGraph/Private/AIGraphConnectionDrawingPolicy.cpp
+++ b/Engine/Source/Editor/AIGraph/Private/AIGraphConnectionDrawingPolicy.cpp
@@ -13,6 +13,8 @@ FAIGraphConnectionDrawingPolicy::FAIGraphConnectionDrawingPolicy(int32 InBackLay
 
 void FAIGraphConnectionDrawingPolicy::DetermineWiringStyle(UEdGraphPin* OutputPin, UEdGraphPin* InputPin, /*inout*/ FConnectionParams& Params)
 {
+	Params.AssociatedPin1 = OutputPin;
+	Params.AssociatedPin2 = InputPin;
 	Params.WireThickness = 1.5f;
 
 	const bool bDeemphasizeUnhoveredPins = HoveredPins.Num() > 0;
@@ -40,7 +42,7 @@ void FAIGraphConnectionDrawingPolicy::Draw(TMap<TSharedRef<SWidget>, FArrangedWi
 void FAIGraphConnectionDrawingPolicy::DrawPreviewConnector(const FGeometry& PinGeometry, const FVector2D& StartPoint, const FVector2D& EndPoint, UEdGraphPin* Pin)
 {
 	FConnectionParams Params;
-	DetermineWiringStyle(Pin, NULL, /*inout*/ Params);
+	DetermineWiringStyle(Pin, nullptr, /*inout*/ Params);
 
 	if (Pin->Direction == EEdGraphPinDirection::EGPD_Output)
 	{
diff --git a/Engine/Source/Editor/BehaviorTreeEditor/Private/BehaviorTreeConnectionDrawingPolicy.cpp b/Engine/Source/Editor/BehaviorTreeEditor/Private/BehaviorTreeConnectionDrawingPolicy.cpp
--- a/Engine/Source/Editor/BehaviorTreeEditor/Private/BehaviorTreeConnectionDrawingPolicy.cpp
+++ b/Engine/Source/Editor/BehaviorTreeEditor/Private/BehaviorTreeConnectionDrawingPolicy.cpp
@@ -13,6 +13,8 @@ FBehaviorTreeConnectionDrawingPolicy::FBehaviorTreeConnectionDrawingPolicy(int32
 
 void FBehaviorTreeConnectionDrawingPolicy::DetermineWiringStyle(UEdGraphPin* OutputPin, UEdGraphPin* InputPin, /*inout*/ FConnectionParams& Params)
 {
+	Params.AssociatedPin1 = OutputPin;
+	Params.AssociatedPin2 = InputPin;
 	Params.WireThickness = 1.5f;
 
 	Params.WireColor = BehaviorTreeColors::Connection::Default;
diff --git a/Engine/Source/Editor/GraphEditor/Private/BlueprintConnectionDrawingPolicy.cpp b/Engine/Source/Editor/GraphEditor/Private/BlueprintConnectionDrawingPolicy.cpp
--- a/Engine/Source/Editor/GraphEditor/Private/BlueprintConnectionDrawingPolicy.cpp
+++ b/Engine/Source/Editor/GraphEditor/Private/BlueprintConnectionDrawingPolicy.cpp
@@ -422,6 +422,9 @@ bool FKismetConnectionDrawingPolicy::ShouldChangeTangentForKnot(UK2Node_Knot* Kn
 // Give specific editor modes a chance to highlight this connection or darken non-interesting connections
 void FKismetConnectionDrawingPolicy::DetermineWiringStyle(UEdGraphPin* OutputPin, UEdGraphPin* InputPin, /*inout*/ FConnectionParams& Params)
 {
+	Params.AssociatedPin1 = OutputPin;
+	Params.AssociatedPin2 = InputPin;
+
 	// Get the schema and grab the default color from it
 	check(OutputPin);
 	check(GraphObj);
@@ -523,13 +526,13 @@ void FKismetConnectionDrawingPolicy::SetIncompatiblePinDrawState(const TSharedPt
 {
 	ResetIncompatiblePinDrawState(VisiblePins);
 
-	for(auto VisiblePinIterator = VisiblePins.CreateConstIterator(); VisiblePinIterator; ++VisiblePinIterator)
+	for (auto VisiblePinIterator = VisiblePins.CreateConstIterator(); VisiblePinIterator; ++VisiblePinIterator)
 	{
 		TSharedPtr<SGraphPin> CheckPin = StaticCastSharedRef<SGraphPin>(*VisiblePinIterator);
-		if(CheckPin != StartPin)
+		if (CheckPin != StartPin)
 		{
 			const FPinConnectionResponse Response = StartPin->GetPinObj()->GetSchema()->CanCreateConnection(StartPin->GetPinObj(), CheckPin->GetPinObj());
-			if(Response.Response == CONNECT_RESPONSE_DISALLOW)
+			if (Response.Response == CONNECT_RESPONSE_DISALLOW)
 			{
 				CheckPin->SetPinColorModifier(FLinearColor(0.25f, 0.25f, 0.25f, 0.5f));
 			}
@@ -539,7 +542,7 @@ void FKismetConnectionDrawingPolicy::SetIncompatiblePinDrawState(const TSharedPt
 
 void FKismetConnectionDrawingPolicy::ResetIncompatiblePinDrawState(const TSet< TSharedRef<SWidget> >& VisiblePins)
 {
-	for(auto VisiblePinIterator = VisiblePins.CreateConstIterator(); VisiblePinIterator; ++VisiblePinIterator)
+	for (auto VisiblePinIterator = VisiblePins.CreateConstIterator(); VisiblePinIterator; ++VisiblePinIterator)
 	{
 		TSharedPtr<SGraphPin> VisiblePin = StaticCastSharedRef<SGraphPin>(*VisiblePinIterator);
 		VisiblePin->SetPinColorModifier(FLinearColor::White);
diff --git a/Engine/Source/Editor/GraphEditor/Private/ConnectionDrawingPolicy.cpp b/Engine/Source/Editor/GraphEditor/Private/ConnectionDrawingPolicy.cpp
--- a/Engine/Source/Editor/GraphEditor/Private/ConnectionDrawingPolicy.cpp
+++ b/Engine/Source/Editor/GraphEditor/Private/ConnectionDrawingPolicy.cpp
@@ -86,10 +86,11 @@ FConnectionDrawingPolicy::FConnectionDrawingPolicy(int32 InBackLayerID, int32 In
 	, ZoomFactor(InZoomFactor)
 	, ClippingRect(InClippingRect)
 	, DrawElementsList(InDrawElements)
+	, LocalMousePosition(0.0f, 0.0f)
 {
 	ArrowImage = FEditorStyle::GetBrush( TEXT("Graph.Arrow") );
 	ArrowRadius = ArrowImage->ImageSize * ZoomFactor * 0.5f;
-	MidpointImage = NULL;
+	MidpointImage = nullptr;
 	MidpointRadius = FVector2D::ZeroVector;
 	HoverDeemphasisDarkFraction = 0.8f;
 
@@ -167,6 +168,11 @@ void FConnectionDrawingPolicy::SetHoveredPins(const TSet< TWeakObjectPtr<UEdGrap
 	}
 }
 
+void FConnectionDrawingPolicy::SetMousePosition(const FVector2D& InMousePos)
+{
+	LocalMousePosition = InMousePos;
+}
+
 void FConnectionDrawingPolicy::SetMarkedPin(TWeakPtr<SGraphPin> InMarkedPin)
 {
 	if (InMarkedPin.IsValid())
@@ -186,10 +192,10 @@ void FConnectionDrawingPolicy::SetMarkedPin(TWeakPtr<SGraphPin> InMarkedPin)
 /** Util to make a 'distance->alpha' table and also return spline length */
 float FConnectionDrawingPolicy::MakeSplineReparamTable(const FVector2D& P0, const FVector2D& P0Tangent, const FVector2D& P1, const FVector2D& P1Tangent, FInterpCurve<float>& OutReparamTable)
 {
-	OutReparamTable.Reset();
-
 	const int32 NumSteps = 10; // TODO: Make this adaptive...
 
+	OutReparamTable.Points.Empty(NumSteps);
+
 	// Find range of input
 	float Param = 0.f;
 	const float MaxInput = 1.f;
@@ -231,6 +237,89 @@ void FConnectionDrawingPolicy::DrawConnection(int32 LayerId, const FVector2D& St
 	const FVector2D P0Tangent = (Params.StartDirection == EGPD_Output) ? SplineTangent : -SplineTangent;
 	const FVector2D P1Tangent = (Params.EndDirection == EGPD_Input) ? SplineTangent : -SplineTangent;
 
+	if (Settings->bTreatSplinesLikePinsEXPERIMENTAL)
+	{
+		// Distance to consider as an overlap
+		const float QueryDistanceTriggerThresholdSquared = FMath::Square(Settings->SplineHoverToleranceEXPERIMENTAL + Params.WireThickness * 0.5f);
+
+		// Distance to pass the bounding box cull test (may want to expand this later on if we want to do 'closest pin' actions that don't require an exact hit)
+		const float QueryDistanceToBoundingBoxSquared = QueryDistanceTriggerThresholdSquared;
+
+		bool bCloseToSpline = false;
+		{
+			// The curve will include the endpoints but can extend out of a tight bounds because of the tangents
+			// P0Tangent coefficient maximizes to 4/27 at a=1/3, and P1Tangent minimizes to -4/27 at a=2/3.
+			const float MaximumTangentContribution = 4.0f / 27.0f;
+			FBox2D Bounds(ForceInit);
+
+			Bounds += FVector2D(P0);
+			Bounds += FVector2D(P0 + MaximumTangentContribution * P0Tangent);
+			Bounds += FVector2D(P1);
+			Bounds += FVector2D(P1 - MaximumTangentContribution * P1Tangent);
+
+			bCloseToSpline = Bounds.ComputeSquaredDistanceToPoint(LocalMousePosition) < QueryDistanceToBoundingBoxSquared;
+
+			// Draw the bounding box for debugging
+#if 0
+#define DrawSpaceLine(Point1, Point2, DebugWireColor) {const FVector2D FakeTangent = (Point2 - Point1).GetSafeNormal(); FSlateDrawElement::MakeDrawSpaceSpline(DrawElementsList, LayerId, Point1, FakeTangent, Point2, FakeTangent, ClippingRect, 1.0f, ESlateDrawEffect::None, DebugWireColor); }
+
+			if (bCloseToSpline)
+			{
+				const FLinearColor BoundsWireColor = bCloseToSpline ? FLinearColor::Green : FLinearColor::White;
+
+				FVector2D TL = Bounds.Min;
+				FVector2D BR = Bounds.Max;
+				FVector2D TR = FVector2D(Bounds.Max.X, Bounds.Min.Y);
+				FVector2D BL = FVector2D(Bounds.Min.X, Bounds.Max.Y);
+
+				DrawSpaceLine(TL, TR, BoundsWireColor);
+				DrawSpaceLine(TR, BR, BoundsWireColor);
+				DrawSpaceLine(BR, BL, BoundsWireColor);
+				DrawSpaceLine(BL, TL, BoundsWireColor);
+			}
+#endif
+		}
+
+		if (bCloseToSpline)
+		{
+			// Find the closest approach to the spline
+			FVector2D ClosestPoint(ForceInit);
+			float ClosestDistanceSquared = FLT_MAX;
+
+			const int32 NumStepsToTest = 16;
+			float TestAlpha = 0.0f;
+			const float StepInterval = 1.0f / (float)NumStepsToTest;
+			FVector2D Point1 = FMath::CubicInterp(P0, P0Tangent, P1, P1Tangent, 0.0f);
+			for (float TestAlpha = 0.0f; TestAlpha < 1.0f; TestAlpha += StepInterval)
+			{
+				const FVector2D Point2 = FMath::CubicInterp(P0, P0Tangent, P1, P1Tangent, TestAlpha + StepInterval);
+
+				const FVector2D ClosestPointToSegment = FMath::ClosestPointOnSegment2D(LocalMousePosition, Point1, Point2);
+				const float DistanceSquared = (LocalMousePosition - ClosestPointToSegment).SizeSquared();
+
+				if (DistanceSquared < ClosestDistanceSquared)
+				{
+					ClosestDistanceSquared = DistanceSquared;
+					ClosestPoint = ClosestPointToSegment;
+				}
+
+				Point1 = Point2;
+			}
+
+			// Record the overlap
+			if (ClosestDistanceSquared < QueryDistanceTriggerThresholdSquared)
+			{
+				if (ClosestDistanceSquared < SplineOverlapResult.DistanceSquared)
+				{
+					const float SquaredDistToPin1 = (Params.AssociatedPin1 != nullptr) ? (P0 - ClosestPoint).SizeSquared() : FLT_MAX;
+					const float SquaredDistToPin2 = (Params.AssociatedPin2 != nullptr) ? (P1 - ClosestPoint).SizeSquared() : FLT_MAX;
+
+					SplineOverlapResult = FGraphSplineOverlapResult(Params.AssociatedPin1, Params.AssociatedPin2, ClosestDistanceSquared, SquaredDistToPin1, SquaredDistToPin2);
+				}
+			}
+		}
+	}
+
 	// Draw the spline itself
 	FSlateDrawElement::MakeDrawSpaceSpline(
 		DrawElementsList,
@@ -243,11 +332,11 @@ void FConnectionDrawingPolicy::DrawConnection(int32 LayerId, const FVector2D& St
 		Params.WireColor
 	);
 
-	if (Params.bDrawBubbles || (MidpointImage != NULL))
+	if (Params.bDrawBubbles || (MidpointImage != nullptr))
 	{
 		// This table maps distance along curve to alpha
 		FInterpCurve<float> SplineReparamTable;
-		float SplineLength = MakeSplineReparamTable(P0, P0Tangent, P1, P1Tangent, SplineReparamTable);
+		const float SplineLength = MakeSplineReparamTable(P0, P0Tangent, P1, P1Tangent, SplineReparamTable);
 
 		// Draw bubbles on the spline
 		if (Params.bDrawBubbles)
@@ -282,7 +371,7 @@ void FConnectionDrawingPolicy::DrawConnection(int32 LayerId, const FVector2D& St
 		}
 
 		// Draw the midpoint image
-		if (MidpointImage != NULL)
+		if (MidpointImage != nullptr)
 		{
 			// Determine the spline position for the midpoint
 			const float MidpointAlpha = SplineReparamTable.Eval(SplineLength * 0.5f, 0.f);
@@ -324,6 +413,8 @@ void FConnectionDrawingPolicy::DrawPreviewConnector(const FGeometry& PinGeometry
 
 void FConnectionDrawingPolicy::DetermineWiringStyle(UEdGraphPin* OutputPin, UEdGraphPin* InputPin, /*inout*/ FConnectionParams& Params)
 {
+	Params.AssociatedPin1 = OutputPin;
+	Params.AssociatedPin2 = InputPin;
 }
 
 void FConnectionDrawingPolicy::DetermineLinkGeometry(
@@ -367,8 +458,8 @@ void FConnectionDrawingPolicy::Draw(TMap<TSharedRef<SWidget>, FArrangedWidget>&
 		{
 			for (int32 LinkIndex=0; LinkIndex < ThePin->LinkedTo.Num(); ++LinkIndex)
 			{
-				FArrangedWidget* LinkStartWidgetGeometry = NULL;
-				FArrangedWidget* LinkEndWidgetGeometry = NULL;
+				FArrangedWidget* LinkStartWidgetGeometry = nullptr;
+				FArrangedWidget* LinkEndWidgetGeometry = nullptr;
 
 				UEdGraphPin* TargetPin = ThePin->LinkedTo[LinkIndex];
 
@@ -416,7 +507,7 @@ void FConnectionDrawingPolicy::ApplyHoverDeemphasis(UEdGraphPin* OutputPin, UEdG
 
 	const bool bContainsBoth = HoveredPins.Contains(InputPin) && HoveredPins.Contains(OutputPin);
 	const bool bContainsOutput = HoveredPins.Contains(OutputPin);
-	const bool bEmphasize = bContainsBoth || (bContainsOutput && (InputPin == NULL));
+	const bool bEmphasize = bContainsBoth || (bContainsOutput && (InputPin == nullptr));
 	if (bEmphasize)
 	{
 		Thickness = FMath::Lerp(Thickness, Thickness * ((Thickness < 3.0f) ? 5.0f : 3.0f), TimeFraction);
@@ -427,3 +518,38 @@ void FConnectionDrawingPolicy::ApplyHoverDeemphasis(UEdGraphPin* OutputPin, UEdG
 		WireColor = FMath::Lerp<FLinearColor>(WireColor, DarkenedColor, HoverDeemphasisDarkFraction * TimeFraction);
 	}
 }
+
+//////////////////////////////////////////////////////////////////////////
+// FGraphSplineOverlapResult
+
+void FGraphSplineOverlapResult::ComputeBestPin()
+{
+	if (Pin1 == nullptr)
+	{
+		BestPin = Pin2;
+	}
+	else if (Pin2 == nullptr)
+	{
+		BestPin = Pin1;
+	}
+	else
+	{
+		// Both are valid, now see if one of the pins has one connection while the other has more than one
+		const int32 LinksTo1 = Pin1->LinkedTo.Num();
+		const int32 LinksTo2 = Pin2->LinkedTo.Num();
+
+		if ((LinksTo1 > 1) && (LinksTo2 == 1))
+		{
+			BestPin = Pin2;
+		}
+		else if ((LinksTo1 == 1) && (LinksTo2 > 1))
+		{
+			BestPin = Pin1;
+		}
+		else
+		{
+			// Both pins have multiple links, or both pins have one link, choose based on distance to the pins
+			BestPin = (DistanceSquaredToPin1 < DistanceSquaredToPin2) ? Pin1 : Pin2;
+		}
+	}
+}
diff --git a/Engine/Source/Editor/GraphEditor/Private/GraphEditorSettings.cpp b/Engine/Source/Editor/GraphEditor/Private/GraphEditorSettings.cpp
--- a/Engine/Source/Editor/GraphEditor/Private/GraphEditorSettings.cpp
+++ b/Engine/Source/Editor/GraphEditor/Private/GraphEditorSettings.cpp
@@ -13,6 +13,8 @@ UGraphEditorSettings::UGraphEditorSettings( const FObjectInitializer& ObjectInit
 	, PaddingRightOfInput(10.0f)
 	, PaddingLeftOfOutput(10.0f)
 	, PaddingTowardsNodeEdge(10.0f)
+	, bTreatSplinesLikePinsEXPERIMENTAL(false)
+	, SplineHoverToleranceEXPERIMENTAL(2.0f)
 	, ForwardSplineHorizontalDeltaRange(1000.0f)
 	, ForwardSplineVerticalDeltaRange(1000.0f)
 	, ForwardSplineTangentFromHorizontalDelta(1.0f, 0.0f)
diff --git a/Engine/Source/Editor/GraphEditor/Private/MaterialGraphConnectionDrawingPolicy.cpp b/Engine/Source/Editor/GraphEditor/Private/MaterialGraphConnectionDrawingPolicy.cpp
--- a/Engine/Source/Editor/GraphEditor/Private/MaterialGraphConnectionDrawingPolicy.cpp
+++ b/Engine/Source/Editor/GraphEditor/Private/MaterialGraphConnectionDrawingPolicy.cpp
@@ -21,6 +21,8 @@ FMaterialGraphConnectionDrawingPolicy::FMaterialGraphConnectionDrawingPolicy(int
 
 void FMaterialGraphConnectionDrawingPolicy::DetermineWiringStyle(UEdGraphPin* OutputPin, UEdGraphPin* InputPin, /*inout*/ FConnectionParams& Params)
 {
+	Params.AssociatedPin1 = OutputPin;
+	Params.AssociatedPin2 = InputPin;
 	Params.WireColor = MaterialGraphSchema->ActivePinColor;
 
 	// Have to consider both pins as the input will be an 'output' when previewing a connection
diff --git a/Engine/Source/Editor/GraphEditor/Private/SGraphPanel.cpp b/Engine/Source/Editor/GraphEditor/Private/SGraphPanel.cpp
--- a/Engine/Source/Editor/GraphEditor/Private/SGraphPanel.cpp
+++ b/Engine/Source/Editor/GraphEditor/Private/SGraphPanel.cpp
@@ -100,6 +100,9 @@ void SGraphPanel::Construct( const SGraphPanel::FArguments& InArgs )
 	MyRegisteredGraphChangedDelegateHandle = this->GraphObj->AddOnGraphChangedHandler(MyRegisteredGraphChangedDelegate);
 	
 	ShowGraphStateOverlay = InArgs._ShowGraphStateOverlay;
+
+	SavedMousePosForOnPaintEventLocalSpace = FVector2D::ZeroVector;
+	PreviousFrameSavedMousePosForSplineOverlap = FVector2D::ZeroVector;
 }
 
 SGraphPanel::~SGraphPanel()
@@ -244,7 +247,7 @@ int32 SGraphPanel::OnPaint( const FPaintArgs& Args, const FGeometry& AllottedGeo
 
 					/* When dragging off a pin, we want to duck the alpha of some nodes */
 					TSharedPtr< SGraphPin > OnlyStartPin = (1 == PreviewConnectorFromPins.Num()) ? PreviewConnectorFromPins[0].FindInGraphPanel(*this) : TSharedPtr< SGraphPin >();
-					const bool bNodeIsNotUsableInCurrentContext = Schema->FadeNodeWhenDraggingOffPin(NodeObj, OnlyStartPin.IsValid() ? OnlyStartPin.Get()->GetPinObj() : NULL);
+					const bool bNodeIsNotUsableInCurrentContext = Schema->FadeNodeWhenDraggingOffPin(NodeObj, OnlyStartPin.IsValid() ? OnlyStartPin.Get()->GetPinObj() : nullptr);
 					const FWidgetStyle& NodeStyleToUse = (bNodeIsDifferent && !bNodeIsNotUsableInCurrentContext)? InWidgetStyle : FadedStyle;
 
 					// Draw the node.O
@@ -264,7 +267,7 @@ int32 SGraphPanel::OnPaint( const FPaintArgs& Args, const FGeometry& AllottedGeo
 						{
 							FOverlayBrushInfo& OverlayInfo = OverlayBrushes[BrushIndex];
 							const FSlateBrush* OverlayBrush = OverlayInfo.Brush;
-							if(OverlayBrush != NULL)
+							if (OverlayBrush != nullptr)
 							{
 								FPaintGeometry BouncedGeometry = CurWidget.Geometry.ToPaintGeometry(OverlayInfo.OverlayOffset, OverlayBrush->ImageSize, 1.f);
 
@@ -358,8 +361,9 @@ int32 SGraphPanel::OnPaint( const FPaintArgs& Args, const FGeometry& AllottedGeo
 				OverridePins.Add(Pin);
 			}
 		}
-		ConnectionDrawingPolicy->SetHoveredPins(CurrentHoveredPins, OverridePins, TimeSinceMouseEnteredPin);
+		ConnectionDrawingPolicy->SetHoveredPins(CurrentHoveredPins, OverridePins, TimeWhenMouseEnteredPin);
 		ConnectionDrawingPolicy->SetMarkedPin(MarkedPin);
+		ConnectionDrawingPolicy->SetMousePosition(AllottedGeometry.AbsolutePosition + SavedMousePosForOnPaintEventLocalSpace);
 
 		// Get the set of pins for all children and synthesize geometry for culled out pins so lines can be drawn to them.
 		TMap<TSharedRef<SWidget>, FArrangedWidget> PinGeometries;
@@ -414,7 +418,7 @@ int32 SGraphPanel::OnPaint( const FPaintArgs& Args, const FGeometry& AllottedGeo
 				}
 				const FArrangedWidget* PinGeometry = PinGeometries.Find( CurrentStartPin.ToSharedRef() );
 
-				if (PinGeometry != NULL)
+				if (PinGeometry != nullptr)
 				{
 					FVector2D StartPoint;
 					FVector2D EndPoint;
@@ -446,6 +450,27 @@ int32 SGraphPanel::OnPaint( const FPaintArgs& Args, const FGeometry& AllottedGeo
 		// Draw all regular connections
 		ConnectionDrawingPolicy->Draw(PinGeometries, ArrangedChildren);
 
+		// Pull back data from the drawing policy
+		{
+			FGraphSplineOverlapResult OverlapData = ConnectionDrawingPolicy->SplineOverlapResult;
+
+			if (OverlapData.IsValid())
+			{
+				OverlapData.ComputeBestPin();
+
+				// Only allow spline overlaps when there is no node under the cursor
+				const FVector2D PaintAbsoluteSpaceMousePos = AllottedGeometry.AbsolutePosition + SavedMousePosForOnPaintEventLocalSpace;
+				const int32 HoveredNodeIndex = SWidget::FindChildUnderPosition(ArrangedChildren, PaintAbsoluteSpaceMousePos);
+				if (HoveredNodeIndex != INDEX_NONE)
+				{
+					OverlapData = FGraphSplineOverlapResult();
+				}
+			}
+
+			// Update the spline hover state
+			const_cast<SGraphPanel*>(this)->OnSplineHoverStateChanged(OverlapData);
+		}
+
 		delete ConnectionDrawingPolicy;
 	}
 
@@ -453,21 +478,21 @@ int32 SGraphPanel::OnPaint( const FPaintArgs& Args, const FGeometry& AllottedGeo
 	++MaxLayerId;
 	PaintSurroundSunkenShadow(FEditorStyle::GetBrush(TEXT("Graph.Shadow")), AllottedGeometry, MyClippingRect, OutDrawElements, MaxLayerId);
 
-	if(ShowGraphStateOverlay.Get())
+	if (ShowGraphStateOverlay.Get())
 	{
 		const FSlateBrush* BorderBrush = nullptr;
-		if((GEditor->bIsSimulatingInEditor || GEditor->PlayWorld != NULL))
+		if ((GEditor->bIsSimulatingInEditor || GEditor->PlayWorld != nullptr))
 		{
 			// Draw a surrounding indicator when PIE is active, to make it clear that the graph is read-only, etc...
 			BorderBrush = FEditorStyle::GetBrush(TEXT("Graph.PlayInEditor"));
 		}
-		else if(!IsEditable.Get())
+		else if (!IsEditable.Get())
 		{
 			// Draw a different border when we're not simulating but the graph is read-only
 			BorderBrush = FEditorStyle::GetBrush(TEXT("Graph.ReadOnlyBorder"));
 		}
 
-		if(BorderBrush)
+		if (BorderBrush != nullptr)
 		{
 			// Actually draw the border
 			FSlateDrawElement::MakeBox(
@@ -490,6 +515,45 @@ int32 SGraphPanel::OnPaint( const FPaintArgs& Args, const FGeometry& AllottedGeo
 	return MaxLayerId;
 }
 
+void SGraphPanel::OnSplineHoverStateChanged(const FGraphSplineOverlapResult& NewSplineHoverState)
+{
+	UEdGraphPin* OldPin = PreviousFrameSplineOverlap.GetBestPin();
+	PreviousFrameSplineOverlap = NewSplineHoverState;
+	UEdGraphPin* NewPin = PreviousFrameSplineOverlap.GetBestPin();
+
+	PreviousFrameSavedMousePosForSplineOverlap = SavedMousePosForOnPaintEventLocalSpace;
+
+	// Handle mouse enter/leaves on the associated pin
+	if (OldPin != NewPin)
+	{
+		if (OldPin != nullptr)
+		{
+			FGraphPinHandle OldPinHandle(OldPin);
+			TSharedPtr<SGraphPin> OldPinWidget = OldPinHandle.FindInGraphPanel(*this);
+
+			if (OldPinWidget.IsValid())
+			{
+				OldPinWidget->OnMouseLeave(LastPointerEvent);
+			}
+		}
+
+		if (NewPin != nullptr)
+		{
+			FGraphPinHandle NewPinHandle(NewPin);
+			TSharedPtr<SGraphPin> NewPinWidget = NewPinHandle.FindInGraphPanel(*this);
+
+			if (NewPinWidget.IsValid())
+			{
+				NewPinWidget->OnMouseEnter(LastPointerGeometry, LastPointerEvent);
+
+				// Get the pin/wire glowing quicker, since it's a direct selection (this time was already set to 'now' as part of entering the pin)
+				//@TODO: Source this parameter from the graph rendering settings once it is there (see code in ApplyHoverDeemphasis)
+				TimeWhenMouseEnteredPin -= 0.75f;
+			}
+		}
+	}
+}
+
 bool SGraphPanel::SupportsKeyboardFocus() const
 {
 	return true;
@@ -518,12 +582,12 @@ void SGraphPanel::OnArrangeChildren( const FGeometry& AllottedGeometry, FArrange
 	ArrangedChildren.Append(MyArrangedChildren);
 }
 
-void SGraphPanel::UpdateSelectedNodesPositions (FVector2D PositionIncrement)
+void SGraphPanel::UpdateSelectedNodesPositions(FVector2D PositionIncrement)
 {
 	for (FGraphPanelSelectionSet::TIterator NodeIt(SelectionManager.SelectedNodes); NodeIt; ++NodeIt)
 	{
 		TSharedRef<SNode>* pWidget = NodeToWidgetLookup.Find(*NodeIt);
-		if (pWidget != NULL)
+		if (pWidget != nullptr)
 		{
 			SNode& Widget = pWidget->Get();
 			SNode::FNodeSet NodeFilter;
@@ -566,12 +630,59 @@ FReply SGraphPanel::OnKeyDown( const FGeometry& MyGeometry, const FKeyEvent& InK
 			ChangeZoomLevel(+1, CachedAllottedGeometryScaledSize / 2.f, InKeyEvent.IsControlDown());
 			return FReply::Handled();
 		}
-
 	}
 
 	return SNodePanel::OnKeyDown(MyGeometry, InKeyEvent);
 }
 
+FReply SGraphPanel::OnMouseButtonDown(const FGeometry& MyGeometry, const FPointerEvent& MouseEvent)
+{
+	if ((MouseEvent.GetEffectingButton() == EKeys::LeftMouseButton) && (MouseEvent.IsAltDown() || MouseEvent.IsControlDown()))
+	{
+		if (SGraphPin* BestPinFromHoveredSpline = GetBestPinFromHoveredSpline())
+		{
+			return BestPinFromHoveredSpline->OnPinMouseDown(MyGeometry, MouseEvent);
+		}
+	}
+
+	return SNodePanel::OnMouseButtonDown(MyGeometry, MouseEvent);
+}
+
+FReply SGraphPanel::OnMouseButtonUp(const FGeometry& MyGeometry, const FPointerEvent& MouseEvent)
+{
+	if ((MouseEvent.GetEffectingButton() == EKeys::LeftMouseButton) && (MouseEvent.IsShiftDown()))
+	{
+		if (SGraphPin* BestPinFromHoveredSpline = GetBestPinFromHoveredSpline())
+		{
+			return BestPinFromHoveredSpline->OnMouseButtonUp(MyGeometry, MouseEvent);
+		}
+	}
+
+	return SNodePanel::OnMouseButtonUp(MyGeometry, MouseEvent);
+}
+
+FReply SGraphPanel::OnMouseButtonDoubleClick(const FGeometry& MyGeometry, const FPointerEvent& MouseEvent)
+{
+	if (SGraphPin* BestPinFromHoveredSpline = GetBestPinFromHoveredSpline())
+	{
+		//@TODO: Should insert a reroute node here
+	}
+
+	return SNodePanel::OnMouseButtonDoubleClick(MyGeometry, MouseEvent);
+}
+
+class SGraphPin* SGraphPanel::GetBestPinFromHoveredSpline() const
+{
+	if (PreviousFrameSplineOverlap.IsValid())
+	{
+		FGraphPinHandle SplineAssociatedPinHandle(PreviousFrameSplineOverlap.GetBestPin());
+		return SplineAssociatedPinHandle.FindInGraphPanel(*this).Get();
+	}
+	else
+	{
+		return nullptr;
+	}
+}
 
 void SGraphPanel::GetAllPins(TSet< TSharedRef<SWidget> >& AllPins)
 {
@@ -586,13 +697,13 @@ void SGraphPanel::GetAllPins(TSet< TSharedRef<SWidget> >& AllPins)
 void SGraphPanel::AddPinToHoverSet(UEdGraphPin* HoveredPin)
 {
 	CurrentHoveredPins.Add(HoveredPin);
-	TimeSinceMouseEnteredPin = FSlateApplication::Get().GetCurrentTime();
+	TimeWhenMouseEnteredPin = FSlateApplication::Get().GetCurrentTime();
 }
 
 void SGraphPanel::RemovePinFromHoverSet(UEdGraphPin* UnhoveredPin)
 {
 	CurrentHoveredPins.Remove(UnhoveredPin);
-	TimeSinceMouseLeftPin = FSlateApplication::Get().GetCurrentTime();
+	TimeWhenMouseLeftPin = FSlateApplication::Get().GetCurrentTime();
 }
 
 void SGraphPanel::ArrangeChildrenForContextMenuSummon(const FGeometry& AllottedGeometry, FArrangedChildren& ArrangedChildren) const
@@ -625,8 +736,8 @@ TSharedPtr<SWidget> SGraphPanel::OnSummonContextMenu(const FGeometry& MyGeometry
 	{
 		// If we didn't drag very far, summon a context menu.
 		// Figure out what's under the mouse: Node, Pin or just the Panel, and summon the context menu for that.
-		UEdGraphNode* NodeUnderCursor = NULL;
-		UEdGraphPin* PinUnderCursor = NULL;
+		UEdGraphNode* NodeUnderCursor = nullptr;
+		UEdGraphPin* PinUnderCursor = nullptr;
 		{
 			FArrangedChildren ArrangedNodes(EVisibility::Visible);
 			this->ArrangeChildrenForContextMenuSummon(MyGeometry, ArrangedNodes);
@@ -710,6 +821,29 @@ bool SGraphPanel::OnHandleLeftMouseRelease(const FGeometry& MyGeometry, const FP
 	}
 }
 
+FReply SGraphPanel::OnMouseMove(const FGeometry& MyGeometry, const FPointerEvent& MouseEvent)
+{
+	LastPointerEvent = MouseEvent;
+	LastPointerGeometry = MyGeometry;
+
+	// Save the mouse position to use in OnPaint for spline hit detection
+	SavedMousePosForOnPaintEventLocalSpace = MyGeometry.AbsoluteToLocal(MouseEvent.GetScreenSpacePosition());
+
+	// Invalidate the spline results if we moved very far
+	const FVector2D MouseDelta = SavedMousePosForOnPaintEventLocalSpace - PreviousFrameSavedMousePosForSplineOverlap;
+	const float MouseDeltaLengthSquared = MouseDelta.SizeSquared();
+	const bool bCursorInDeadZone = MouseDeltaLengthSquared <= FMath::Square(FSlateApplication::Get().GetDragTriggerDistance());
+
+	if (!bCursorInDeadZone)
+	{
+		//@TODO: Should we do this or just rely on the next OnPaint?
+		// Our frame-latent approximation is going to be totally junk if the mouse is moving quickly
+		OnSplineHoverStateChanged(FGraphSplineOverlapResult());
+	}
+
+	return SNodePanel::OnMouseMove(MyGeometry, MouseEvent);
+}
+
 void SGraphPanel::OnDragEnter( const FGeometry& MyGeometry, const FDragDropEvent& DragDropEvent )
 {
 	TSharedPtr<FGraphEditorDragDropAction> DragConnectionOp = DragDropEvent.GetOperationAs<FGraphEditorDragDropAction>();
@@ -724,7 +858,7 @@ void SGraphPanel::OnDragLeave( const FDragDropEvent& DragDropEvent )
 	TSharedPtr<FGraphEditorDragDropAction> Operation = DragDropEvent.GetOperationAs<FGraphEditorDragDropAction>();
 	if( Operation.IsValid() )
 	{
-		Operation->SetHoveredGraph( TSharedPtr<SGraphPanel>(NULL) );
+		Operation->SetHoveredGraph(TSharedPtr<SGraphPanel>(nullptr));
 	}
 	else
 	{
@@ -745,31 +879,31 @@ FReply SGraphPanel::OnDragOver( const FGeometry& MyGeometry, const FDragDropEven
 	}
 
 	// Handle Read only graphs
-	if( !IsEditable.Get() )
+	if (!IsEditable.Get())
 	{
 		TSharedPtr<FGraphEditorDragDropAction> GraphDragDropOp = DragDropEvent.GetOperationAs<FGraphEditorDragDropAction>();
 
-		if( GraphDragDropOp.IsValid() )
+		if (GraphDragDropOp.IsValid())
 		{
-			GraphDragDropOp->SetDropTargetValid( false );
+			GraphDragDropOp->SetDropTargetValid(false);
 		}
 		else
 		{
 			TSharedPtr<FDecoratedDragDropOp> AssetOp = DragDropEvent.GetOperationAs<FDecoratedDragDropOp>();
-			if( AssetOp.IsValid() )
+			if (AssetOp.IsValid())
 			{
 				FText Tooltip = AssetOp->GetHoverText();
-				if( Tooltip.IsEmpty() )
+				if (Tooltip.IsEmpty())
 				{
 					Tooltip = NSLOCTEXT( "GraphPanel", "DragDropOperation", "Graph is Read-Only" );
 				}
-				AssetOp->SetToolTip( Tooltip, FEditorStyle::GetBrush(TEXT("Graph.ConnectorFeedback.Error")));
+				AssetOp->SetToolTip(Tooltip, FEditorStyle::GetBrush(TEXT("Graph.ConnectorFeedback.Error")));
 			}
 		}
 		return FReply::Handled();
 	}
 
-	if( Operation->IsOfType<FGraphEditorDragDropAction>() )
+	if (Operation->IsOfType<FGraphEditorDragDropAction>())
 	{
 		PreviewConnectorEndpoint = MyGeometry.AbsoluteToLocal( DragDropEvent.GetScreenSpacePosition() );
 		return FReply::Handled();
@@ -780,7 +914,7 @@ FReply SGraphPanel::OnDragOver( const FGeometry& MyGeometry, const FDragDropEven
 	}
 	else if (Operation->IsOfType<FAssetDragDropOp>())
 	{
-		if(GraphObj != NULL && GraphObj->GetSchema())
+		if ((GraphObj != nullptr) && (GraphObj->GetSchema() != nullptr))
 		{
 			TSharedPtr<FAssetDragDropOp> AssetOp = StaticCastSharedPtr<FAssetDragDropOp>(Operation);
 			bool bOkIcon = false;
@@ -835,7 +969,7 @@ FReply SGraphPanel::OnDrop( const FGeometry& MyGeometry, const FDragDropEvent& D
 	}
 	else
 	{
-		if(GraphObj != NULL && GraphObj->GetSchema() != NULL)
+		if ((GraphObj != nullptr) && (GraphObj->GetSchema() != nullptr))
 		{
 			TArray< FAssetData > DroppedAssetData = AssetUtil::ExtractAssetDataFromDrag( DragDropEvent );
 
@@ -997,7 +1131,7 @@ TSharedPtr<SGraphNode> SGraphPanel::GetNodeWidgetFromGuid(FGuid Guid) const
 void SGraphPanel::Update()
 {
 	// Add widgets for all the nodes that don't have one.
-	if(GraphObj != NULL)
+	if (GraphObj != nullptr)
 	{
 		// Scan for all missing nodes
 		for (int32 NodeIndex = 0; NodeIndex < GraphObj->Nodes.Num(); ++NodeIndex)
@@ -1051,7 +1185,7 @@ bool SGraphPanel::IsNodeTitleVisible(const class UEdGraphNode* Node, bool bReque
 	bool bTitleVisible = false;
 	TSharedRef<SNode>* pWidget = NodeToWidgetLookup.Find(Node);
 
-	if (pWidget != NULL)
+	if (pWidget != nullptr)
 	{
 		TWeakPtr<SGraphNode> GraphNode = StaticCastSharedRef<SGraphNode>(*pWidget);
 		if(GraphNode.IsValid() && !HasMouseCapture())
@@ -1210,13 +1344,13 @@ void SGraphPanel::OnGraphChanged(const FEdGraphEditAction& EditAction)
 	}
 }
 
-void SGraphPanel::NotifyGraphChanged ( const FEdGraphEditAction& EditAction)
+void SGraphPanel::NotifyGraphChanged(const FEdGraphEditAction& EditAction)
 {
 	// Forward call
 	OnGraphChanged(EditAction);
 }
 
-void SGraphPanel::AddReferencedObjects( FReferenceCollector& Collector )
+void SGraphPanel::AddReferencedObjects(FReferenceCollector& Collector)
 {
 	Collector.AddReferencedObject( GraphObj );
 	Collector.AddReferencedObject( GraphObjToDiff );
diff --git a/Engine/Source/Editor/GraphEditor/Private/SGraphPanel.h b/Engine/Source/Editor/GraphEditor/Private/SGraphPanel.h
--- a/Engine/Source/Editor/GraphEditor/Private/SGraphPanel.h
+++ b/Engine/Source/Editor/GraphEditor/Private/SGraphPanel.h
@@ -3,6 +3,7 @@
 #pragma once
 
 #include "GraphEditor.h"
+#include "GraphSplineOverlapResult.h"
 
 class SGraphNode;
 class SGraphPin;
@@ -93,7 +94,11 @@ public:
 	~SGraphPanel();
 public:
 	// SWidget interface
-	virtual void OnDragEnter( const FGeometry& MyGeometry, const FDragDropEvent& DragDropEvent ) override;
+	virtual FReply OnMouseButtonDown(const FGeometry& MyGeometry, const FPointerEvent& MouseEvent) override;
+	virtual FReply OnMouseButtonUp(const FGeometry& MyGeometry, const FPointerEvent& MouseEvent) override;
+	virtual FReply OnMouseButtonDoubleClick(const FGeometry& InMyGeometry, const FPointerEvent& InMouseEvent) override;
+	virtual FReply OnMouseMove(const FGeometry& MyGeometry, const FPointerEvent& MouseEvent) override;
+	virtual void OnDragEnter(const FGeometry& MyGeometry, const FDragDropEvent& DragDropEvent) override;
 	virtual void OnDragLeave( const FDragDropEvent& DragDropEvent ) override;
 	virtual FReply OnDragOver( const FGeometry& MyGeometry, const FDragDropEvent& DragDropEvent ) override;
 	virtual FReply OnDrop( const FGeometry& MyGeometry, const FDragDropEvent& DragDropEvent ) override;
@@ -197,20 +202,28 @@ protected:
 	/** Pin visibility mode */
 	SGraphEditor::EPinVisibility PinVisibility;
 
-	/** The selected (hovered) pin in the graph */
-	TWeakPtr<SGraphPin> SelectedPin;
-
 	/** List of pins currently being hovered over */
 	TSet< TWeakObjectPtr<UEdGraphPin> > CurrentHoveredPins;
 
 	/** Time since the last mouse enter/exit on a pin */
-	double TimeSinceMouseEnteredPin;
-	double TimeSinceMouseLeftPin;
+	double TimeWhenMouseEnteredPin;
+	double TimeWhenMouseLeftPin;
 
 	/** Sometimes the panel draws a preview connector; e.g. when the user is connecting pins */
 	TArray< FGraphPinHandle > PreviewConnectorFromPins;
 	FVector2D PreviewConnectorEndpoint;
 
+	/** Last mouse position seen, used for paint-centric highlighting */
+	FVector2D SavedMousePosForOnPaintEventLocalSpace;
+	
+	/** The overlap results from the previous OnPaint call */
+	FVector2D PreviousFrameSavedMousePosForSplineOverlap;
+	FGraphSplineOverlapResult PreviousFrameSplineOverlap;
+
+	/** The mouse state from the last mouse move event, used to synthesize pin actions when hovering over a spline on the panel */
+	FGeometry LastPointerGeometry;
+	FPointerEvent LastPointerEvent;
+
 	/** Invoked when we need to summon a context menu */
 	FOnGetContextMenuFor OnGetContextMenuFor;
 
@@ -267,5 +280,11 @@ private:
 	void OnGraphChanged( const FEdGraphEditAction& InAction );
 
 	/** Update all selected nodes position by provided vector2d */
-	void UpdateSelectedNodesPositions (FVector2D PositionIncrement);
+	void UpdateSelectedNodesPositions(FVector2D PositionIncrement);
+
+	/** Handle updating the spline hover state */
+	void OnSplineHoverStateChanged(const FGraphSplineOverlapResult& NewSplineHoverState);
+
+	/** Returns the pin that we're considering as hovered if we are hovering over a spline; may be null */
+	class SGraphPin* GetBestPinFromHoveredSpline() const;
 };
diff --git a/Engine/Source/Editor/GraphEditor/Private/SGraphPin.cpp b/Engine/Source/Editor/GraphEditor/Private/SGraphPin.cpp
--- a/Engine/Source/Editor/GraphEditor/Private/SGraphPin.cpp
+++ b/Engine/Source/Editor/GraphEditor/Private/SGraphPin.cpp
@@ -551,13 +551,16 @@ FReply SGraphPin::OnMouseButtonUp( const FGeometry& MyGeometry, const FPointerEv
 
 void SGraphPin::OnMouseEnter( const FGeometry& MyGeometry, const FPointerEvent& MouseEvent )
 {
-	FKnotNetCollector NetCollector(GetPinObj());
-
-	TSharedPtr<SGraphPanel> Panel = OwnerNodePtr.Pin()->GetOwnerPanel();
-	for (UEdGraphPin* PinInNet : NetCollector.VisitedPins)
+	if (!bIsHovered)
 	{
-		Panel->AddPinToHoverSet(PinInNet);
-		HoverPinSet.Add(PinInNet);
+		FKnotNetCollector NetCollector(GetPinObj());
+
+		TSharedPtr<SGraphPanel> Panel = OwnerNodePtr.Pin()->GetOwnerPanel();
+		for (UEdGraphPin* PinInNet : NetCollector.VisitedPins)
+		{
+			Panel->AddPinToHoverSet(PinInNet);
+			HoverPinSet.Add(PinInNet);
+		}
 	}
 
 	SCompoundWidget::OnMouseEnter(MyGeometry, MouseEvent);
diff --git a/Engine/Source/Editor/GraphEditor/Private/SoundCueGraphConnectionDrawingPolicy.cpp b/Engine/Source/Editor/GraphEditor/Private/SoundCueGraphConnectionDrawingPolicy.cpp
--- a/Engine/Source/Editor/GraphEditor/Private/SoundCueGraphConnectionDrawingPolicy.cpp
+++ b/Engine/Source/Editor/GraphEditor/Private/SoundCueGraphConnectionDrawingPolicy.cpp
@@ -114,6 +114,9 @@ void FSoundCueGraphConnectionDrawingPolicy::BuildAudioFlowRoadmap()
 // Give specific editor modes a chance to highlight this connection or darken non-interesting connections
 void FSoundCueGraphConnectionDrawingPolicy::DetermineWiringStyle(UEdGraphPin* OutputPin, UEdGraphPin* InputPin, /*inout*/ FConnectionParams& Params)
 {
+	Params.AssociatedPin1 = OutputPin;
+	Params.AssociatedPin2 = InputPin;
+
 	// Get the schema and grab the default color from it
 	check(OutputPin);
 	check(GraphObj);
diff --git a/Engine/Source/Editor/GraphEditor/Private/StateMachineConnectionDrawingPolicy.cpp b/Engine/Source/Editor/GraphEditor/Private/StateMachineConnectionDrawingPolicy.cpp
--- a/Engine/Source/Editor/GraphEditor/Private/StateMachineConnectionDrawingPolicy.cpp
+++ b/Engine/Source/Editor/GraphEditor/Private/StateMachineConnectionDrawingPolicy.cpp
@@ -17,6 +17,8 @@ FStateMachineConnectionDrawingPolicy::FStateMachineConnectionDrawingPolicy(int32
 
 void FStateMachineConnectionDrawingPolicy::DetermineWiringStyle(UEdGraphPin* OutputPin, UEdGraphPin* InputPin, /*inout*/ FConnectionParams& Params)
 {
+	Params.AssociatedPin1 = OutputPin;
+	Params.AssociatedPin2 = InputPin;
 	Params.WireThickness = 1.5f;
 
 	if (InputPin)
diff --git a/Engine/Source/Editor/GraphEditor/Public/ConnectionDrawingPolicy.h b/Engine/Source/Editor/GraphEditor/Public/ConnectionDrawingPolicy.h
--- a/Engine/Source/Editor/GraphEditor/Public/ConnectionDrawingPolicy.h
+++ b/Engine/Source/Editor/GraphEditor/Public/ConnectionDrawingPolicy.h
@@ -2,6 +2,8 @@
 
 #pragma once
 
+#include "GraphSplineOverlapResult.h"
+
 /////////////////////////////////////////////////////
 
 GRAPHEDITOR_API DECLARE_LOG_CATEGORY_EXTERN(LogConnectionDrawingPolicy, Log, All);
@@ -29,6 +31,9 @@ public:
 struct GRAPHEDITOR_API FConnectionParams
 {
 	FLinearColor WireColor;
+	UEdGraphPin* AssociatedPin1;
+	UEdGraphPin* AssociatedPin2;
+
 	float WireThickness;
 	bool bDrawBubbles;
 	bool bUserFlag1;
@@ -39,6 +44,8 @@ struct GRAPHEDITOR_API FConnectionParams
 
 	FConnectionParams()
 		: WireColor(FLinearColor::White)
+		, AssociatedPin1(nullptr)
+		, AssociatedPin2(nullptr)
 		, WireThickness(1.0f)
 		, bDrawBubbles(false)
 		, bUserFlag1(false)
@@ -68,6 +75,9 @@ protected:
 public:
 	FVector2D ArrowRadius;
 	FVector2D MidpointRadius;
+
+	FGraphSplineOverlapResult SplineOverlapResult;
+
 protected:
 	float ZoomFactor; 
 	float HoverDeemphasisDarkFraction;
@@ -77,6 +87,7 @@ protected:
 	TSet< UEdGraphPin* > HoveredPins;
 	TMap<TSharedRef<SWidget>, FArrangedWidget>* PinGeometries;
 	double LastHoverTimeEvent;
+	FVector2D LocalMousePosition;
 public:
 	virtual ~FConnectionDrawingPolicy() {}
 
@@ -85,6 +96,8 @@ public:
 	// Update the drawing policy with the set of hovered pins (which can be empty)
 	void SetHoveredPins(const TSet< TWeakObjectPtr<UEdGraphPin> >& InHoveredPins, const TArray< TSharedPtr<SGraphPin> >& OverridePins, double HoverTime);
 
+	void SetMousePosition(const FVector2D& InMousePos);
+
 	// Update the drawing policy with the marked pin (which may not be valid)
 	void SetMarkedPin(TWeakPtr<SGraphPin> InMarkedPin);
 
diff --git a/Engine/Source/Editor/GraphEditor/Public/GraphEditorSettings.h b/Engine/Source/Editor/GraphEditor/Public/GraphEditorSettings.h
--- a/Engine/Source/Editor/GraphEditor/Public/GraphEditorSettings.h
+++ b/Engine/Source/Editor/GraphEditor/Public/GraphEditorSettings.h
@@ -10,8 +10,7 @@
  * Implements settings for the graph editor.
  */
 UCLASS(config=EditorUserSettings)
-class GRAPHEDITOR_API UGraphEditorSettings
-	:	public UObject
+class GRAPHEDITOR_API UGraphEditorSettings : public UObject
 {
 	GENERATED_UCLASS_BODY()
 
@@ -19,8 +18,8 @@ public:
 
 #if WITH_EDITOR
 	// UObject interface
-
 	virtual void PostEditChangeProperty(FPropertyChangedEvent& PropertyChangedEvent) override;
+	// End of UObject interface
 #endif
 
 public:
@@ -63,6 +62,14 @@ public:
 
 public:
 
+	/** If enabled, allows splines to be Alt+Clicked to break them or Ctrl+Dragged to move them as if these actions were taking place on the associated pin. */
+	UPROPERTY(EditAnywhere, config, Category=Splines)
+	bool bTreatSplinesLikePinsEXPERIMENTAL;
+
+	/** The distance threshold controlling how close the mouse has to be to the spline in order to trigger a hover response */ 
+	UPROPERTY(EditAnywhere, config, Category=Splines, meta=(EditCondition=bTreatSplinesLikePinsEXPERIMENTAL))
+	float SplineHoverToleranceEXPERIMENTAL;
+
 	/** The maximum value to clamp the absolute value of the horizontal distance between endpoints when calculating tangents (when the wire is moving forward) */
 	UPROPERTY(config, EditAnywhere, Category=Splines, AdvancedDisplay)
 	float ForwardSplineHorizontalDeltaRange;
@@ -101,113 +108,113 @@ public:
 public:
 
 	/** The default color is used only for types not specifically defined below.  Generally if it's seen, it means another type needs to be defined so that the wire in question can have an appropriate color. */
-	UPROPERTY(EditAnywhere,  config, Category=PinColors)
+	UPROPERTY(EditAnywhere, config, Category=PinColors)
 	FLinearColor DefaultPinTypeColor;
 
 	/** Execution pin type color */
-	UPROPERTY(EditAnywhere,  config, Category=PinColors)
+	UPROPERTY(EditAnywhere, config, Category=PinColors)
 	FLinearColor ExecutionPinTypeColor;
 
 	/** Boolean pin type color */
-	UPROPERTY(EditAnywhere,  config, Category=PinColors)
+	UPROPERTY(EditAnywhere, config, Category=PinColors)
 	FLinearColor BooleanPinTypeColor;
 
 	/** Byte pin type color */
-	UPROPERTY(EditAnywhere,  config, Category=PinColors)
+	UPROPERTY(EditAnywhere, config, Category=PinColors)
 	FLinearColor BytePinTypeColor;
 
 	/** Class pin type color */
-	UPROPERTY(EditAnywhere,  config, Category=PinColors)
+	UPROPERTY(EditAnywhere, config, Category=PinColors)
 	FLinearColor ClassPinTypeColor;
 
 	/** Integer pin type color */
-	UPROPERTY(EditAnywhere,  config, Category=PinColors)
+	UPROPERTY(EditAnywhere, config, Category=PinColors)
 	FLinearColor IntPinTypeColor;
 
 	/** Floating-point pin type color */
-	UPROPERTY(EditAnywhere,  config, Category=PinColors)
+	UPROPERTY(EditAnywhere, config, Category=PinColors)
 	FLinearColor FloatPinTypeColor;
 
 	/** Name pin type color */
-	UPROPERTY(EditAnywhere,  config, Category=PinColors)
+	UPROPERTY(EditAnywhere, config, Category=PinColors)
 	FLinearColor NamePinTypeColor;
 
 	/** Delegate pin type color */
-	UPROPERTY(EditAnywhere,  config, Category=PinColors)
+	UPROPERTY(EditAnywhere, config, Category=PinColors)
 	FLinearColor DelegatePinTypeColor;
 
 	/** Object pin type color */
-	UPROPERTY(EditAnywhere,  config, Category=PinColors)
+	UPROPERTY(EditAnywhere, config, Category=PinColors)
 	FLinearColor ObjectPinTypeColor;
 
 	/** Interface pin type color */
-	UPROPERTY(EditAnywhere,  config, Category=PinColors)
+	UPROPERTY(EditAnywhere, config, Category=PinColors)
 	FLinearColor InterfacePinTypeColor;
 
 	/** String pin type color */
-	UPROPERTY(EditAnywhere,  config, Category=PinColors)
+	UPROPERTY(EditAnywhere, config, Category=PinColors)
 	FLinearColor StringPinTypeColor;
 
 	/** Text pin type color */
-	UPROPERTY(EditAnywhere,  config, Category=PinColors)
+	UPROPERTY(EditAnywhere, config, Category=PinColors)
 	FLinearColor TextPinTypeColor;
 
 	/** Struct pin type color */
-	UPROPERTY(EditAnywhere,  config, Category=PinColors)
+	UPROPERTY(EditAnywhere, config, Category=PinColors)
 	FLinearColor StructPinTypeColor;
 
 	/** Wildcard pin type color */
-	UPROPERTY(EditAnywhere,  config, Category=PinColors)
+	UPROPERTY(EditAnywhere, config, Category=PinColors)
 	FLinearColor WildcardPinTypeColor;
 
 	/** Vector pin type color */
-	UPROPERTY(EditAnywhere,  config, Category=PinColors)
+	UPROPERTY(EditAnywhere, config, Category=PinColors)
 	FLinearColor VectorPinTypeColor;
 
 	/** Rotator pin type color */
-	UPROPERTY(EditAnywhere,  config, Category=PinColors)
+	UPROPERTY(EditAnywhere, config, Category=PinColors)
 	FLinearColor RotatorPinTypeColor;
 
 	/** Transform pin type color */
-	UPROPERTY(EditAnywhere,  config, Category=PinColors)
+	UPROPERTY(EditAnywhere, config, Category=PinColors)
 	FLinearColor TransformPinTypeColor;
 
 	/** Index pin type color */
-	UPROPERTY(EditAnywhere,  config, Category=PinColors)
+	UPROPERTY(EditAnywhere, config, Category=PinColors)
 	FLinearColor IndexPinTypeColor;
 
 public:
 
 	/** Event node title color */
-	UPROPERTY(EditAnywhere,  config, Category=NodeTitleColors)
+	UPROPERTY(EditAnywhere, config, Category=NodeTitleColors)
 	FLinearColor EventNodeTitleColor;
 
 	/** CallFunction node title color */
-	UPROPERTY(EditAnywhere,  config, Category=NodeTitleColors)
+	UPROPERTY(EditAnywhere, config, Category=NodeTitleColors)
 	FLinearColor FunctionCallNodeTitleColor;
 
 	/** Pure function call node title color */
-	UPROPERTY(EditAnywhere,  config, Category=NodeTitleColors)
+	UPROPERTY(EditAnywhere, config, Category=NodeTitleColors)
 	FLinearColor PureFunctionCallNodeTitleColor;
 
 	/** Parent class function call node title color */
-	UPROPERTY(EditAnywhere,  config, Category=NodeTitleColors)
+	UPROPERTY(EditAnywhere, config, Category=NodeTitleColors)
 	FLinearColor ParentFunctionCallNodeTitleColor;
 
 	/** Function Terminator node title color */
-	UPROPERTY(EditAnywhere,  config, Category=NodeTitleColors)
+	UPROPERTY(EditAnywhere, config, Category=NodeTitleColors)
 	FLinearColor FunctionTerminatorNodeTitleColor;
 
 	/** Exec Branch node title color */
-	UPROPERTY(EditAnywhere,  config, Category=NodeTitleColors)
+	UPROPERTY(EditAnywhere, config, Category=NodeTitleColors)
 	FLinearColor ExecBranchNodeTitleColor;
 
 	/** Exec Sequence node title color */
-	UPROPERTY(EditAnywhere,  config, Category=NodeTitleColors)
+	UPROPERTY(EditAnywhere, config, Category=NodeTitleColors)
 	FLinearColor ExecSequenceNodeTitleColor;
 
 	/** Result node title color */
-	UPROPERTY(EditAnywhere,  config, Category=NodeTitleColors)
+	UPROPERTY(EditAnywhere, config, Category=NodeTitleColors)
 	FLinearColor ResultNodeTitleColor;
 
 	/** Default Comment node title color */
@@ -216,10 +223,10 @@ public:
 
 public:
 
-	UPROPERTY(EditAnywhere,  config, Category=Tracing)
+	UPROPERTY(EditAnywhere, config, Category=Tracing)
 	FLinearColor TraceAttackColor;
 
-	UPROPERTY(EditAnywhere,  config, Category=Tracing)
+	UPROPERTY(EditAnywhere, config, Category=Tracing)
 	float TraceAttackWireThickness;
 
 	/** How long is the attack color fully visible */
@@ -233,20 +240,20 @@ public:
 	UPROPERTY()
 	float TraceDecayExponent;
 
-	UPROPERTY(EditAnywhere,  config, Category=Tracing)
+	UPROPERTY(EditAnywhere, config, Category=Tracing)
 	FLinearColor TraceSustainColor;
 
-	UPROPERTY(EditAnywhere,  config, Category=Tracing)
+	UPROPERTY(EditAnywhere, config, Category=Tracing)
 	float TraceSustainWireThickness;
 
 	/** How long is the sustain color fully visible */
 	UPROPERTY()
 	float TraceSustainHoldPeriod;
 
-	UPROPERTY(EditAnywhere,  config, Category=Tracing)
+	UPROPERTY(EditAnywhere, config, Category=Tracing)
 	FLinearColor TraceReleaseColor;
 
-	UPROPERTY(EditAnywhere,  config, Category=Tracing)
+	UPROPERTY(EditAnywhere, config, Category=Tracing)
 	float TraceReleaseWireThickness;
 
 	/** How long does it take to fade from the sustain to the release color */
diff --git a/Engine/Source/Editor/GraphEditor/Public/GraphSplineOverlapResult.h b/Engine/Source/Editor/GraphEditor/Public/GraphSplineOverlapResult.h
new file mode 100644
--- /dev/null
+++ b/Engine/Source/Editor/GraphEditor/Public/GraphSplineOverlapResult.h
@@ -0,0 +1,50 @@
+// Copyright 1998-2015 Epic Games, Inc. All Rights Reserved.
+
+#pragma once
+
+class UEdGraphPin;
+
+/////////////////////////////////////////////////////
+// FGraphSplineOverlapResult
+
+struct FGraphSplineOverlapResult
+{
+	UEdGraphPin* BestPin;
+	UEdGraphPin* Pin1;
+	UEdGraphPin* Pin2;
+	float DistanceSquared;
+	float DistanceSquaredToPin1;
+	float DistanceSquaredToPin2;
+
+	FGraphSplineOverlapResult()
+		: BestPin(nullptr)
+		, Pin1(nullptr)
+		, Pin2(nullptr)
+		, DistanceSquared(FLT_MAX)
+		, DistanceSquaredToPin1(FLT_MAX)
+		, DistanceSquaredToPin2(FLT_MAX)
+	{
+	}
+
+	FGraphSplineOverlapResult(UEdGraphPin* InPin1, UEdGraphPin* InPin2, float InDistanceSquared, float InDistanceSquaredToPin1, float InDistanceSquaredToPin2)
+		: BestPin(nullptr)
+		, Pin1(InPin1)
+		, Pin2(InPin2)
+		, DistanceSquared(InDistanceSquared)
+		, DistanceSquaredToPin1(InDistanceSquaredToPin1)
+		, DistanceSquaredToPin2(InDistanceSquaredToPin2)
+	{
+	}
+
+	bool IsValid() const
+	{
+		return DistanceSquared < FLT_MAX;
+	}
+
+	GRAPHEDITOR_API void ComputeBestPin();
+
+	UEdGraphPin* GetBestPin() const
+	{
+		return IsValid() ? BestPin : nullptr;
+	}
+};
