Graph Editor: Add tool tip display to hovered splines

diff --git a/Engine/Source/Editor/GraphEditor/Private/SGraphPanel.cpp b/Engine/Source/Editor/GraphEditor/Private/SGraphPanel.cpp
--- a/Engine/Source/Editor/GraphEditor/Private/SGraphPanel.cpp
+++ b/Engine/Source/Editor/GraphEditor/Private/SGraphPanel.cpp
@@ -582,6 +582,16 @@ void SGraphPanel::OnArrangeChildren( const FGeometry& AllottedGeometry, FArrange
 	ArrangedChildren.Append(MyArrangedChildren);
 }
 
+TSharedPtr<IToolTip> SGraphPanel::GetToolTip()
+{
+	if (SGraphPin* BestPinFromHoveredSpline = GetBestPinFromHoveredSpline())
+	{
+		return BestPinFromHoveredSpline->GetToolTip();
+	}
+
+	return SNodePanel::GetToolTip();
+}
+
 void SGraphPanel::UpdateSelectedNodesPositions(FVector2D PositionIncrement)
 {
 	for (FGraphPanelSelectionSet::TIterator NodeIt(SelectionManager.SelectedNodes); NodeIt; ++NodeIt)
diff --git a/Engine/Source/Editor/GraphEditor/Private/SGraphPanel.h b/Engine/Source/Editor/GraphEditor/Private/SGraphPanel.h
--- a/Engine/Source/Editor/GraphEditor/Private/SGraphPanel.h
+++ b/Engine/Source/Editor/GraphEditor/Private/SGraphPanel.h
@@ -106,6 +106,7 @@ public:
 	virtual FReply OnKeyDown( const FGeometry& MyGeometry, const FKeyEvent& InKeyEvent ) override;
 	virtual bool SupportsKeyboardFocus() const override;
 	virtual void OnArrangeChildren( const FGeometry& AllottedGeometry, FArrangedChildren& ArrangedChildren ) const override;
+	virtual TSharedPtr<IToolTip> GetToolTip() override;
 	// End of SWidget interface
 
 	// SNodePanel interface
diff --git a/Engine/Source/Editor/GraphEditor/Private/SGraphPin.cpp b/Engine/Source/Editor/GraphEditor/Private/SGraphPin.cpp
--- a/Engine/Source/Editor/GraphEditor/Private/SGraphPin.cpp
+++ b/Engine/Source/Editor/GraphEditor/Private/SGraphPin.cpp
@@ -153,7 +153,7 @@ SGraphPin::SGraphPin()
 void SGraphPin::Construct(const FArguments& InArgs, UEdGraphPin* InPin)
 {
 	bUsePinColorForText = InArgs._UsePinColorForText;
-	this->SetCursor( EMouseCursor::Default );
+	this->SetCursor(EMouseCursor::Default);
 
 	Visibility = TAttribute<EVisibility>(this, &SGraphPin::GetPinVisiblity);
 
@@ -167,35 +167,35 @@ void SGraphPin::Construct(const FArguments& InArgs, UEdGraphPin* InPin)
 	const bool bIsInput = (GetDirection() == EGPD_Input);
 
 	// Create the pin icon widget
-	TSharedRef<SWidget> ActualPinWidget = 
+	TSharedRef<SWidget> ActualPinWidget =
 		SAssignNew(PinImage, SImage)
-		.Image( this, &SGraphPin::GetPinIcon )
-		.IsEnabled( bCanConnectToPin )
+		.Image(this, &SGraphPin::GetPinIcon)
+		.IsEnabled(bCanConnectToPin)
 		.ColorAndOpacity(this, &SGraphPin::GetPinColor)
-		.OnMouseButtonDown( this, &SGraphPin::OnPinMouseDown )
-		.Cursor( this, &SGraphPin::GetPinCursor );
+		.OnMouseButtonDown(this, &SGraphPin::OnPinMouseDown)
+		.Cursor(this, &SGraphPin::GetPinCursor);
 
 	// Create the pin indicator widget (used for watched values)
 	static const FName NAME_NoBorder("NoBorder");
 	TSharedRef<SWidget> PinStatusIndicator =
 		SNew(SButton)
-		.ButtonStyle( FEditorStyle::Get(), NAME_NoBorder )
-		.Visibility( this, &SGraphPin::GetPinStatusIconVisibility )
+		.ButtonStyle(FEditorStyle::Get(), NAME_NoBorder)
+		.Visibility(this, &SGraphPin::GetPinStatusIconVisibility)
 		.ContentPadding(0)
-		.OnClicked( this, &SGraphPin::ClickedOnPinStatusIcon )
+		.OnClicked(this, &SGraphPin::ClickedOnPinStatusIcon)
 		[
 			SNew(SImage)
-			.Image( this, &SGraphPin::GetPinStatusIcon )
+			.Image(this, &SGraphPin::GetPinStatusIcon)
 		];
 
 	TSharedRef<SWidget> LabelWidget = SNew(STextBlock)
 		.Text(this, &SGraphPin::GetPinLabel)
-		.TextStyle( FEditorStyle::Get(), InArgs._PinLabelStyle )
+		.TextStyle(FEditorStyle::Get(), InArgs._PinLabelStyle)
 		.Visibility(this, &SGraphPin::GetPinLabelVisibility)
 		.ColorAndOpacity(this, &SGraphPin::GetPinTextColor);
 
 	// Create the widget used for the pin body (status indicator, label, and value)
-	TSharedRef<SWrapBox> LabelAndValue = 
+	TSharedRef<SWrapBox> LabelAndValue =
 		SNew(SWrapBox)
 		.PreferredWidth(150.f);
 
@@ -245,18 +245,11 @@ void SGraphPin::Construct(const FArguments& InArgs, UEdGraphPin* InPin)
 			];
 	}
 
-	TAttribute<FText> ToolTipAttribute;
-	if (InArgs._HasToolTip)
-	{
-		ToolTipAttribute = TAttribute<FText>::Create(TAttribute<FText>::FGetter::CreateSP(this, &SGraphPin::GetTooltip));
-	}
-
 	TSharedPtr<SWidget> PinContent;
 	if (bIsInput)
 	{
 		// Input pin
 		PinContent = SNew(SHorizontalBox)
-			.ToolTipText( ToolTipAttribute )
 			+SHorizontalBox::Slot()
 			.AutoWidth()
 			.VAlign(VAlign_Center)
@@ -264,7 +257,7 @@ void SGraphPin::Construct(const FArguments& InArgs, UEdGraphPin* InPin)
 			[
 				ActualPinWidget
 			]
-			+ SHorizontalBox::Slot()
+			+SHorizontalBox::Slot()
 			.AutoWidth()
 			.VAlign(VAlign_Center)
 			[
@@ -275,7 +268,6 @@ void SGraphPin::Construct(const FArguments& InArgs, UEdGraphPin* InPin)
 	{
 		// Output pin
 		PinContent = SNew(SHorizontalBox)
-			.ToolTipText( ToolTipAttribute )
 			+SHorizontalBox::Slot()
 			.AutoWidth()
 			.VAlign(VAlign_Center)
@@ -292,10 +284,10 @@ void SGraphPin::Construct(const FArguments& InArgs, UEdGraphPin* InPin)
 	}
 
 	// Set up a hover for pins that is tinted the color of the pin.
-	SBorder::Construct( SBorder::FArguments()
-		.BorderImage( this, &SGraphPin::GetPinBorder )
-		.BorderBackgroundColor( this, &SGraphPin::GetPinColor )
-		.OnMouseButtonDown( this, &SGraphPin::OnPinNameMouseDown )
+	SBorder::Construct(SBorder::FArguments()
+		.BorderImage(this, &SGraphPin::GetPinBorder)
+		.BorderBackgroundColor(this, &SGraphPin::GetPinColor)
+		.OnMouseButtonDown(this, &SGraphPin::OnPinNameMouseDown)
 		[
 			SNew(SLevelOfDetailBranchNode)
 			.UseLowDetailSlot(this, &SGraphPin::UseLowDetailPinNames)
@@ -305,11 +297,14 @@ void SGraphPin::Construct(const FArguments& InArgs, UEdGraphPin* InPin)
 				ActualPinWidget
 			]
 			.HighDetail()
-			[
-				PinContent.ToSharedRef()
-			]
+				[
+					PinContent.ToSharedRef()
+				]
 		]
 	);
+
+	TAttribute<FText> ToolTipAttribute = TAttribute<FText>::Create(TAttribute<FText>::FGetter::CreateSP(this, &SGraphPin::GetTooltipText));
+	SetToolTipText(ToolTipAttribute);
 }
 
 
@@ -965,7 +960,7 @@ void SGraphPin::SetShowLabel(bool bNewShowLabel)
 	bShowLabel = bNewShowLabel;
 }
 
-FText SGraphPin::GetTooltip() const
+FText SGraphPin::GetTooltipText() const
 {
 	FText HoverText = FText::GetEmpty();
 
diff --git a/Engine/Source/Editor/GraphEditor/Public/SGraphPin.h b/Engine/Source/Editor/GraphEditor/Public/SGraphPin.h
--- a/Engine/Source/Editor/GraphEditor/Public/SGraphPin.h
+++ b/Engine/Source/Editor/GraphEditor/Public/SGraphPin.h
@@ -16,12 +16,10 @@ public:
 
 	SLATE_BEGIN_ARGS(SGraphPin)
 		: _PinLabelStyle(NAME_DefaultPinLabelStyle)
-		, _HasToolTip(true)
 		, _UsePinColorForText(false)
 		, _SideToSideMargin(5.0f)
 		{}
 		SLATE_ARGUMENT(FName, PinLabelStyle)
-		SLATE_ARGUMENT(bool, HasToolTip)
 		SLATE_ARGUMENT(bool, UsePinColorForText)
 		SLATE_ARGUMENT(float, SideToSideMargin)
 	SLATE_END_ARGS()
@@ -143,7 +141,7 @@ protected:
 	virtual FSlateColor GetPinTextColor() const;
 
 	/** @return The tooltip to display for this pin */
-	FText GetTooltip() const;
+	FText GetTooltipText() const;
 
 	TOptional<EMouseCursor::Type> GetPinCursor() const;
 
