Graph Editor: Added a TODO on some duplicated constants

diff --git a/Engine/Source/Editor/GraphEditor/Private/ConnectionDrawingPolicy.cpp b/Engine/Source/Editor/GraphEditor/Private/ConnectionDrawingPolicy.cpp
--- a/Engine/Source/Editor/GraphEditor/Private/ConnectionDrawingPolicy.cpp
+++ b/Engine/Source/Editor/GraphEditor/Private/ConnectionDrawingPolicy.cpp
@@ -497,6 +497,7 @@ void FConnectionDrawingPolicy::ResetIncompatiblePinDrawState(const TSet< TShared
 
 void FConnectionDrawingPolicy::ApplyHoverDeemphasis(UEdGraphPin* OutputPin, UEdGraphPin* InputPin, /*inout*/ float& Thickness, /*inout*/ FLinearColor& WireColor)
 {
+	//@TODO: Move these parameters into the settings object
 	const float FadeInBias = 0.75f; // Time in seconds before the fading starts to occur
 	const float FadeInPeriod = 0.6f; // Time in seconds after the bias before the fade is fully complete
 	const float TimeFraction = FMath::SmoothStep(0.0f, FadeInPeriod, (float)(FSlateApplication::Get().GetCurrentTime() - LastHoverTimeEvent - FadeInBias));
