Graph Editor: Split out SGraphPinInteger from SGraphPinNum to push schema checking into the NodeFactory where it belongs
- Fixed SGraphPinNum::Construct skipping its direct parent

diff --git a/Engine/Source/Editor/GraphEditor/Private/KismetPins/SGraphPinInteger.cpp b/Engine/Source/Editor/GraphEditor/Private/KismetPins/SGraphPinInteger.cpp
new file mode 100644
--- /dev/null
+++ b/Engine/Source/Editor/GraphEditor/Private/KismetPins/SGraphPinInteger.cpp
@@ -0,0 +1,31 @@
+// Copyright 1998-2015 Epic Games, Inc. All Rights Reserved.
+
+#include "GraphEditorCommon.h"
+#include "SGraphPin.h"
+#include "SGraphPinString.h"
+#include "SGraphPinNum.h"
+#include "SGraphPinInteger.h"
+#include "DefaultValueHelper.h"
+#include "Editor/UnrealEd/Public/ScopedTransaction.h"
+
+void SGraphPinInteger::Construct(const FArguments& InArgs, UEdGraphPin* InGraphPinObj)
+{
+	SGraphPinNum::Construct(SGraphPinNum::FArguments(), InGraphPinObj);
+}
+
+void SGraphPinInteger::SetTypeInValue(const FText& NewTypeInValue, ETextCommit::Type /*CommitInfo*/)
+{
+	const FString TypeValueString = NewTypeInValue.ToString();
+	if (FDefaultValueHelper::IsStringValidFloat(TypeValueString) || FDefaultValueHelper::IsStringValidInteger(TypeValueString))
+	{
+		if (GraphPinObj->GetDefaultAsString() != TypeValueString)
+		{
+			const FScopedTransaction Transaction(NSLOCTEXT("GraphEditor", "ChangeNumberPinValue", "Change Number Pin Value"));
+			GraphPinObj->Modify();
+
+			// Round-tripped here to allow floating point values to be pasted and truncated rather than failing to be set at all
+			const int32 IntValue = FCString::Atoi(*TypeValueString);
+			GraphPinObj->GetSchema()->TrySetDefaultValue(*GraphPinObj, FString::FromInt(IntValue));
+		}
+	}
+}
diff --git a/Engine/Source/Editor/GraphEditor/Private/KismetPins/SGraphPinInteger.h b/Engine/Source/Editor/GraphEditor/Private/KismetPins/SGraphPinInteger.h
new file mode 100644
--- /dev/null
+++ b/Engine/Source/Editor/GraphEditor/Private/KismetPins/SGraphPinInteger.h
@@ -0,0 +1,17 @@
+// Copyright 1998-2015 Epic Games, Inc. All Rights Reserved.
+
+#pragma once
+
+class SGraphPinInteger : public SGraphPinNum
+{
+public:
+	SLATE_BEGIN_ARGS(SGraphPinNum) {}
+	SLATE_END_ARGS()
+
+	void Construct(const FArguments& InArgs, UEdGraphPin* InGraphPinObj);
+
+protected:
+	// Begin SGraphPinString interface
+	virtual void SetTypeInValue(const FText& NewTypeInValue, ETextCommit::Type CommitInfo) override;
+	// End of SGraphPinString interface
+};
diff --git a/Engine/Source/Editor/GraphEditor/Private/KismetPins/SGraphPinNum.cpp b/Engine/Source/Editor/GraphEditor/Private/KismetPins/SGraphPinNum.cpp
--- a/Engine/Source/Editor/GraphEditor/Private/KismetPins/SGraphPinNum.cpp
+++ b/Engine/Source/Editor/GraphEditor/Private/KismetPins/SGraphPinNum.cpp
@@ -10,31 +10,20 @@
 
 void SGraphPinNum::Construct(const FArguments& InArgs, UEdGraphPin* InGraphPinObj)
 {
-	SGraphPin::Construct(SGraphPin::FArguments(), InGraphPinObj);
+	SGraphPinString::Construct(SGraphPinString::FArguments(), InGraphPinObj);
 }
 
 void SGraphPinNum::SetTypeInValue(const FText& NewTypeInValue, ETextCommit::Type /*CommitInfo*/)
 {
 	FString TypeValueString = NewTypeInValue.ToString();
-	if (FDefaultValueHelper::IsStringValidFloat(TypeValueString))
+	if (FDefaultValueHelper::IsStringValidFloat(TypeValueString) || FDefaultValueHelper::IsStringValidInteger(TypeValueString))
 	{
 		if(GraphPinObj->GetDefaultAsString() != TypeValueString)
 		{
 			const FScopedTransaction Transaction( NSLOCTEXT("GraphEditor", "ChangeNumberPinValue", "Change Number Pin Value" ) );
 			GraphPinObj->Modify();
 
-			//This is rubbish.
-			//TODO: Make not rubbish.
-			const UEdGraphSchema_K2* K2Schema = Cast<const UEdGraphSchema_K2>(GraphPinObj->GetSchema());
-			if ((K2Schema && (GraphPinObj->PinType.PinCategory == K2Schema->PC_Int || GraphPinObj->PinType.PinCategory == K2Schema->PC_Byte)))
-			{
-				int32 IntValue = FCString::Atoi(*TypeValueString);
-				GraphPinObj->GetSchema()->TrySetDefaultValue(*GraphPinObj, FString::FromInt(IntValue));
-			}
-			else
-			{
-				GraphPinObj->GetSchema()->TrySetDefaultValue(*GraphPinObj, *TypeValueString);
-			}
+			GraphPinObj->GetSchema()->TrySetDefaultValue(*GraphPinObj, *TypeValueString);
 		}
 	}
 }
diff --git a/Engine/Source/Editor/GraphEditor/Private/KismetPins/SGraphPinNum.h b/Engine/Source/Editor/GraphEditor/Private/KismetPins/SGraphPinNum.h
--- a/Engine/Source/Editor/GraphEditor/Private/KismetPins/SGraphPinNum.h
+++ b/Engine/Source/Editor/GraphEditor/Private/KismetPins/SGraphPinNum.h
@@ -11,7 +11,7 @@ public:
 	void Construct(const FArguments& InArgs, UEdGraphPin* InGraphPinObj);
 
 protected:
-	// Begin SGraphPinText interface
+	// SGraphPinString interface
 	virtual void SetTypeInValue(const FText& NewTypeInValue, ETextCommit::Type CommitInfo) override;
-	// End SGraphPinText interface
+	// End of SGraphPinString interface
 };
diff --git a/Engine/Source/Editor/GraphEditor/Private/NodeFactory.cpp b/Engine/Source/Editor/GraphEditor/Private/NodeFactory.cpp
--- a/Engine/Source/Editor/GraphEditor/Private/NodeFactory.cpp
+++ b/Engine/Source/Editor/GraphEditor/Private/NodeFactory.cpp
@@ -57,6 +57,7 @@
 #include "KismetPins/SGraphPinClass.h"
 #include "KismetPins/SGraphPinExec.h"
 #include "KismetPins/SGraphPinNum.h"
+#include "KismetPins/SGraphPinInteger.h"
 #include "KismetPins/SGraphPinColor.h"
 #include "KismetPins/SGraphPinEnum.h"
 #include "KismetPins/SGraphPinKey.h"
@@ -293,7 +294,11 @@ TSharedPtr<SGraphPin> FNodeFactory::CreatePinWidget(UEdGraphPin* InPin)
 		{
 			return SNew(SGraphPinClass, InPin);
 		}
-		else if ((InPin->PinType.PinCategory == K2Schema->PC_Int) || (InPin->PinType.PinCategory == K2Schema->PC_Float))
+		else if (InPin->PinType.PinCategory == K2Schema->PC_Int)
+		{
+			return SNew(SGraphPinInteger, InPin);
+		}
+		else if (InPin->PinType.PinCategory == K2Schema->PC_Float)
 		{
 			return SNew(SGraphPinNum, InPin);
 		}
@@ -343,7 +348,7 @@ TSharedPtr<SGraphPin> FNodeFactory::CreatePinWidget(UEdGraphPin* InPin)
 			}
 			else
 			{
-				return SNew(SGraphPinNum, InPin);
+				return SNew(SGraphPinInteger, InPin);
 			}
 		}
 		else if ((InPin->PinType.PinCategory == K2Schema->PC_Wildcard) && (InPin->PinType.PinSubCategory == K2Schema->PSC_Index))
