Graph Editor: Enable direct spline interactions by default

diff --git a/Engine/Source/Editor/GraphEditor/Private/ConnectionDrawingPolicy.cpp b/Engine/Source/Editor/GraphEditor/Private/ConnectionDrawingPolicy.cpp
--- a/Engine/Source/Editor/GraphEditor/Private/ConnectionDrawingPolicy.cpp
+++ b/Engine/Source/Editor/GraphEditor/Private/ConnectionDrawingPolicy.cpp
@@ -237,10 +237,10 @@ void FConnectionDrawingPolicy::DrawConnection(int32 LayerId, const FVector2D& St
 	const FVector2D P0Tangent = (Params.StartDirection == EGPD_Output) ? SplineTangent : -SplineTangent;
 	const FVector2D P1Tangent = (Params.EndDirection == EGPD_Input) ? SplineTangent : -SplineTangent;
 
-	if (Settings->bTreatSplinesLikePinsEXPERIMENTAL)
+	if (Settings->bTreatSplinesLikePins)
 	{
 		// Distance to consider as an overlap
-		const float QueryDistanceTriggerThresholdSquared = FMath::Square(Settings->SplineHoverToleranceEXPERIMENTAL + Params.WireThickness * 0.5f);
+		const float QueryDistanceTriggerThresholdSquared = FMath::Square(Settings->SplineHoverTolerance + Params.WireThickness * 0.5f);
 
 		// Distance to pass the bounding box cull test (may want to expand this later on if we want to do 'closest pin' actions that don't require an exact hit)
 		const float QueryDistanceToBoundingBoxSquared = QueryDistanceTriggerThresholdSquared;
diff --git a/Engine/Source/Editor/GraphEditor/Private/GraphEditorSettings.cpp b/Engine/Source/Editor/GraphEditor/Private/GraphEditorSettings.cpp
--- a/Engine/Source/Editor/GraphEditor/Private/GraphEditorSettings.cpp
+++ b/Engine/Source/Editor/GraphEditor/Private/GraphEditorSettings.cpp
@@ -13,8 +13,8 @@ UGraphEditorSettings::UGraphEditorSettings( const FObjectInitializer& ObjectInit
 	, PaddingRightOfInput(10.0f)
 	, PaddingLeftOfOutput(10.0f)
 	, PaddingTowardsNodeEdge(10.0f)
-	, bTreatSplinesLikePinsEXPERIMENTAL(false)
-	, SplineHoverToleranceEXPERIMENTAL(2.0f)
+	, bTreatSplinesLikePins(true)
+	, SplineHoverTolerance(2.0f)
 	, ForwardSplineHorizontalDeltaRange(1000.0f)
 	, ForwardSplineVerticalDeltaRange(1000.0f)
 	, ForwardSplineTangentFromHorizontalDelta(1.0f, 0.0f)
diff --git a/Engine/Source/Editor/GraphEditor/Public/GraphEditorSettings.h b/Engine/Source/Editor/GraphEditor/Public/GraphEditorSettings.h
--- a/Engine/Source/Editor/GraphEditor/Public/GraphEditorSettings.h
+++ b/Engine/Source/Editor/GraphEditor/Public/GraphEditorSettings.h
@@ -64,11 +64,11 @@ public:
 
 	/** If enabled, allows splines to be Alt+Clicked to break them or Ctrl+Dragged to move them as if these actions were taking place on the associated pin. */
 	UPROPERTY(EditAnywhere, config, Category=Splines)
-	bool bTreatSplinesLikePinsEXPERIMENTAL;
+	bool bTreatSplinesLikePins;
 
 	/** The distance threshold controlling how close the mouse has to be to the spline in order to trigger a hover response */ 
-	UPROPERTY(EditAnywhere, config, Category=Splines, meta=(EditCondition=bTreatSplinesLikePinsEXPERIMENTAL))
-	float SplineHoverToleranceEXPERIMENTAL;
+	UPROPERTY(EditAnywhere, config, Category=Splines, meta=(EditCondition=bTreatSplinesLikePins))
+	float SplineHoverTolerance;
 
 	/** The maximum value to clamp the absolute value of the horizontal distance between endpoints when calculating tangents (when the wire is moving forward) */
 	UPROPERTY(config, EditAnywhere, Category=Splines, AdvancedDisplay)
