Graph Editor: Changed how spline hovering stores data, avoiding direct pin pointers beyond the frame
[UE-13270]

diff --git a/Engine/Source/Editor/GraphEditor/Private/ConnectionDrawingPolicy.cpp b/Engine/Source/Editor/GraphEditor/Private/ConnectionDrawingPolicy.cpp
--- a/Engine/Source/Editor/GraphEditor/Private/ConnectionDrawingPolicy.cpp
+++ b/Engine/Source/Editor/GraphEditor/Private/ConnectionDrawingPolicy.cpp
@@ -309,7 +309,7 @@ void FConnectionDrawingPolicy::DrawConnection(int32 LayerId, const FVector2D& St
 			// Record the overlap
 			if (ClosestDistanceSquared < QueryDistanceTriggerThresholdSquared)
 			{
-				if (ClosestDistanceSquared < SplineOverlapResult.DistanceSquared)
+				if (ClosestDistanceSquared < SplineOverlapResult.GetDistanceSquared())
 				{
 					const float SquaredDistToPin1 = (Params.AssociatedPin1 != nullptr) ? (P0 - ClosestPoint).SizeSquared() : FLT_MAX;
 					const float SquaredDistToPin2 = (Params.AssociatedPin2 != nullptr) ? (P1 - ClosestPoint).SizeSquared() : FLT_MAX;
@@ -525,6 +525,7 @@ void FConnectionDrawingPolicy::ApplyHoverDeemphasis(UEdGraphPin* OutputPin, UEdG
 
 void FGraphSplineOverlapResult::ComputeBestPin()
 {
+	UEdGraphPin* BestPin = nullptr;
 	if (Pin1 == nullptr)
 	{
 		BestPin = Pin2;
@@ -553,4 +554,34 @@ void FGraphSplineOverlapResult::ComputeBestPin()
 			BestPin = (DistanceSquaredToPin1 < DistanceSquaredToPin2) ? Pin1 : Pin2;
 		}
 	}
+
+	BestPinHandle = FGraphPinHandle(BestPin);
+
+	Pin1 = nullptr;
+	Pin2 = nullptr;
+}
+
+bool FGraphSplineOverlapResult::GetPins(const class SGraphPanel& InGraphPanel, UEdGraphPin*& OutPin1, UEdGraphPin*& OutPin2) const
+{
+	OutPin1 = nullptr;
+	OutPin2 = nullptr;
+
+	if (IsValid())
+	{
+		TSharedPtr<SGraphPin> Pin1Widget = Pin1Handle.FindInGraphPanel(InGraphPanel);
+		TSharedPtr<SGraphPin> Pin2Widget = Pin2Handle.FindInGraphPanel(InGraphPanel);
+
+		if (Pin1Widget.IsValid())
+		{
+			OutPin1 = Pin1Widget->GetPinObj();
+		}
+
+		if (Pin2Widget.IsValid())
+		{
+			OutPin2 = Pin2Widget->GetPinObj();
+		}
+	}
+
+	return (OutPin1 != nullptr) && (OutPin2 != nullptr);
 }
+
diff --git a/Engine/Source/Editor/GraphEditor/Private/SGraphPanel.cpp b/Engine/Source/Editor/GraphEditor/Private/SGraphPanel.cpp
--- a/Engine/Source/Editor/GraphEditor/Private/SGraphPanel.cpp
+++ b/Engine/Source/Editor/GraphEditor/Private/SGraphPanel.cpp
@@ -30,38 +30,44 @@
 
 DEFINE_LOG_CATEGORY_STATIC(LogGraphPanel, Log, All);
 
-SGraphPanel::FGraphPinHandle::FGraphPinHandle(UEdGraphPin* InPin)
+//////////////////////////////////////////////////////////////////////////
+// FGraphPinHandle
+
+FGraphPinHandle::FGraphPinHandle(UEdGraphPin* InPin)
 {
-	if (auto* Node = InPin->GetOwningNode())
+	if (InPin != nullptr)
 	{
-		PinName = InPin->PinName;
-		NodeGuid = Node->NodeGuid;
+		if (UEdGraphNode* Node = InPin->GetOwningNodeUnchecked())
+		{
+			PinName = InPin->PinName;
+			NodeGuid = Node->NodeGuid;
+		}
 	}
 }
 
-TSharedPtr<SGraphPin> SGraphPanel::FGraphPinHandle::FindInGraphPanel(const SGraphPanel& InPanel) const
+TSharedPtr<SGraphPin> FGraphPinHandle::FindInGraphPanel(const SGraphPanel& InPanel) const
 {
 	// First off, find the node
-	TSharedPtr<SGraphNode> GraphNode = InPanel.GetNodeWidgetFromGuid(NodeGuid);
-	if (GraphNode.IsValid())
+	if (NodeGuid.IsValid())
 	{
-		UEdGraphNode* Node = GraphNode->GetNodeObj();
-		UEdGraphPin* Pin = Node->FindPin(PinName);
-
-		if (Pin)
+		TSharedPtr<SGraphNode> GraphNode = InPanel.GetNodeWidgetFromGuid(NodeGuid);
+		if (GraphNode.IsValid())
 		{
-			return GraphNode->FindWidgetForPin(Pin);
+			UEdGraphNode* Node = GraphNode->GetNodeObj();
+
+			if (UEdGraphPin* Pin = Node->FindPin(PinName))
+			{
+				return GraphNode->FindWidgetForPin(Pin);
+			}
 		}
 	}
 
 	return TSharedPtr<SGraphPin>();
 }
 
-/**
- * Construct a widget
- *
- * @param InArgs    The declaration describing how the widgets should be constructed.
- */
+//////////////////////////////////////////////////////////////////////////
+// SGraphPanel
+
 void SGraphPanel::Construct( const SGraphPanel::FArguments& InArgs )
 {
 	SNodePanel::Construct();
@@ -517,39 +523,27 @@ int32 SGraphPanel::OnPaint( const FPaintArgs& Args, const FGeometry& AllottedGeo
 
 void SGraphPanel::OnSplineHoverStateChanged(const FGraphSplineOverlapResult& NewSplineHoverState)
 {
-	UEdGraphPin* OldPin = PreviousFrameSplineOverlap.GetBestPin();
+	TSharedPtr<SGraphPin> OldPinWidget = PreviousFrameSplineOverlap.GetBestPinWidget(*this);
 	PreviousFrameSplineOverlap = NewSplineHoverState;
-	UEdGraphPin* NewPin = PreviousFrameSplineOverlap.GetBestPin();
+	TSharedPtr<SGraphPin> NewPinWidget = PreviousFrameSplineOverlap.GetBestPinWidget(*this);
 
 	PreviousFrameSavedMousePosForSplineOverlap = SavedMousePosForOnPaintEventLocalSpace;
 
 	// Handle mouse enter/leaves on the associated pin
-	if (OldPin != NewPin)
+	if (OldPinWidget != NewPinWidget)
 	{
-		if (OldPin != nullptr)
+		if (OldPinWidget.IsValid())
 		{
-			FGraphPinHandle OldPinHandle(OldPin);
-			TSharedPtr<SGraphPin> OldPinWidget = OldPinHandle.FindInGraphPanel(*this);
-
-			if (OldPinWidget.IsValid())
-			{
-				OldPinWidget->OnMouseLeave(LastPointerEvent);
-			}
+			OldPinWidget->OnMouseLeave(LastPointerEvent);
 		}
 
-		if (NewPin != nullptr)
+		if (NewPinWidget.IsValid())
 		{
-			FGraphPinHandle NewPinHandle(NewPin);
-			TSharedPtr<SGraphPin> NewPinWidget = NewPinHandle.FindInGraphPanel(*this);
-
-			if (NewPinWidget.IsValid())
-			{
-				NewPinWidget->OnMouseEnter(LastPointerGeometry, LastPointerEvent);
+			NewPinWidget->OnMouseEnter(LastPointerGeometry, LastPointerEvent);
 
-				// Get the pin/wire glowing quicker, since it's a direct selection (this time was already set to 'now' as part of entering the pin)
-				//@TODO: Source this parameter from the graph rendering settings once it is there (see code in ApplyHoverDeemphasis)
-				TimeWhenMouseEnteredPin -= 0.75f;
-			}
+			// Get the pin/wire glowing quicker, since it's a direct selection (this time was already set to 'now' as part of entering the pin)
+			//@TODO: Source this parameter from the graph rendering settings once it is there (see code in ApplyHoverDeemphasis)
+			TimeWhenMouseEnteredPin -= 0.75f;
 		}
 	}
 }
@@ -673,13 +667,15 @@ FReply SGraphPanel::OnMouseButtonUp(const FGeometry& MyGeometry, const FPointerE
 
 FReply SGraphPanel::OnMouseButtonDoubleClick(const FGeometry& MyGeometry, const FPointerEvent& MouseEvent)
 {
-	if (PreviousFrameSplineOverlap.IsValid() && (PreviousFrameSplineOverlap.Pin1 != nullptr) && (PreviousFrameSplineOverlap.Pin2 != nullptr))
+	UEdGraphPin* Pin1;
+	UEdGraphPin* Pin2;
+	if (PreviousFrameSplineOverlap.GetPins(*this, /*out*/ Pin1, /*out*/ Pin2))
 	{
 		// Give the schema a chance to do something interesting with a double click on a proper spline (both ends are attached to a pin, i.e., not a preview/drag one)
 		const FVector2D DoubleClickPositionInGraphSpace = PanelCoordToGraphCoord(MyGeometry.AbsoluteToLocal(MouseEvent.GetScreenSpacePosition()));
 
 		const UEdGraphSchema* Schema = GraphObj->GetSchema();
-		Schema->OnPinConnectionDoubleCicked(PreviousFrameSplineOverlap.Pin1, PreviousFrameSplineOverlap.Pin2, DoubleClickPositionInGraphSpace);
+		Schema->OnPinConnectionDoubleCicked(Pin1, Pin2, DoubleClickPositionInGraphSpace);
 	}
 
 	return SNodePanel::OnMouseButtonDoubleClick(MyGeometry, MouseEvent);
@@ -687,15 +683,8 @@ FReply SGraphPanel::OnMouseButtonDoubleClick(const FGeometry& MyGeometry, const
 
 class SGraphPin* SGraphPanel::GetBestPinFromHoveredSpline() const
 {
-	if (PreviousFrameSplineOverlap.IsValid())
-	{
-		FGraphPinHandle SplineAssociatedPinHandle(PreviousFrameSplineOverlap.GetBestPin());
-		return SplineAssociatedPinHandle.FindInGraphPanel(*this).Get();
-	}
-	else
-	{
-		return nullptr;
-	}
+	TSharedPtr<SGraphPin> BestPinWidget = PreviousFrameSplineOverlap.GetBestPinWidget(*this);
+	return BestPinWidget.Get();
 }
 
 void SGraphPanel::GetAllPins(TSet< TSharedRef<SWidget> >& AllPins)
@@ -1191,6 +1180,10 @@ void SGraphPanel::Update()
 // Purges the existing visual representation (typically followed by an Update call in the next tick)
 void SGraphPanel::PurgeVisualRepresentation()
 {
+	// No need to call OnSplineHoverStateChanged since we're about to destroy all the nodes and pins
+	PreviousFrameSplineOverlap = FGraphSplineOverlapResult();
+
+	// Clear all of the nodes and pins
 	RemoveAllNodes();
 }
 
diff --git a/Engine/Source/Editor/GraphEditor/Private/SGraphPanel.h b/Engine/Source/Editor/GraphEditor/Private/SGraphPanel.h
--- a/Engine/Source/Editor/GraphEditor/Private/SGraphPanel.h
+++ b/Engine/Source/Editor/GraphEditor/Private/SGraphPanel.h
@@ -67,22 +67,6 @@ public:
 		//SLATE_ATTRIBUTE( FGraphAppearanceInfo, Appearance )
 	SLATE_END_ARGS()
 
-	/** A handle to a pin, defined by its owning node's GUID, and the pin's name. Used to reference a pin without referring to its widget */
-	struct FGraphPinHandle
-	{
-		/** The GUID of the node to which this pin belongs */
-		FGuid NodeGuid;
-
-		/** The name of the pin we are referencing */
-		FString PinName;
-
-		/** Constructor */
-		FGraphPinHandle(UEdGraphPin* InPin);
-
-		/** Find a pin widget in the specified panel from this handle */
-		TSharedPtr<SGraphPin> FindInGraphPanel(const SGraphPanel& InPanel) const;
-	};
-
 	/**
 	 * Construct a widget
 	 *
diff --git a/Engine/Source/Editor/GraphEditor/Public/GraphSplineOverlapResult.h b/Engine/Source/Editor/GraphEditor/Public/GraphSplineOverlapResult.h
--- a/Engine/Source/Editor/GraphEditor/Public/GraphSplineOverlapResult.h
+++ b/Engine/Source/Editor/GraphEditor/Public/GraphSplineOverlapResult.h
@@ -5,19 +5,34 @@
 class UEdGraphPin;
 
 /////////////////////////////////////////////////////
-// FGraphSplineOverlapResult
+// FGraphPinHandle
 
-struct FGraphSplineOverlapResult
+/** A handle to a pin, defined by its owning node's GUID, and the pin's name. Used to reference a pin without referring to its widget */
+struct GRAPHEDITOR_API FGraphPinHandle
 {
-	UEdGraphPin* BestPin;
-	UEdGraphPin* Pin1;
-	UEdGraphPin* Pin2;
-	float DistanceSquared;
-	float DistanceSquaredToPin1;
-	float DistanceSquaredToPin2;
+	/** The GUID of the node to which this pin belongs */
+	FGuid NodeGuid;
+
+	/** The name of the pin we are referencing */
+	FString PinName;
+
+	/** Constructor */
+	FGraphPinHandle(UEdGraphPin* InPin);
+
+	/** Find a pin widget in the specified panel from this handle */
+	TSharedPtr<class SGraphPin> FindInGraphPanel(const class SGraphPanel& InPanel) const;
+};
+
+/////////////////////////////////////////////////////
+// FGraphSplineOverlapResult
 
+struct GRAPHEDITOR_API FGraphSplineOverlapResult
+{
+public:
 	FGraphSplineOverlapResult()
-		: BestPin(nullptr)
+		: Pin1Handle(nullptr)
+		, Pin2Handle(nullptr)
+		, BestPinHandle(nullptr)
 		, Pin1(nullptr)
 		, Pin2(nullptr)
 		, DistanceSquared(FLT_MAX)
@@ -27,7 +42,9 @@ struct FGraphSplineOverlapResult
 	}
 
 	FGraphSplineOverlapResult(UEdGraphPin* InPin1, UEdGraphPin* InPin2, float InDistanceSquared, float InDistanceSquaredToPin1, float InDistanceSquaredToPin2)
-		: BestPin(nullptr)
+		: Pin1Handle(InPin1)
+		, Pin2Handle(InPin2)
+		, BestPinHandle(nullptr)
 		, Pin1(InPin1)
 		, Pin2(InPin2)
 		, DistanceSquared(InDistanceSquared)
@@ -41,10 +58,32 @@ struct FGraphSplineOverlapResult
 		return DistanceSquared < FLT_MAX;
 	}
 
-	GRAPHEDITOR_API void ComputeBestPin();
+	void ComputeBestPin();
+
+	float GetDistanceSquared() const
+	{
+		return DistanceSquared;
+	}
 
-	UEdGraphPin* GetBestPin() const
+	TSharedPtr<class SGraphPin> GetBestPinWidget(const class SGraphPanel& InGraphPanel) const
 	{
-		return IsValid() ? BestPin : nullptr;
+		TSharedPtr<class SGraphPin> Result;
+		if (IsValid())
+		{
+			Result = BestPinHandle.FindInGraphPanel(InGraphPanel);
+		}
+		return Result;
 	}
+
+	bool GetPins(const class SGraphPanel& InGraphPanel, UEdGraphPin*& OutPin1, UEdGraphPin*& OutPin2) const;
+
+protected:
+	FGraphPinHandle Pin1Handle;
+	FGraphPinHandle Pin2Handle;
+	FGraphPinHandle BestPinHandle;
+	UEdGraphPin* Pin1;
+	UEdGraphPin* Pin2;
+	float DistanceSquared;
+	float DistanceSquaredToPin1;
+	float DistanceSquaredToPin2;
 };
