This change prevents comment balloon toggle button from displaying inside read only graphs.
#UE-13295 - Comment balloons can be toggled in read only graphs
#Branch UE4
#Proj GraphEditor

diff --git a/Engine/Source/Editor/GraphEditor/Private/SCommentBubble.cpp b/Engine/Source/Editor/GraphEditor/Private/SCommentBubble.cpp
--- a/Engine/Source/Editor/GraphEditor/Private/SCommentBubble.cpp
+++ b/Engine/Source/Editor/GraphEditor/Private/SCommentBubble.cpp
@@ -78,7 +78,10 @@ void SCommentBubble::Tick( const FGeometry& AllottedGeometry, const double InCur
 {
 	SCompoundWidget::Tick( AllottedGeometry, InCurrentTime, InDeltaTime );
 
-	const bool bTitleBarBubbleVisible = bEnableTitleBarBubble && IsGraphNodeHovered.IsBound();
+	// Check Editable and Hovered so we can prevent bubble toggling in read only graphs.
+	const bool bNodeEditable = GraphNode->NodeWidget.IsValid() ? GraphNode->NodeWidget.Pin()->IsNodeEditable() : false;
+	const bool bEnableTitleHintBubble = bEnableTitleBarBubble && bNodeEditable;
+	const bool bTitleBarBubbleVisible = bEnableTitleHintBubble && IsGraphNodeHovered.IsBound();
 
 	if( bTitleBarBubbleVisible || IsBubbleVisible() )
 	{
