Updated the comment balloon text hint to make use of the recently implemented interface in SMultiLineEditableTextBox.
#UE-13108 - Blueprint Editor: Problems with editing comment balloons inline
#Branch UE4
#Proj GraphEditor

diff --git a/Engine/Source/Editor/GraphEditor/Private/SCommentBubble.cpp b/Engine/Source/Editor/GraphEditor/Private/SCommentBubble.cpp
--- a/Engine/Source/Editor/GraphEditor/Private/SCommentBubble.cpp
+++ b/Engine/Source/Editor/GraphEditor/Private/SCommentBubble.cpp
@@ -318,10 +318,6 @@ bool SCommentBubble::IsScalingAllowed() const
 
 FText SCommentBubble::GetCommentText() const
 {
-	if( CommentAttribute.Get().IsEmpty() )
-	{
-		return HintText.Get();
-	}
 	return CachedCommentText;
 }
 
@@ -378,8 +374,7 @@ FSlateColor SCommentBubble::GetTextForegroundColor() const
 
 void SCommentBubble::OnCommentTextCommitted( const FText& NewText, ETextCommit::Type CommitInfo )
 {
-	const bool bValidText = !NewText.EqualTo( HintText.Get() );
-	if( GraphNode && bValidText )
+	if( GraphNode )
 	{
 		GraphNode->OnUpdateCommentText( NewText.ToString() );
 	}
