This change prevents comment balloons modifying text in read-only graphs.
#UE-13297 - Comment balloons are editable in readonly graphs
#Branch UE4
#Proj GraphEditor

diff --git a/Engine/Source/Editor/GraphEditor/Private/SCommentBubble.cpp b/Engine/Source/Editor/GraphEditor/Private/SCommentBubble.cpp
--- a/Engine/Source/Editor/GraphEditor/Private/SCommentBubble.cpp
+++ b/Engine/Source/Editor/GraphEditor/Private/SCommentBubble.cpp
@@ -79,7 +79,7 @@ void SCommentBubble::Tick( const FGeometry& AllottedGeometry, const double InCur
 	SCompoundWidget::Tick( AllottedGeometry, InCurrentTime, InDeltaTime );
 
 	// Check Editable and Hovered so we can prevent bubble toggling in read only graphs.
-	const bool bNodeEditable = GraphNode->NodeWidget.IsValid() ? GraphNode->NodeWidget.Pin()->IsNodeEditable() : false;
+	const bool bNodeEditable = !IsReadOnly();
 	const bool bEnableTitleHintBubble = bEnableTitleBarBubble && bNodeEditable;
 	const bool bTitleBarBubbleVisible = bEnableTitleHintBubble && IsGraphNodeHovered.IsBound();
 
@@ -220,8 +220,10 @@ void SCommentBubble::UpdateBubble()
 							SAssignNew(TextBlock, SMultiLineEditableTextBox)
 							.Text( this, &SCommentBubble::GetCommentText )
 							.HintText( NSLOCTEXT( "CommentBubble", "EditCommentHint", "Click to edit" ))
+							.IsReadOnly( this, &SCommentBubble::IsReadOnly )
 							.Font( FEditorStyle::GetFontStyle( TEXT("Graph.Node.CommentFont")))
 							.SelectAllTextWhenFocused( true )
+							.RevertTextOnEscape( true )
 							.ForegroundColor( this, &SCommentBubble::GetTextForegroundColor )
 							.BackgroundColor( this, &SCommentBubble::GetTextBackgroundColor )
 							.OnTextCommitted( this, &SCommentBubble::OnCommentTextCommitted )
@@ -417,4 +419,14 @@ void SCommentBubble::OnPinStateToggle( ECheckBoxState State ) const
 		GraphNode->Modify();
 		GraphNode->bCommentBubblePinned = State == ECheckBoxState::Checked;
 	}
+}
+
+bool SCommentBubble::IsReadOnly() const
+{
+	bool bReadOnly = true;
+	if( GraphNode && GraphNode->NodeWidget.IsValid() )
+	{
+		bReadOnly = !GraphNode->NodeWidget.Pin()->IsNodeEditable();
+	}
+	return bReadOnly;
 }
\ No newline at end of file
diff --git a/Engine/Source/Editor/GraphEditor/Public/SCommentBubble.h b/Engine/Source/Editor/GraphEditor/Public/SCommentBubble.h
--- a/Engine/Source/Editor/GraphEditor/Public/SCommentBubble.h
+++ b/Engine/Source/Editor/GraphEditor/Public/SCommentBubble.h
@@ -116,6 +116,9 @@ protected:
 	/** Called when the comment text is committed */
 	void OnCommentTextCommitted( const FText& NewText, ETextCommit::Type CommitInfo );
 
+	/** Called to determine if the comment bubble is readonly */
+	bool IsReadOnly() const;
+
 protected:
 
 	/** The GraphNode this widget interacts with */
