This change enables comment balloon coloring set in the preferences and restores the comment node option to enable coloring.
#UE-13299 - Comments and comment balloons don't obey colors set in editor preferences
#Branch UE4
#Proj BehaviourTreeEditor, EnvironmentQueryEditor, GraphEditor

diff --git a/Engine/Source/Editor/BehaviorTreeEditor/Private/SGraphNode_BehaviorTree.cpp b/Engine/Source/Editor/BehaviorTreeEditor/Private/SGraphNode_BehaviorTree.cpp
--- a/Engine/Source/Editor/BehaviorTreeEditor/Private/SGraphNode_BehaviorTree.cpp
+++ b/Engine/Source/Editor/BehaviorTreeEditor/Private/SGraphNode_BehaviorTree.cpp
@@ -492,11 +492,12 @@ void SGraphNode_BehaviorTree::UpdateGraphNode()
 		];
 	// Create comment bubble
 	TSharedPtr<SCommentBubble> CommentBubble;
+	const FSlateColor CommentColor = GetDefault<UGraphEditorSettings>()->DefaultCommentNodeTitleColor;
 
 	SAssignNew( CommentBubble, SCommentBubble )
 	.GraphNode( GraphNode )
 	.Text( this, &SGraphNode::GetNodeComment )
-	.ColorAndOpacity( this, &SGraphNode_BehaviorTree::GetCommentColor )
+	.ColorAndOpacity( CommentColor )
 	.AllowPinning( true )
 	.EnableTitleBarBubble( true )
 	.EnableBubbleCtrls( true )
diff --git a/Engine/Source/Editor/EnvironmentQueryEditor/Private/SGraphNode_EnvironmentQuery.cpp b/Engine/Source/Editor/EnvironmentQueryEditor/Private/SGraphNode_EnvironmentQuery.cpp
--- a/Engine/Source/Editor/EnvironmentQueryEditor/Private/SGraphNode_EnvironmentQuery.cpp
+++ b/Engine/Source/Editor/EnvironmentQueryEditor/Private/SGraphNode_EnvironmentQuery.cpp
@@ -295,11 +295,12 @@ void SGraphNode_EnvironmentQuery::UpdateGraphNode()
 
 	// Create comment bubble
 	TSharedPtr<SCommentBubble> CommentBubble;
+	const FSlateColor CommentColor = GetDefault<UGraphEditorSettings>()->DefaultCommentNodeTitleColor;
 
 	SAssignNew(CommentBubble, SCommentBubble)
 		.GraphNode(GraphNode)
 		.Text(this, &SGraphNode::GetNodeComment)
-		.ColorAndOpacity(this, &SGraphNode_EnvironmentQuery::GetCommentColor)
+		.ColorAndOpacity(CommentColor)
 		.AllowPinning(true)
 		.EnableTitleBarBubble(true)
 		.EnableBubbleCtrls(true)
diff --git a/Engine/Source/Editor/GraphEditor/Private/KismetNodes/SGraphNodeK2Base.cpp b/Engine/Source/Editor/GraphEditor/Private/KismetNodes/SGraphNodeK2Base.cpp
--- a/Engine/Source/Editor/GraphEditor/Private/KismetNodes/SGraphNodeK2Base.cpp
+++ b/Engine/Source/Editor/GraphEditor/Private/KismetNodes/SGraphNodeK2Base.cpp
@@ -140,11 +140,12 @@ void SGraphNodeK2Base::UpdateCompactNode()
 
 	// Create comment bubble
 	TSharedPtr<SCommentBubble> CommentBubble;
+	const FSlateColor CommentColor = GetDefault<UGraphEditorSettings>()->DefaultCommentNodeTitleColor;
 
 	SAssignNew( CommentBubble, SCommentBubble )
 	.GraphNode( GraphNode )
 	.Text( this, &SGraphNode::GetNodeComment )
-	.ColorAndOpacity( this, &SGraphNodeK2Base::GetCommentColor )
+	.ColorAndOpacity( CommentColor )
 	.AllowPinning( true )
 	.EnableTitleBarBubble( true )
 	.EnableBubbleCtrls( true )
diff --git a/Engine/Source/Editor/GraphEditor/Private/KismetNodes/SGraphNodeK2Composite.cpp b/Engine/Source/Editor/GraphEditor/Private/KismetNodes/SGraphNodeK2Composite.cpp
--- a/Engine/Source/Editor/GraphEditor/Private/KismetNodes/SGraphNodeK2Composite.cpp
+++ b/Engine/Source/Editor/GraphEditor/Private/KismetNodes/SGraphNodeK2Composite.cpp
@@ -121,11 +121,12 @@ void SGraphNodeK2Composite::UpdateGraphNode()
 		];
 	// Create comment bubble
 	TSharedPtr<SCommentBubble> CommentBubble;
+	const FSlateColor CommentColor = GetDefault<UGraphEditorSettings>()->DefaultCommentNodeTitleColor;
 
 	SAssignNew( CommentBubble, SCommentBubble )
 	.GraphNode( GraphNode )
 	.Text( this, &SGraphNode::GetNodeComment )
-	.ColorAndOpacity( this, &SGraphNodeK2Composite::GetCommentColor )
+	.ColorAndOpacity( CommentColor )
 	.AllowPinning( true )
 	.EnableTitleBarBubble( true )
 	.EnableBubbleCtrls( true )
diff --git a/Engine/Source/Editor/GraphEditor/Private/KismetNodes/SGraphNodeK2Var.cpp b/Engine/Source/Editor/GraphEditor/Private/KismetNodes/SGraphNodeK2Var.cpp
--- a/Engine/Source/Editor/GraphEditor/Private/KismetNodes/SGraphNodeK2Var.cpp
+++ b/Engine/Source/Editor/GraphEditor/Private/KismetNodes/SGraphNodeK2Var.cpp
@@ -274,11 +274,12 @@ void SGraphNodeK2Var::UpdateGraphNode()
 	}
 	// Create comment bubble
 	TSharedPtr<SCommentBubble> CommentBubble;
+	const FSlateColor CommentColor = GetDefault<UGraphEditorSettings>()->DefaultCommentNodeTitleColor;
 
 	SAssignNew( CommentBubble, SCommentBubble )
 	.GraphNode( GraphNode )
 	.Text( this, &SGraphNode::GetNodeComment )
-	.ColorAndOpacity( this, &SGraphNodeK2Var::GetCommentColor )
+	.ColorAndOpacity( CommentColor )
 	.AllowPinning( true )
 	.EnableTitleBarBubble( true )
 	.EnableBubbleCtrls( true )
diff --git a/Engine/Source/Editor/GraphEditor/Private/SGraphNode.cpp b/Engine/Source/Editor/GraphEditor/Private/SGraphNode.cpp
--- a/Engine/Source/Editor/GraphEditor/Private/SGraphNode.cpp
+++ b/Engine/Source/Editor/GraphEditor/Private/SGraphNode.cpp
@@ -840,11 +840,12 @@ void SGraphNode::UpdateGraphNode()
 
 	// Create comment bubble
 	TSharedPtr<SCommentBubble> CommentBubble;
+	const FSlateColor CommentColor = GetDefault<UGraphEditorSettings>()->DefaultCommentNodeTitleColor;
 
 	SAssignNew( CommentBubble, SCommentBubble )
 	.GraphNode( GraphNode )
 	.Text( this, &SGraphNode::GetNodeComment )
-	.ColorAndOpacity( this, &SGraphNode::GetCommentColor )
+	.ColorAndOpacity( CommentColor )
 	.AllowPinning( true )
 	.EnableTitleBarBubble( true )
 	.EnableBubbleCtrls( true )
diff --git a/Engine/Source/Editor/GraphEditor/Private/SGraphNodeComment.cpp b/Engine/Source/Editor/GraphEditor/Private/SGraphNodeComment.cpp
--- a/Engine/Source/Editor/GraphEditor/Private/SGraphNodeComment.cpp
+++ b/Engine/Source/Editor/GraphEditor/Private/SGraphNodeComment.cpp
@@ -413,7 +413,7 @@ FSlateColor SGraphNodeComment::GetCommentBodyColor() const
 {
 	UEdGraphNode_Comment* CommentNode = Cast<UEdGraphNode_Comment>(GraphNode);
 
-	if (CommentNode)
+	if (CommentNode && CommentNode->bColorCommentBubble)
 	{
 		return CommentNode->CommentColor;
 	}
@@ -426,7 +426,7 @@ FSlateColor SGraphNodeComment::GetCommentBodyColor() const
 FSlateColor SGraphNodeComment::GetCommentTitleBarColor() const
 {
 	UEdGraphNode_Comment* CommentNode = Cast<UEdGraphNode_Comment>(GraphNode);
-	if (CommentNode)
+	if (CommentNode && CommentNode->bColorCommentBubble)
 	{
 		const FLinearColor Color = CommentNode->CommentColor * SCommentNodeDefs::TitleBarColorMultiplier;
 		return FLinearColor(Color.R, Color.G, Color.B);
