Fix for reference viewer allowing users to edit the comment bubble text.
#UE-13446 - Reference viewer comments allow the user to edit the comment.
#Branch UE4
#Proj GraphEditor, ReferenceViewer

diff --git a/Engine/Source/Editor/GraphEditor/Public/SGraphNode.h b/Engine/Source/Editor/GraphEditor/Public/SGraphNode.h
--- a/Engine/Source/Editor/GraphEditor/Public/SGraphNode.h
+++ b/Engine/Source/Editor/GraphEditor/Public/SGraphNode.h
@@ -91,7 +91,7 @@ public:
 	void SetIsEditable(TAttribute<bool> InIsEditable);
 
 	/** Returns if widget is editable, additionally considers if the owning graph is read only */
-	bool IsNodeEditable() const;
+	virtual bool IsNodeEditable() const;
 
 	/** Set event when node is double clicked */
 	void SetDoubleClickEvent(FSingleNodeEvent InDoubleClickEvent);
diff --git a/Engine/Source/Editor/ReferenceViewer/Private/SReferenceNode.h b/Engine/Source/Editor/ReferenceViewer/Private/SReferenceNode.h
--- a/Engine/Source/Editor/ReferenceViewer/Private/SReferenceNode.h
+++ b/Engine/Source/Editor/ReferenceViewer/Private/SReferenceNode.h
@@ -17,6 +17,7 @@ public:
 
 	// SGraphNode implementation
 	virtual void UpdateGraphNode() override;
+	virtual bool IsNodeEditable() const override { return false; }
 	// End SGraphNode implementation
 
 private:
