This changes the comment balloons and context comment text boxes to require ctrl + enter for newlines to match editor standards.
#UE-6148 - Pressing "enter" in a node comment bubble moves cursor to a new line instead of submitting text
#Branch UE4
#Proj Engine, GraphEditor

diff --git a/Engine/Source/Editor/GraphEditor/Private/SCommentBubble.cpp b/Engine/Source/Editor/GraphEditor/Private/SCommentBubble.cpp
--- a/Engine/Source/Editor/GraphEditor/Private/SCommentBubble.cpp
+++ b/Engine/Source/Editor/GraphEditor/Private/SCommentBubble.cpp
@@ -224,6 +224,7 @@ void SCommentBubble::UpdateBubble()
 							.Font( FEditorStyle::GetFontStyle( TEXT("Graph.Node.CommentFont")))
 							.SelectAllTextWhenFocused( true )
 							.RevertTextOnEscape( true )
+							.ModiferKeyForNewLine( EModifierKey::Control )
 							.ForegroundColor( this, &SCommentBubble::GetTextForegroundColor )
 							.BackgroundColor( this, &SCommentBubble::GetTextBackgroundColor )
 							.OnTextCommitted( this, &SCommentBubble::OnCommentTextCommitted )
diff --git a/Engine/Source/Runtime/Engine/Private/EdGraph/EdGraphSchema.cpp b/Engine/Source/Runtime/Engine/Private/EdGraph/EdGraphSchema.cpp
--- a/Engine/Source/Runtime/Engine/Private/EdGraph/EdGraphSchema.cpp
+++ b/Engine/Source/Runtime/Engine/Private/EdGraph/EdGraphSchema.cpp
@@ -620,6 +620,7 @@ void UEdGraphSchema::GetContextMenuActions(const UEdGraph* CurrentGraph, const U
 					.OnTextCommitted_Static(&Local::OnNodeCommentTextCommitted, SelectedNodeWeakPtr)
 					.SelectAllTextWhenFocused( true )
 					.RevertTextOnEscape( true )
+					.ModiferKeyForNewLine( EModifierKey::Control )
 				];
 				MenuBuilder->EndSection();
 			}
