Fix for missing comment bubbles on nodes
#UE-11414 - Blueprint array nodes do not display comments
#Branch UE4
#Project GraphEditor

diff --git a/Engine/Source/Editor/GraphEditor/Private/KismetNodes/SGraphNodeK2ArrayFunction.cpp b/Engine/Source/Editor/GraphEditor/Private/KismetNodes/SGraphNodeK2ArrayFunction.cpp
--- a/Engine/Source/Editor/GraphEditor/Private/KismetNodes/SGraphNodeK2ArrayFunction.cpp
+++ b/Engine/Source/Editor/GraphEditor/Private/KismetNodes/SGraphNodeK2ArrayFunction.cpp
@@ -164,6 +164,29 @@ void SGraphNodeK2ArrayFunction::UpdateGraphNode()
 					]
 			];
 
+		// Create comment bubble
+		TSharedPtr<SCommentBubble> CommentBubble;
+
+		SAssignNew( CommentBubble, SCommentBubble )
+		.GraphNode( GraphNode )
+		.Text( this, &SGraphNode::GetNodeComment )
+		.ColorAndOpacity( this, &SGraphNodeK2ArrayFunction::GetCommentColor )
+		.AllowPinning( true )
+		.EnableTitleBarBubble( true )
+		.EnableBubbleCtrls( true )
+		.GraphLOD( this, &SGraphNode::GetCurrentLOD )
+		.IsGraphNodeHovered( this, &SGraphNode::IsHovered );
+
+		GetOrAddSlot( ENodeZone::TopCenter )
+		.SlotOffset( TAttribute<FVector2D>( CommentBubble.Get(), &SCommentBubble::GetOffset ))
+		.SlotSize( TAttribute<FVector2D>( CommentBubble.Get(), &SCommentBubble::GetSize ))
+		.AllowScaling( TAttribute<bool>( CommentBubble.Get(), &SCommentBubble::IsScalingAllowed ))
+		.VAlign( VAlign_Top )
+		[
+			CommentBubble.ToSharedRef()
+		];
+
+		// Create pins
 		CreatePinWidgets();
 
 		// Hide pin labels
