Fix for the comment nodes displaying comment balloons at fulldetail
#UE-13107 - Blueprint Editor: Newly-created Comments draw a redundant comment balloon on top by default
#UE-13301 - Comment depth is editable in the details panel for comment nodes
#Branch UE4
#Proj UnrealEd, GraphEditor

diff --git a/Engine/Source/Editor/GraphEditor/Private/SCommentBubble.cpp b/Engine/Source/Editor/GraphEditor/Private/SCommentBubble.cpp
--- a/Engine/Source/Editor/GraphEditor/Private/SCommentBubble.cpp
+++ b/Engine/Source/Editor/GraphEditor/Private/SCommentBubble.cpp
@@ -42,6 +42,7 @@ void SCommentBubble::Construct( const FArguments& InArgs )
 	bAllowPinning			= InArgs._AllowPinning;
 	bEnableTitleBarBubble	= InArgs._EnableTitleBarBubble;
 	bEnableBubbleCtrls		= InArgs._EnableBubbleCtrls;
+	bInvertLODCulling		= InArgs._InvertLODCulling;
 	GraphLOD				= InArgs._GraphLOD;
 	IsGraphNodeHovered		= InArgs._IsGraphNodeHovered;
 	HintText				= InArgs._HintText.IsSet() ? InArgs._HintText : NSLOCTEXT( "CommentBubble", "EditCommentHint", "Click to edit" );
@@ -118,7 +119,8 @@ void SCommentBubble::Tick( const FGeometry& AllottedGeometry, const double InCur
 	{
 		CachedComment = GraphNode->NodeComment;
 		CachedCommentText = FText::FromString( CachedComment );
-		GraphNode->bCommentBubbleVisible = !CachedComment.IsEmpty();
+		// Comment nodes should always have bCommentBubbleVisible set to true, they also are the nodes that set bInvertCulling.
+		GraphNode->bCommentBubbleVisible = bInvertLODCulling ? GraphNode->bCommentBubbleVisible : !CachedComment.IsEmpty();
 		UpdateBubble();
 	}
 }
@@ -139,6 +141,7 @@ void SCommentBubble::UpdateBubble()
 			if( bAllowPinning )
 			{
 				SAssignNew( BubbleControls, SVerticalBox )
+				.Visibility( this, &SCommentBubble::GetBubbleVisibility )
 				+SVerticalBox::Slot()
 				.Padding( 1.f )
 				.AutoHeight()
@@ -175,6 +178,7 @@ void SCommentBubble::UpdateBubble()
 			else
 			{
 				SAssignNew( BubbleControls, SVerticalBox )
+				.Visibility( this, &SCommentBubble::GetBubbleVisibility )
 				+SVerticalBox::Slot()
 				.AutoHeight()
 				.Padding( 1.f )
@@ -194,6 +198,7 @@ void SCommentBubble::UpdateBubble()
 		ChildSlot
 		[
 			SNew(SVerticalBox)
+			.Visibility( this, &SCommentBubble::GetBubbleVisibility )
 			+SVerticalBox::Slot()
 			.AutoHeight()
 			[
@@ -304,14 +309,14 @@ FVector2D SCommentBubble::GetSize() const
 bool SCommentBubble::IsBubbleVisible() const
 {
 	EGraphRenderingLOD::Type CurrLOD = GraphLOD.Get();
-	const bool bShowScaled = CurrLOD > EGraphRenderingLOD::LowestDetail;
+	const bool bShowScaled = CurrLOD > EGraphRenderingLOD::LowDetail;
 	const bool bShowPinned = CurrLOD <= EGraphRenderingLOD::MediumDetail;
 
-	if( bAllowPinning )
+	if( bAllowPinning && !bInvertLODCulling )
 	{
 		return GraphNode->bCommentBubblePinned ? true : bShowScaled;
 	}
-	return bShowPinned;
+	return bInvertLODCulling ? bShowPinned : !bShowPinned;
 }
 
 bool SCommentBubble::IsScalingAllowed() const
@@ -396,8 +401,7 @@ EVisibility SCommentBubble::GetToggleButtonVisibility() const
 
 EVisibility SCommentBubble::GetBubbleVisibility() const
 {
-	const bool bIsVisible = !CachedCommentText.IsEmpty() && IsBubbleVisible();
-	return bIsVisible ? EVisibility::Visible : EVisibility::Hidden;
+	return IsBubbleVisible() ? EVisibility::Visible : EVisibility::Hidden;
 }
 
 void SCommentBubble::OnCommentBubbleToggle( ECheckBoxState State )
diff --git a/Engine/Source/Editor/GraphEditor/Private/SGraphNodeComment.cpp b/Engine/Source/Editor/GraphEditor/Private/SGraphNodeComment.cpp
--- a/Engine/Source/Editor/GraphEditor/Private/SGraphNodeComment.cpp
+++ b/Engine/Source/Editor/GraphEditor/Private/SGraphNodeComment.cpp
@@ -188,9 +188,10 @@ void SGraphNodeComment::UpdateGraphNode()
 	.Text( this, &SGraphNode::GetNodeComment )
 	.ColorAndOpacity( this, &SGraphNodeComment::GetCommentColor )
 	.AllowPinning( true )
-	.EnableTitleBarBubble( true )
-	.EnableBubbleCtrls( true )
+	.EnableTitleBarBubble( false )
+	.EnableBubbleCtrls( false )
 	.GraphLOD( this, &SGraphNode::GetCurrentLOD )
+	.InvertLODCulling( true )
 	.IsGraphNodeHovered( this, &SGraphNode::IsHovered );
 
 	GetOrAddSlot( ENodeZone::TopCenter )
diff --git a/Engine/Source/Editor/GraphEditor/Public/SCommentBubble.h b/Engine/Source/Editor/GraphEditor/Public/SCommentBubble.h
--- a/Engine/Source/Editor/GraphEditor/Public/SCommentBubble.h
+++ b/Engine/Source/Editor/GraphEditor/Public/SCommentBubble.h
@@ -12,6 +12,7 @@ class GRAPHEDITOR_API SCommentBubble : public SCompoundWidget
 		, _AllowPinning( false )
 		, _EnableTitleBarBubble( false )
 		, _EnableBubbleCtrls( false )
+		, _InvertLODCulling( false )
 		, _GraphLOD( EGraphRenderingLOD::DefaultDetail )
 	{}
 
@@ -42,6 +43,9 @@ class GRAPHEDITOR_API SCommentBubble : public SCompoundWidget
 		/** Enable the controls within the bubble */
 		SLATE_ARGUMENT( bool, EnableBubbleCtrls )
 
+		/** Invert LOD culling */
+		SLATE_ARGUMENT( bool, InvertLODCulling )
+
 		/** The current level of detail */
 		SLATE_ATTRIBUTE( EGraphRenderingLOD::Type, GraphLOD )
 
@@ -148,6 +152,8 @@ protected:
 	bool bAllowPinning;
 	/** Allow in bubble controls */
 	bool bEnableBubbleCtrls;
+	/** Invert the LOD culling behavior, used by comment nodes */
+	bool bInvertLODCulling;
 	/** Enable the title bar bubble toggle */
 	bool bEnableTitleBarBubble;
 	/** Used to Control hover fade up/down for widgets */
diff --git a/Engine/Source/Editor/UnrealEd/Public/EdGraphNode_Comment.h b/Engine/Source/Editor/UnrealEd/Public/EdGraphNode_Comment.h
--- a/Engine/Source/Editor/UnrealEd/Public/EdGraphNode_Comment.h
+++ b/Engine/Source/Editor/UnrealEd/Public/EdGraphNode_Comment.h
@@ -42,7 +42,7 @@ public:
 	TEnumAsByte<ECommentBoxMode::Type> MoveMode;
 
 	/** comment Depth */
-	UPROPERTY(EditAnywhere, Category=Comment, meta=(DisplayName="CommentDepth"))
+	UPROPERTY()
 	int32 CommentDepth;
 
 public:
