Fix to remove the ability to toggle comment balloons in read only graphs.
#UE-13295 - Comment balloons can be toggled in read only graphs
#Branch UE4
#Proj GraphEditor

diff --git a/Engine/Source/Editor/GraphEditor/Private/SCommentBubble.cpp b/Engine/Source/Editor/GraphEditor/Private/SCommentBubble.cpp
--- a/Engine/Source/Editor/GraphEditor/Private/SCommentBubble.cpp
+++ b/Engine/Source/Editor/GraphEditor/Private/SCommentBubble.cpp
@@ -52,6 +52,8 @@ void SCommentBubble::Construct( const FArguments& InArgs )
 																			FOnCheckStateChanged::CreateSP( this, &SCommentBubble::OnCommentBubbleToggle );
 	CommentBubblePinnedDelegate = InArgs._OnCommentBubblePinned.IsBound() ? InArgs._OnCommentBubblePinned : 
 																			FOnCheckStateChanged::CreateSP( this, &SCommentBubble::OnPinStateToggle );
+	ToggleButtonCheck			= InArgs._ToggleButtonCheck.IsBound() ?		InArgs._ToggleButtonCheck : 
+																			TAttribute<ECheckBoxState>( this, &SCommentBubble::GetToggleButtonCheck );
 	// Ensue this value is set to something sensible
 	ForegroundColor = SCommentBubbleDefs::LightForegroundClr;
 
@@ -154,7 +156,7 @@ void SCommentBubble::UpdateBubble()
 					[
 						SNew( SCheckBox )
 						.Style( &FEditorStyle::Get().GetWidgetStyle<FCheckBoxStyle>( "CommentBubblePin" ))
-						.IsChecked( GraphNode->bCommentBubblePinned ? ECheckBoxState::Checked : ECheckBoxState::Unchecked )
+						.IsChecked( this, &SCommentBubble::GetPinnedButtonCheck )
 						.OnCheckStateChanged( CommentBubblePinnedDelegate )
 						.ToolTipText( this, &SCommentBubble::GetScaleButtonTooltip )
 						.Cursor( EMouseCursor::Default )
@@ -168,7 +170,7 @@ void SCommentBubble::UpdateBubble()
 				[
 					SNew( SCheckBox )
 					.Style( &FEditorStyle::Get().GetWidgetStyle< FCheckBoxStyle >( "CommentBubbleButton" ))
-					.IsChecked( ECheckBoxState::Checked )
+					.IsChecked( ToggleButtonCheck )
 					.OnCheckStateChanged( CommentBubbleToggleDelegate )
 					.ToolTipText( NSLOCTEXT( "CommentBubble", "ToggleCommentTooltip", "Toggle Comment Bubble" ))
 					.Cursor( EMouseCursor::Default )
@@ -186,7 +188,7 @@ void SCommentBubble::UpdateBubble()
 				[
 					SNew( SCheckBox )
 					.Style( &FEditorStyle::Get().GetWidgetStyle< FCheckBoxStyle >( "CommentBubbleButton" ))
-					.IsChecked( ECheckBoxState::Checked )
+					.IsChecked( ToggleButtonCheck )
 					.OnCheckStateChanged( CommentBubbleToggleDelegate )
 					.ToolTipText( NSLOCTEXT( "CommentBubble", "ToggleCommentTooltip", "Toggle Comment Bubble" ))
 					.Cursor( EMouseCursor::Default )
@@ -277,7 +279,7 @@ void SCommentBubble::UpdateBubble()
 			[
 				SNew( SCheckBox )
 				.Style( &FEditorStyle::Get().GetWidgetStyle< FCheckBoxStyle >( "CommentTitleButton" ))
-				.IsChecked( ECheckBoxState::Unchecked )
+				.IsChecked( ToggleButtonCheck )
 				.OnCheckStateChanged( CommentBubbleToggleDelegate )
 				.ToolTipText( NSLOCTEXT( "CommentBubble", "ToggleCommentTooltip", "Toggle Comment Bubble" ))
 				.Cursor( EMouseCursor::Default )
@@ -291,7 +293,6 @@ void SCommentBubble::UpdateBubble()
 	}
 }
 
-
 FVector2D SCommentBubble::GetOffset() const
 {
 	if( GraphNode->bCommentBubbleVisible || OpacityValue > 0.f )
@@ -404,9 +405,19 @@ EVisibility SCommentBubble::GetBubbleVisibility() const
 	return IsBubbleVisible() ? EVisibility::Visible : EVisibility::Hidden;
 }
 
-void SCommentBubble::OnCommentBubbleToggle( ECheckBoxState State )
+ECheckBoxState SCommentBubble::GetToggleButtonCheck() const
 {
+	ECheckBoxState CheckState = ECheckBoxState::Unchecked;
 	if( GraphNode )
+	{
+		CheckState = GraphNode->bCommentBubbleVisible ? ECheckBoxState::Checked : ECheckBoxState::Unchecked;
+	}
+	return CheckState;
+}
+
+void SCommentBubble::OnCommentBubbleToggle( ECheckBoxState State )
+{
+	if( !IsReadOnly() )
 	{
 		const FScopedTransaction Transaction( NSLOCTEXT( "CommentBubble", "BubbleVisibility", "Comment Bubble Visibilty" ) );
 		GraphNode->Modify();
@@ -416,9 +427,19 @@ void SCommentBubble::OnCommentBubbleToggle( ECheckBoxState State )
 	}
 }
 
-void SCommentBubble::OnPinStateToggle( ECheckBoxState State ) const
+ECheckBoxState SCommentBubble::GetPinnedButtonCheck() const
 {
+	ECheckBoxState CheckState = ECheckBoxState::Unchecked;
 	if( GraphNode )
+	{
+		CheckState = GraphNode->bCommentBubblePinned ? ECheckBoxState::Checked : ECheckBoxState::Unchecked;
+	}
+	return CheckState;
+}
+
+void SCommentBubble::OnPinStateToggle( ECheckBoxState State ) const
+{
+	if( !IsReadOnly() )
 	{
 		const FScopedTransaction Transaction( NSLOCTEXT( "CommentBubble", "BubblePinned", "Comment Bubble Pin" ) );
 		GraphNode->Modify();
diff --git a/Engine/Source/Editor/GraphEditor/Public/SCommentBubble.h b/Engine/Source/Editor/GraphEditor/Public/SCommentBubble.h
--- a/Engine/Source/Editor/GraphEditor/Public/SCommentBubble.h
+++ b/Engine/Source/Editor/GraphEditor/Public/SCommentBubble.h
@@ -25,6 +25,9 @@ class GRAPHEDITOR_API SCommentBubble : public SCompoundWidget
 		/** Called when the comment bubble pin button is pressed */
 		SLATE_EVENT( FOnCheckStateChanged, OnCommentBubblePinned )
 
+		/** Called when the comment bubble toggle button is pressed */
+		SLATE_ATTRIBUTE( ECheckBoxState, ToggleButtonCheck )
+
 		/** The comment text for the bubble */
 		SLATE_ATTRIBUTE( FString, Text )
 
@@ -120,6 +123,12 @@ protected:
 	/** Called when the comment text is committed */
 	void OnCommentTextCommitted( const FText& NewText, ETextCommit::Type CommitInfo );
 
+	/** Returns bubble toggle check state */
+	ECheckBoxState GetToggleButtonCheck() const;
+
+	/** Returns pinned check state */
+	ECheckBoxState GetPinnedButtonCheck() const;
+
 	/** Called to determine if the comment bubble is readonly */
 	bool IsReadOnly() const;
 
@@ -138,6 +147,8 @@ protected:
 	TAttribute<EGraphRenderingLOD::Type> GraphLOD;
 	/** Hint Text */
 	TAttribute<FText> HintText;
+	/** Toggle button checked state  */
+	TAttribute<ECheckBoxState> ToggleButtonCheck;
 
 	/** Optional delegate to call when the comment is toggled */
 	FOnCheckStateChanged CommentBubbleToggleDelegate;
