Keywords metadata is now localized.

Blueprint node searching now leverages localized keyword metadata for searching, so searches can now be done in the current langauge and English to find the same nodes.

diff --git a/Engine/Config/BaseEditor.ini b/Engine/Config/BaseEditor.ini
--- a/Engine/Config/BaseEditor.ini
+++ b/Engine/Config/BaseEditor.ini
@@ -1,6 +1,7 @@
 [Internationalization]
 +LocalizationPaths=../../../Engine/Content/Localization/Editor
 +LocalizationPaths=../../../Engine/Content/Localization/EditorTutorials
++LocalizationPaths=../../../Engine/Content/Localization/Keywords
 +ToolTipLocalizationPaths=../../../Engine/Content/Localization/ToolTips
 +PropertyNameLocalizationPaths=../../../Engine/Content/Localization/PropertyNames
 ShouldLoadLocalizedPropertyNames=True
diff --git a/Engine/Config/Localization/Keywords.ini b/Engine/Config/Localization/Keywords.ini
new file mode 100644
--- /dev/null
+++ b/Engine/Config/Localization/Keywords.ini
@@ -0,0 +1,49 @@
+[CommonSettings]
+SourcePath=./Content/Localization/Keywords
+DestinationPath=./Content/Localization/Keywords
+ManifestName=Keywords.manifest
+ArchiveName=Keywords.archive
+PortableObjectName=Keywords.po
+SourceCulture=en
+CulturesToGenerate=es
+CulturesToGenerate=en
+CulturesToGenerate=ja
+CulturesToGenerate=ko
+CulturesToGenerate=pt-BR
+CulturesToGenerate=zh-CN
+
+;Gather text from metadata
+[GatherTextStep0]
+CommandletClass=GatherTextFromMetaData
+IncludePaths=./Source/Editor/*
+IncludePaths=./Source/Runtime/*
+IncludePaths=./Source/Developer/*
+ExcludePaths=*/Config/Localization/*
+ExcludePaths=*/Source/Developer/NoRedist/CommunityPortalServices/*
+InputKeys=Keywords
+OutputNamespaces=UObjectKeywords
+OutputKeys="{FieldPath}"
+
+;Write Manifest
+[GatherTextStep1]
+CommandletClass=GenerateGatherManifest
+
+;Write Archives
+[GatherTextStep2]
+CommandletClass=GenerateGatherArchive
+bPurgeOldEmptyEntries=true
+
+;Import localized PO files
+[GatherTextStep3]
+CommandletClass=InternationalizationExport
+bImportLoc=true
+
+;Write Localized Text Resource
+[GatherTextStep4]
+CommandletClass=GenerateTextLocalizationResource
+ResourceName=Keywords.locres
+
+;Export to PO files
+[GatherTextStep5]
+CommandletClass=InternationalizationExport
+bExportLoc=true
\ No newline at end of file
diff --git a/Engine/Plugins/2D/Paper2D/Source/Paper2D/Private/Rendering/MaterialExpressionSpriteTextureSampler.cpp b/Engine/Plugins/2D/Paper2D/Source/Paper2D/Private/Rendering/MaterialExpressionSpriteTextureSampler.cpp
--- a/Engine/Plugins/2D/Paper2D/Source/Paper2D/Private/Rendering/MaterialExpressionSpriteTextureSampler.cpp
+++ b/Engine/Plugins/2D/Paper2D/Source/Paper2D/Private/Rendering/MaterialExpressionSpriteTextureSampler.cpp
@@ -39,10 +39,13 @@ void UMaterialExpressionSpriteTextureSampler::GetCaption(TArray<FString>& OutCap
 }
 
 #if WITH_EDITOR
-FString UMaterialExpressionSpriteTextureSampler::GetKeywords() const
+FText UMaterialExpressionSpriteTextureSampler::GetKeywords() const
 {
-	FString ParentKeywords = Super::GetKeywords();
-	return ParentKeywords + TEXT(" Paper2D Sprite");
+	FText ParentKeywords = Super::GetKeywords();
+
+	FFormatNamedArguments Args;
+	Args.Add(TEXT("ParentKeywords"), ParentKeywords);
+	return FText::Format(LOCTEXT("SpriteTextureSamplerKeywords", "{ParentKeywords} Paper2D Sprite"), Args);
 }
 
 bool UMaterialExpressionSpriteTextureSampler::CanRenameNode() const
diff --git a/Engine/Plugins/2D/Paper2D/Source/Paper2D/Public/MaterialExpressionSpriteTextureSampler.h b/Engine/Plugins/2D/Paper2D/Source/Paper2D/Public/MaterialExpressionSpriteTextureSampler.h
--- a/Engine/Plugins/2D/Paper2D/Source/Paper2D/Public/MaterialExpressionSpriteTextureSampler.h
+++ b/Engine/Plugins/2D/Paper2D/Source/Paper2D/Public/MaterialExpressionSpriteTextureSampler.h
@@ -26,7 +26,7 @@ class PAPER2D_API UMaterialExpressionSpriteTextureSampler : public UMaterialExpr
 	// UMaterialExpression interface
 	virtual void GetCaption(TArray<FString>& OutCaptions) const override;
 #if WITH_EDITOR
-	virtual FString GetKeywords() const override;
+	virtual FText GetKeywords() const override;
 	virtual bool CanRenameNode() const override;
 	virtual void GetExpressionToolTip(TArray<FString>& OutToolTip) override;
 	virtual void PostEditChangeProperty(FPropertyChangedEvent& PropertyChangedEvent) override;
diff --git a/Engine/Source/Editor/AnimGraph/Classes/AnimGraphNode_LookAt.h b/Engine/Source/Editor/AnimGraph/Classes/AnimGraphNode_LookAt.h
--- a/Engine/Source/Editor/AnimGraph/Classes/AnimGraphNode_LookAt.h
+++ b/Engine/Source/Editor/AnimGraph/Classes/AnimGraphNode_LookAt.h
@@ -7,7 +7,7 @@
 #include "EdGraph/EdGraphNodeUtils.h" // for FNodeTitleTextTable
 #include "AnimGraphNode_LookAt.generated.h"
 
-UCLASS(MinimalAPI)
+UCLASS(MinimalAPI, meta=(Keywords = "Look At, Follow, Trace, Track"))
 class UAnimGraphNode_LookAt : public UAnimGraphNode_SkeletalControlBase
 {
 	GENERATED_UCLASS_BODY()
@@ -19,7 +19,6 @@ public:
 	// UEdGraphNode interface
 	virtual FText GetNodeTitle(ENodeTitleType::Type TitleType) const override;
 	virtual FText GetTooltipText() const override;
-	virtual FString GetKeywords() const override;
 	// End of UEdGraphNode interface
 
 	// UAnimGraphNode_SkeletalControlBase interface
diff --git a/Engine/Source/Editor/AnimGraph/Classes/AnimGraphNode_ModifyBone.h b/Engine/Source/Editor/AnimGraph/Classes/AnimGraphNode_ModifyBone.h
--- a/Engine/Source/Editor/AnimGraph/Classes/AnimGraphNode_ModifyBone.h
+++ b/Engine/Source/Editor/AnimGraph/Classes/AnimGraphNode_ModifyBone.h
@@ -7,7 +7,7 @@
 #include "EdGraph/EdGraphNodeUtils.h" // for FNodeTitleTextTable
 #include "AnimGraphNode_ModifyBone.generated.h"
 
-UCLASS(MinimalAPI)
+UCLASS(MinimalAPI, meta=(Keywords = "Modify Transform"))
 class UAnimGraphNode_ModifyBone : public UAnimGraphNode_SkeletalControlBase
 {
 	GENERATED_UCLASS_BODY()
@@ -19,7 +19,6 @@ public:
 	// UEdGraphNode interface
 	virtual FText GetNodeTitle(ENodeTitleType::Type TitleType) const override;
 	virtual FText GetTooltipText() const override;
-	virtual FString GetKeywords() const override;
 	// End of UEdGraphNode interface
 
 protected:
diff --git a/Engine/Source/Editor/AnimGraph/Classes/AnimGraphNode_WheelHandler.h b/Engine/Source/Editor/AnimGraph/Classes/AnimGraphNode_WheelHandler.h
--- a/Engine/Source/Editor/AnimGraph/Classes/AnimGraphNode_WheelHandler.h
+++ b/Engine/Source/Editor/AnimGraph/Classes/AnimGraphNode_WheelHandler.h
@@ -6,7 +6,7 @@
 #include "Animation/BoneControllers/AnimNode_WheelHandler.h"
 #include "AnimGraphNode_WheelHandler.generated.h"
 
-UCLASS(MinimalAPI)
+UCLASS(MinimalAPI, meta=(Keywords = "Modify Wheel Vehicle"))
 class UAnimGraphNode_WheelHandler : public UAnimGraphNode_SkeletalControlBase
 {
 	GENERATED_UCLASS_BODY()
@@ -18,7 +18,6 @@ public:
 	// UEdGraphNode interface
 	virtual FText GetNodeTitle(ENodeTitleType::Type TitleType) const override;
 	virtual FText GetTooltipText() const override;
-	virtual FString GetKeywords() const override;
 	// validate if this is within VehicleAnimInstance
 	virtual void ValidateAnimNodePostCompile(class FCompilerResultsLog& MessageLog, class UAnimBlueprintGeneratedClass* CompiledClass, int32 CompiledNodeIndex) override;
 	virtual bool IsCompatibleWithGraph(const UEdGraph* TargetGraph) const override;
diff --git a/Engine/Source/Editor/AnimGraph/Private/AnimGraphNode_BlendSpaceBase.cpp b/Engine/Source/Editor/AnimGraph/Private/AnimGraphNode_BlendSpaceBase.cpp
--- a/Engine/Source/Editor/AnimGraph/Private/AnimGraphNode_BlendSpaceBase.cpp
+++ b/Engine/Source/Editor/AnimGraph/Private/AnimGraphNode_BlendSpaceBase.cpp
@@ -37,7 +37,7 @@ struct FNewBlendSpacePlayerAction : public FEdGraphSchemaAction_K2NewNode
 		Category = TEXT("Animations");
 
 		// Grab extra keywords
-		Keywords = BlendSpace->GetPathName();
+		Keywords = FText::FromString(BlendSpace->GetPathName());
 	}
 };
 
diff --git a/Engine/Source/Editor/AnimGraph/Private/AnimGraphNode_LookAt.cpp b/Engine/Source/Editor/AnimGraph/Private/AnimGraphNode_LookAt.cpp
--- a/Engine/Source/Editor/AnimGraph/Private/AnimGraphNode_LookAt.cpp
+++ b/Engine/Source/Editor/AnimGraph/Private/AnimGraphNode_LookAt.cpp
@@ -18,11 +18,6 @@ FText UAnimGraphNode_LookAt::GetControllerDescription() const
 	return LOCTEXT("LookAtNode", "Look At");
 }
 
-FString UAnimGraphNode_LookAt::GetKeywords() const
-{
-	return TEXT("Look At, Follow, Trace, Track");
-}
-
 FText UAnimGraphNode_LookAt::GetTooltipText() const
 {
 	return LOCTEXT("AnimGraphNode_LookAt_Tooltip", "This node allow a bone to trace or follow another bone");
diff --git a/Engine/Source/Editor/AnimGraph/Private/AnimGraphNode_ModifyBone.cpp b/Engine/Source/Editor/AnimGraph/Private/AnimGraphNode_ModifyBone.cpp
--- a/Engine/Source/Editor/AnimGraph/Private/AnimGraphNode_ModifyBone.cpp
+++ b/Engine/Source/Editor/AnimGraph/Private/AnimGraphNode_ModifyBone.cpp
@@ -19,11 +19,6 @@ FText UAnimGraphNode_ModifyBone::GetControllerDescription() const
 	return LOCTEXT("TransformModifyBone", "Transform (Modify) Bone");
 }
 
-FString UAnimGraphNode_ModifyBone::GetKeywords() const
-{
-	return TEXT("Modify, Transform");
-}
-
 FText UAnimGraphNode_ModifyBone::GetTooltipText() const
 {
 	return LOCTEXT("AnimGraphNode_ModifyBone_Tooltip", "The Transform Bone node alters the transform - i.e. Translation, Rotation, or Scale - of the bone");
diff --git a/Engine/Source/Editor/AnimGraph/Private/AnimGraphNode_SequencePlayer.cpp b/Engine/Source/Editor/AnimGraph/Private/AnimGraphNode_SequencePlayer.cpp
--- a/Engine/Source/Editor/AnimGraph/Private/AnimGraphNode_SequencePlayer.cpp
+++ b/Engine/Source/Editor/AnimGraph/Private/AnimGraphNode_SequencePlayer.cpp
@@ -35,7 +35,7 @@ public:
 		Category = TEXT("Animations");
 
 		// Grab extra keywords
-		Keywords = InAssetInfo.ObjectPath.ToString();
+		Keywords = FText::FromName(InAssetInfo.ObjectPath);
 	}
 
 	virtual UEdGraphNode* PerformAction(class UEdGraph* ParentGraph, UEdGraphPin* FromPin, const FVector2D Location, bool bSelectNewNode = true) override
diff --git a/Engine/Source/Editor/AnimGraph/Private/AnimGraphNode_WheelHandler.cpp b/Engine/Source/Editor/AnimGraph/Private/AnimGraphNode_WheelHandler.cpp
--- a/Engine/Source/Editor/AnimGraph/Private/AnimGraphNode_WheelHandler.cpp
+++ b/Engine/Source/Editor/AnimGraph/Private/AnimGraphNode_WheelHandler.cpp
@@ -19,11 +19,6 @@ FText UAnimGraphNode_WheelHandler::GetControllerDescription() const
 	return LOCTEXT("AnimGraphNode_WheelHandler", "Wheel Handler for WheeledVehicle");
 }
 
-FString UAnimGraphNode_WheelHandler::GetKeywords() const
-{
-	return TEXT("Modify, Wheel, Vehicle");
-}
-
 FText UAnimGraphNode_WheelHandler::GetTooltipText() const
 {
 	return LOCTEXT("AnimGraphNode_WheelHandler_Tooltip", "This alters the wheel transform based on set up in Wheeled Vehicle. This only works when the owner is WheeledVehicle.");
diff --git a/Engine/Source/Editor/BlueprintGraph/Classes/K2Node_CallFunction.h b/Engine/Source/Editor/BlueprintGraph/Classes/K2Node_CallFunction.h
--- a/Engine/Source/Editor/BlueprintGraph/Classes/K2Node_CallFunction.h
+++ b/Engine/Source/Editor/BlueprintGraph/Classes/K2Node_CallFunction.h
@@ -146,7 +146,7 @@ public:
 	/** Get default category for this function in action menu */
 	static FString GetDefaultCategoryForFunction(const UFunction* Function, const FString& BaseCategory);
 	/** Get keywords for this function in the action menu */
-	static FString GetKeywordsForFunction(const UFunction* Function);
+	static FText GetKeywordsForFunction(const UFunction* Function);
 	/** Should be drawn compact for this function */
 	static bool ShouldDrawCompact(const UFunction* Function);
 	/** Get the compact name for this function */
diff --git a/Engine/Source/Editor/BlueprintGraph/Classes/K2Node_EnumEquality.h b/Engine/Source/Editor/BlueprintGraph/Classes/K2Node_EnumEquality.h
--- a/Engine/Source/Editor/BlueprintGraph/Classes/K2Node_EnumEquality.h
+++ b/Engine/Source/Editor/BlueprintGraph/Classes/K2Node_EnumEquality.h
@@ -5,7 +5,7 @@
 #include "K2Node.h"
 #include "K2Node_EnumEquality.generated.h"
 
-UCLASS(MinimalAPI)
+UCLASS(MinimalAPI, meta=(Keywords = "=="))
 class UK2Node_EnumEquality : public UK2Node
 {
 	GENERATED_UCLASS_BODY()
@@ -14,7 +14,6 @@ class UK2Node_EnumEquality : public UK2Node
 	virtual void AllocateDefaultPins() override;
 	virtual FText GetTooltipText() const override;
 	virtual FText GetNodeTitle(ENodeTitleType::Type TitleType) const override;
-	virtual FString GetKeywords() const override;
 	// End UEdGraphNode interface
 
 	// Begin UK2Node interface
diff --git a/Engine/Source/Editor/BlueprintGraph/Classes/K2Node_EnumInequality.h b/Engine/Source/Editor/BlueprintGraph/Classes/K2Node_EnumInequality.h
--- a/Engine/Source/Editor/BlueprintGraph/Classes/K2Node_EnumInequality.h
+++ b/Engine/Source/Editor/BlueprintGraph/Classes/K2Node_EnumInequality.h
@@ -7,7 +7,7 @@
 #include "K2Node_EnumInequality.generated.h"
 
 
-UCLASS(MinimalAPI)
+UCLASS(MinimalAPI, meta=(Keywords = "!="))
 class UK2Node_EnumInequality : public UK2Node_EnumEquality
 {
 	GENERATED_UCLASS_BODY()
@@ -15,7 +15,6 @@ class UK2Node_EnumInequality : public UK2Node_EnumEquality
 	// Begin UEdGraphNode interface
 	virtual FText GetTooltipText() const override;
 	virtual FText GetNodeTitle(ENodeTitleType::Type TitleType) const override;
-	virtual FString GetKeywords() const override;
 	// End UEdGraphNode interface
 
 	// Begin UK2Node interface
diff --git a/Engine/Source/Editor/BlueprintGraph/Classes/K2Node_Event.h b/Engine/Source/Editor/BlueprintGraph/Classes/K2Node_Event.h
--- a/Engine/Source/Editor/BlueprintGraph/Classes/K2Node_Event.h
+++ b/Engine/Source/Editor/BlueprintGraph/Classes/K2Node_Event.h
@@ -49,7 +49,7 @@ class UK2Node_Event : public UK2Node_EditablePinBase
 	// Begin UEdGraphNode interface
 	BLUEPRINTGRAPH_API virtual void AllocateDefaultPins() override;
 	BLUEPRINTGRAPH_API virtual FText GetTooltipText() const override;
-	BLUEPRINTGRAPH_API virtual FString GetKeywords() const override;	
+	BLUEPRINTGRAPH_API virtual FText GetKeywords() const override;	
 	BLUEPRINTGRAPH_API virtual FText GetNodeTitle(ENodeTitleType::Type TitleType) const override;
 	BLUEPRINTGRAPH_API virtual FLinearColor GetNodeTitleColor() const override;
 	BLUEPRINTGRAPH_API virtual bool CanPasteHere(const UEdGraph* TargetGraph) const override;
diff --git a/Engine/Source/Editor/BlueprintGraph/Classes/K2Node_GetInputAxisKeyValue.h b/Engine/Source/Editor/BlueprintGraph/Classes/K2Node_GetInputAxisKeyValue.h
--- a/Engine/Source/Editor/BlueprintGraph/Classes/K2Node_GetInputAxisKeyValue.h
+++ b/Engine/Source/Editor/BlueprintGraph/Classes/K2Node_GetInputAxisKeyValue.h
@@ -5,7 +5,7 @@
 #include "EdGraph/EdGraphNodeUtils.h" // for FNodeTextCache
 #include "K2Node_GetInputAxisKeyValue.generated.h"
 
-UCLASS(MinimalAPI)
+UCLASS(MinimalAPI, meta=(Keywords = "Get"))
 class UK2Node_GetInputAxisKeyValue : public UK2Node_CallFunction
 {
 	GENERATED_UCLASS_BODY()
@@ -36,7 +36,6 @@ class UK2Node_GetInputAxisKeyValue : public UK2Node_CallFunction
 	virtual void RegisterDynamicBinding(UDynamicBlueprintBinding* BindingObject) const override;
 	virtual void GetMenuActions(FBlueprintActionDatabaseRegistrar& ActionRegistrar) const override;
 	virtual FText GetMenuCategory() const override;
-	virtual FString GetKeywords() const override;
 	virtual FBlueprintNodeSignature GetSignature() const override;
 	// End UK2Node interface
 	
diff --git a/Engine/Source/Editor/BlueprintGraph/Classes/K2Node_GetInputAxisValue.h b/Engine/Source/Editor/BlueprintGraph/Classes/K2Node_GetInputAxisValue.h
--- a/Engine/Source/Editor/BlueprintGraph/Classes/K2Node_GetInputAxisValue.h
+++ b/Engine/Source/Editor/BlueprintGraph/Classes/K2Node_GetInputAxisValue.h
@@ -5,7 +5,7 @@
 #include "EdGraph/EdGraphNodeUtils.h" // for FNodeTextCache
 #include "K2Node_GetInputAxisValue.generated.h"
 
-UCLASS(MinimalAPI)
+UCLASS(MinimalAPI, meta=(Keywords = "Get"))
 class UK2Node_GetInputAxisValue : public UK2Node_CallFunction
 {
 	GENERATED_UCLASS_BODY()
@@ -25,7 +25,6 @@ class UK2Node_GetInputAxisValue : public UK2Node_CallFunction
 	virtual void AllocateDefaultPins() override;
 	virtual FText GetNodeTitle(ENodeTitleType::Type TitleType) const override;
 	virtual FText GetTooltipText() const override;
-	virtual FString GetKeywords() const override;
 	virtual bool IsCompatibleWithGraph(UEdGraph const* Graph) const override;
 	// End EdGraphNode interface
 
diff --git a/Engine/Source/Editor/BlueprintGraph/Classes/K2Node_GetInputVectorAxisValue.h b/Engine/Source/Editor/BlueprintGraph/Classes/K2Node_GetInputVectorAxisValue.h
--- a/Engine/Source/Editor/BlueprintGraph/Classes/K2Node_GetInputVectorAxisValue.h
+++ b/Engine/Source/Editor/BlueprintGraph/Classes/K2Node_GetInputVectorAxisValue.h
@@ -4,14 +4,13 @@
 #include "K2Node_GetInputAxisKeyValue.h"
 #include "K2Node_GetInputVectorAxisValue.generated.h"
 
-UCLASS(MinimalAPI)
+UCLASS(MinimalAPI, meta=(Keywords = "Get"))
 class UK2Node_GetInputVectorAxisValue : public UK2Node_GetInputAxisKeyValue
 {
 	GENERATED_UCLASS_BODY()
 
 	// Begin EdGraphNode interface
 	virtual FText GetTooltipText() const override;
-	virtual FString GetKeywords() const override;
 	// End EdGraphNode interface
 
 	// Begin UK2Node interface
diff --git a/Engine/Source/Editor/BlueprintGraph/Classes/K2Node_IfThenElse.h b/Engine/Source/Editor/BlueprintGraph/Classes/K2Node_IfThenElse.h
--- a/Engine/Source/Editor/BlueprintGraph/Classes/K2Node_IfThenElse.h
+++ b/Engine/Source/Editor/BlueprintGraph/Classes/K2Node_IfThenElse.h
@@ -5,7 +5,7 @@
 #include "K2Node.h"
 #include "K2Node_IfThenElse.generated.h"
 
-UCLASS(MinimalAPI)
+UCLASS(MinimalAPI, meta=(Keywords = "if bool"))
 class UK2Node_IfThenElse : public UK2Node
 {
 	GENERATED_UCLASS_BODY()
@@ -15,7 +15,6 @@ class UK2Node_IfThenElse : public UK2Node
 	virtual FLinearColor GetNodeTitleColor() const override;
 	virtual FText GetTooltipText() const override;
 	virtual FText GetNodeTitle(ENodeTitleType::Type TitleType) const override;
-	virtual FString GetKeywords() const override;
 	virtual FName GetPaletteIcon(FLinearColor& OutColor) const override{ return TEXT("GraphEditor.Branch_16x"); }
 	// End UEdGraphNode interface
 
diff --git a/Engine/Source/Editor/BlueprintGraph/Classes/K2Node_MacroInstance.h b/Engine/Source/Editor/BlueprintGraph/Classes/K2Node_MacroInstance.h
--- a/Engine/Source/Editor/BlueprintGraph/Classes/K2Node_MacroInstance.h
+++ b/Engine/Source/Editor/BlueprintGraph/Classes/K2Node_MacroInstance.h
@@ -36,7 +36,7 @@ public:
 	// Begin UEdGraphNode interface
 	virtual void AllocateDefaultPins() override;
 	virtual FText GetTooltipText() const override;
-	virtual FString GetKeywords() const override;
+	virtual FText GetKeywords() const override;
 	virtual void PostPasteNode() override;
 	virtual FLinearColor GetNodeTitleColor() const override;
 	virtual FText GetNodeTitle(ENodeTitleType::Type TitleType) const override;
diff --git a/Engine/Source/Editor/BlueprintGraph/Classes/K2Node_Select.h b/Engine/Source/Editor/BlueprintGraph/Classes/K2Node_Select.h
--- a/Engine/Source/Editor/BlueprintGraph/Classes/K2Node_Select.h
+++ b/Engine/Source/Editor/BlueprintGraph/Classes/K2Node_Select.h
@@ -6,7 +6,7 @@
 #include "NodeDependingOnEnumInterface.h"
 #include "K2Node_Select.generated.h"
 
-UCLASS(MinimalAPI)
+UCLASS(MinimalAPI, meta=(Keywords = "Ternary If"))
 class UK2Node_Select : public UK2Node, public INodeDependingOnEnumInterface
 {
 	GENERATED_UCLASS_BODY()
@@ -38,7 +38,6 @@ class UK2Node_Select : public UK2Node, public INodeDependingOnEnumInterface
 	// Begin UEdGraphNode interface
 	virtual void AllocateDefaultPins() override;
 	virtual FText GetTooltipText() const override;
-	virtual FString GetKeywords() const override;
 	virtual FText GetNodeTitle(ENodeTitleType::Type TitleType) const override;
 	virtual void NodeConnectionListChanged() override;
 	virtual void PinTypeChanged(UEdGraphPin* Pin) override;
diff --git a/Engine/Source/Editor/BlueprintGraph/Classes/K2Node_Self.h b/Engine/Source/Editor/BlueprintGraph/Classes/K2Node_Self.h
--- a/Engine/Source/Editor/BlueprintGraph/Classes/K2Node_Self.h
+++ b/Engine/Source/Editor/BlueprintGraph/Classes/K2Node_Self.h
@@ -13,7 +13,7 @@ class UK2Node_Self : public UK2Node
 	// Begin UEdGraphNode interface
 	virtual FText GetNodeTitle(ENodeTitleType::Type TitleType) const override;
 	virtual FText GetTooltipText() const override;
-	virtual FString GetKeywords() const override;
+	virtual FText GetKeywords() const override;
 	virtual void AllocateDefaultPins() override;
 	virtual void ValidateNodeDuringCompilation(class FCompilerResultsLog& MessageLog) const override;
 	// End UEdGraphNode interface
diff --git a/Engine/Source/Editor/BlueprintGraph/Private/BlueprintComponentNodeSpawner.cpp b/Engine/Source/Editor/BlueprintGraph/Private/BlueprintComponentNodeSpawner.cpp
--- a/Engine/Source/Editor/BlueprintGraph/Private/BlueprintComponentNodeSpawner.cpp
+++ b/Engine/Source/Editor/BlueprintGraph/Private/BlueprintComponentNodeSpawner.cpp
@@ -63,9 +63,12 @@ UBlueprintComponentNodeSpawner* UBlueprintComponentNodeSpawner::Create(const FCo
 		MenuSignature.MenuName = FText::Format(LOCTEXT("AddComponentMenuName", "Add {0}"), ComponentTypeName);
 		MenuSignature.Category = LOCTEXT("BlueprintComponentCategory", "Custom");
 		MenuSignature.Tooltip = FText::Format(LOCTEXT("AddComponentTooltip", "Spawn a {0}"), ComponentTypeName);
-		// add at least one character, so that PrimeDefaultMenuSignature() doesn't 
+		// add at least one character, so that PrimeDefaultUiSpec() doesn't 
 		// attempt to query the template node
-		MenuSignature.Keywords.AppendChar(TEXT(' '));
+		if (MenuSignature.Keywords.IsEmpty())
+		{
+			MenuSignature.Keywords = FText::FromString(TEXT(" "));
+		}
 		MenuSignature.IconName = FClassIconFinder::FindIconNameForClass(nullptr);
 
 		return NodeSpawner;
@@ -88,10 +91,13 @@ UBlueprintComponentNodeSpawner* UBlueprintComponentNodeSpawner::Create(const FCo
 	MenuSignature.MenuName = FText::Format(LOCTEXT("AddComponentMenuName", "Add {0}"), ComponentTypeName);
 	MenuSignature.Category = BlueprintComponentNodeSpawnerImpl::GetDefaultMenuCategory(AuthoritativeClass);
 	MenuSignature.Tooltip  = FText::Format(LOCTEXT("AddComponentTooltip", "Spawn a {0}"), ComponentTypeName);
-	MenuSignature.Keywords = AuthoritativeClass->GetMetaData(FBlueprintMetadata::MD_FunctionKeywords);
-	// add at least one character, so that PrimeDefaultMenuSignature() doesn't 
+	MenuSignature.Keywords = AuthoritativeClass->GetMetaDataText(*FBlueprintMetadata::MD_FunctionKeywords.ToString(), TEXT("UObjectKeywords"), AuthoritativeClass->GetFullGroupName(false));
+	// add at least one character, so that PrimeDefaultUiSpec() doesn't 
 	// attempt to query the template node
-	MenuSignature.Keywords.AppendChar(TEXT(' '));
+	if (MenuSignature.Keywords.IsEmpty())
+	{
+		MenuSignature.Keywords = FText::FromString(TEXT(" "));
+	}
 	MenuSignature.IconName = FClassIconFinder::FindIconNameForClass(AuthoritativeClass);
 
 	return NodeSpawner;
diff --git a/Engine/Source/Editor/BlueprintGraph/Private/BlueprintEventNodeSpawner.cpp b/Engine/Source/Editor/BlueprintGraph/Private/BlueprintEventNodeSpawner.cpp
--- a/Engine/Source/Editor/BlueprintGraph/Private/BlueprintEventNodeSpawner.cpp
+++ b/Engine/Source/Editor/BlueprintGraph/Private/BlueprintEventNodeSpawner.cpp
@@ -100,7 +100,11 @@ UBlueprintEventNodeSpawner* UBlueprintEventNodeSpawner::Create(UFunction const*
 	FString const FuncCategory = UK2Node_CallFunction::GetDefaultCategoryForFunction(EventFunc, TEXT(""));
 	MenuSignature.Category = FText::FromString(LOCTEXT("AddEventCategory", "Add Event").ToString() + TEXT("|") + FuncCategory);
 	//MenuSignature.Tooltip, will be pulled from the node template
-	MenuSignature.Keywords = UK2Node_CallFunction::GetKeywordsForFunction(EventFunc).AppendChar(TEXT(' '));
+	MenuSignature.Keywords = UK2Node_CallFunction::GetKeywordsForFunction(EventFunc);
+	if (MenuSignature.Keywords.IsEmpty())
+	{
+		MenuSignature.Keywords = FText::FromString(TEXT(" "));
+	}
 	MenuSignature.IconName = TEXT("GraphEditor.Event_16x");
 
 	return NodeSpawner;
diff --git a/Engine/Source/Editor/BlueprintGraph/Private/BlueprintFunctionNodeSpawner.cpp b/Engine/Source/Editor/BlueprintGraph/Private/BlueprintFunctionNodeSpawner.cpp
--- a/Engine/Source/Editor/BlueprintGraph/Private/BlueprintFunctionNodeSpawner.cpp
+++ b/Engine/Source/Editor/BlueprintGraph/Private/BlueprintFunctionNodeSpawner.cpp
@@ -245,8 +245,13 @@ UBlueprintFunctionNodeSpawner* UBlueprintFunctionNodeSpawner::Create(TSubclassOf
 	MenuSignature.MenuName = UK2Node_CallFunction::GetUserFacingFunctionName(Function);
 	MenuSignature.Category = FText::FromString( UK2Node_CallFunction::GetDefaultCategoryForFunction(Function, TEXT("")) );
 	MenuSignature.Tooltip  = FText::FromString( UK2Node_CallFunction::GetDefaultTooltipForFunction(Function) );
-	// add at least one character, so that PrimeDefaultMenuSignature() doesn't attempt to query the template node
-	MenuSignature.Keywords = UK2Node_CallFunction::GetKeywordsForFunction(Function).AppendChar(TEXT(' '));
+	// add at least one character, so that PrimeDefaultUiSpec() doesn't attempt to query the template node
+	MenuSignature.Keywords = UK2Node_CallFunction::GetKeywordsForFunction(Function);
+	if (MenuSignature.Keywords.IsEmpty())
+	{
+		MenuSignature.Keywords = FText::FromString(TEXT(" "));
+	}
+
 	MenuSignature.IconName = UK2Node_CallFunction::GetPaletteIconForFunction(Function, MenuSignature.IconTint);
 
 	if (MenuSignature.Category.IsEmpty())
diff --git a/Engine/Source/Editor/BlueprintGraph/Private/BlueprintNodeSpawner.cpp b/Engine/Source/Editor/BlueprintGraph/Private/BlueprintNodeSpawner.cpp
--- a/Engine/Source/Editor/BlueprintGraph/Private/BlueprintNodeSpawner.cpp
+++ b/Engine/Source/Editor/BlueprintGraph/Private/BlueprintNodeSpawner.cpp
@@ -162,6 +162,11 @@ FBlueprintActionUiSpec const& UBlueprintNodeSpawner::PrimeDefaultUiSpec(UEdGraph
 		NodeTemplate = bTemplateNodeFetched ? NodeTemplate : GetTemplateNode(TargetGraph);
 		if (NodeTemplate != nullptr)
 		{
+			if (NodeClass == UK2Node_IfThenElse::StaticClass())
+			{
+				bool bIsAvail = true;
+				bIsAvail = false;
+			}
 			MenuSignature.Keywords = NodeTemplate->GetKeywords();
 		}
 		// if a target graph was provided, then we've done all we can to spawn a
@@ -169,7 +174,7 @@ FBlueprintActionUiSpec const& UBlueprintNodeSpawner::PrimeDefaultUiSpec(UEdGraph
 		if (MenuSignature.Keywords.IsEmpty() && (TargetGraph != nullptr))
 		{
 			// want to set it to something so we won't end up back in this condition
-			MenuSignature.Keywords = TEXT(" ");
+			MenuSignature.Keywords = FText::FromString(TEXT(" "));
 		}
 		bTemplateNodeFetched = true;
 	}
diff --git a/Engine/Source/Editor/BlueprintGraph/Private/BlueprintVariableNodeSpawner.cpp b/Engine/Source/Editor/BlueprintGraph/Private/BlueprintVariableNodeSpawner.cpp
--- a/Engine/Source/Editor/BlueprintGraph/Private/BlueprintVariableNodeSpawner.cpp
+++ b/Engine/Source/Editor/BlueprintGraph/Private/BlueprintVariableNodeSpawner.cpp
@@ -56,11 +56,15 @@ UBlueprintVariableNodeSpawner* UBlueprintVariableNodeSpawner::Create(TSubclassOf
 		MenuSignature.MenuName = FText::Format(LOCTEXT("SetterMenuName", "Set {0}"), VarName);
 		MenuSignature.Tooltip  = UK2Node_VariableSet::GetPropertyTooltip(VarProperty);
 	}
-	// add at least one character, so that PrimeDefaultMenuSignature() doesn't 
+	// add at least one character, so that PrimeDefaultUiSpec() doesn't 
 	// attempt to query the template node
 	//
 	// @TODO: maybe UPROPERTY() fields should have keyword metadata like functions
-	MenuSignature.Keywords = TEXT(" ");
+	if (MenuSignature.Keywords.IsEmpty())
+	{
+		// want to set it to something so we won't end up back in this condition
+		MenuSignature.Keywords = FText::FromString(TEXT(" "));
+	}
 	MenuSignature.IconName = UK2Node_Variable::GetVarIconFromPinType(NodeSpawner->GetVarType(), MenuSignature.IconTint);
 
 	//--------------------------------------
@@ -125,9 +129,13 @@ UBlueprintVariableNodeSpawner* UBlueprintVariableNodeSpawner::Create(TSubclassOf
 		MenuSignature.MenuName = FText::Format(LOCTEXT("LocalSetterMenuName", "Set {0}"), VarName);
 		MenuSignature.Tooltip  = UK2Node_VariableSet::GetBlueprintVarTooltip(VarDesc);
 	}
-	// add at least one character, so that PrimeDefaultMenuSignature() doesn't 
+	// add at least one character, so that PrimeDefaultUiSpec() doesn't 
 	// attempt to query the template node
-	MenuSignature.Keywords = TEXT(" ");
+	if (MenuSignature.Keywords.IsEmpty())
+	{
+		// want to set it to something so we won't end up back in this condition
+		MenuSignature.Keywords = FText::FromString(TEXT(" "));
+	}
 	MenuSignature.IconName = UK2Node_Variable::GetVarIconFromPinType(NodeSpawner->GetVarType(), MenuSignature.IconTint);
 
 	return NodeSpawner;
diff --git a/Engine/Source/Editor/BlueprintGraph/Private/K2Node_CallFunction.cpp b/Engine/Source/Editor/BlueprintGraph/Private/K2Node_CallFunction.cpp
--- a/Engine/Source/Editor/BlueprintGraph/Private/K2Node_CallFunction.cpp
+++ b/Engine/Source/Editor/BlueprintGraph/Private/K2Node_CallFunction.cpp
@@ -1337,7 +1337,7 @@ FString UK2Node_CallFunction::GetDefaultCategoryForFunction(const UFunction* Fun
 }
 
 
-FString UK2Node_CallFunction::GetKeywordsForFunction(const UFunction* Function)
+FText UK2Node_CallFunction::GetKeywordsForFunction(const UFunction* Function)
 {
 	// If the friendly name and real function name do not match add the real function name friendly name as a keyword.
 	FString Keywords;
@@ -1352,15 +1352,18 @@ FString UK2Node_CallFunction::GetKeywordsForFunction(const UFunction* Function)
 		Keywords += GetCompactNodeTitle(Function);
 	}
 
-	FString MetaKeywords = Function->GetMetaData(FBlueprintMetadata::MD_FunctionKeywords);
+	FText MetadataKeywords = Function->GetMetaDataText(FBlueprintMetadata::MD_FunctionKeywords, TEXT("UObjectKeywords"), Function->GetFullGroupName(false));
+	FText ResultKeywords;
 
-	if (!MetaKeywords.IsEmpty())
+	if (!MetadataKeywords.IsEmpty())
 	{
-		Keywords.AppendChar(TEXT(' '));
-		Keywords += MetaKeywords;
+		FFormatNamedArguments Args;
+		Args.Add(TEXT("Name"), FText::FromString(Keywords));
+		Args.Add(TEXT("MetadataKeywords"), MetadataKeywords);
+		ResultKeywords = FText::Format(FText::FromString("{Name} {MetadataKeywords}"), Args);
 	}
 
-	return Keywords;
+	return ResultKeywords;
 }
 
 void UK2Node_CallFunction::SetFromFunction(const UFunction* Function)
diff --git a/Engine/Source/Editor/BlueprintGraph/Private/K2Node_EnumEquality.cpp b/Engine/Source/Editor/BlueprintGraph/Private/K2Node_EnumEquality.cpp
--- a/Engine/Source/Editor/BlueprintGraph/Private/K2Node_EnumEquality.cpp
+++ b/Engine/Source/Editor/BlueprintGraph/Private/K2Node_EnumEquality.cpp
@@ -40,11 +40,6 @@ FText UK2Node_EnumEquality::GetNodeTitle(ENodeTitleType::Type TitleType) const
 	return NSLOCTEXT("K2Node_EnumEquality", "EqualEnum", "Equal (Enum)");
 }
 
-FString UK2Node_EnumEquality::GetKeywords() const
-{
-	return TEXT("==");
-}
-
 void UK2Node_EnumEquality::PostReconstructNode()
 {
 	PinConnectionListChanged(GetInput1Pin());
diff --git a/Engine/Source/Editor/BlueprintGraph/Private/K2Node_EnumInequality.cpp b/Engine/Source/Editor/BlueprintGraph/Private/K2Node_EnumInequality.cpp
--- a/Engine/Source/Editor/BlueprintGraph/Private/K2Node_EnumInequality.cpp
+++ b/Engine/Source/Editor/BlueprintGraph/Private/K2Node_EnumInequality.cpp
@@ -25,10 +25,6 @@ FText UK2Node_EnumInequality::GetNodeTitle(ENodeTitleType::Type TitleType) const
 	return NSLOCTEXT("K2Node_EnumInequality", "NotEqualEnum", "NotEqual (Enum)");
 }
 
-FString UK2Node_EnumInequality::GetKeywords() const
-{
-	return TEXT("!=");
-}
 void UK2Node_EnumInequality::GetConditionalFunction(FName& FunctionName, UClass** FunctionClass) const
 {
 	FunctionName = GET_FUNCTION_NAME_CHECKED(UKismetMathLibrary, NotEqual_ByteByte);
diff --git a/Engine/Source/Editor/BlueprintGraph/Private/K2Node_Event.cpp b/Engine/Source/Editor/BlueprintGraph/Private/K2Node_Event.cpp
--- a/Engine/Source/Editor/BlueprintGraph/Private/K2Node_Event.cpp
+++ b/Engine/Source/Editor/BlueprintGraph/Private/K2Node_Event.cpp
@@ -141,9 +141,9 @@ FText UK2Node_Event::GetTooltipText() const
 	return CachedTooltip;
 }
 
-FString UK2Node_Event::GetKeywords() const
+FText UK2Node_Event::GetKeywords() const
 {
-	FString Keywords;
+	FText Keywords;
 
 	UFunction* Function = EventReference.ResolveMember<UFunction>(GetBlueprintClassFromNode());
 	if (Function != NULL)
diff --git a/Engine/Source/Editor/BlueprintGraph/Private/K2Node_GetInputAxisKeyValue.cpp b/Engine/Source/Editor/BlueprintGraph/Private/K2Node_GetInputAxisKeyValue.cpp
--- a/Engine/Source/Editor/BlueprintGraph/Private/K2Node_GetInputAxisKeyValue.cpp
+++ b/Engine/Source/Editor/BlueprintGraph/Private/K2Node_GetInputAxisKeyValue.cpp
@@ -52,11 +52,6 @@ FText UK2Node_GetInputAxisKeyValue::GetNodeTitle(ENodeTitleType::Type TitleType)
 	return CachedNodeTitle;
 }
 
-FString UK2Node_GetInputAxisKeyValue::GetKeywords() const
-{
-	return TEXT("Get");
-}
-
 FText UK2Node_GetInputAxisKeyValue::GetTooltipText() const
 {
 	if (CachedTooltip.IsOutOfDate(this))
diff --git a/Engine/Source/Editor/BlueprintGraph/Private/K2Node_GetInputAxisValue.cpp b/Engine/Source/Editor/BlueprintGraph/Private/K2Node_GetInputAxisValue.cpp
--- a/Engine/Source/Editor/BlueprintGraph/Private/K2Node_GetInputAxisValue.cpp
+++ b/Engine/Source/Editor/BlueprintGraph/Private/K2Node_GetInputAxisValue.cpp
@@ -51,11 +51,6 @@ FText UK2Node_GetInputAxisValue::GetNodeTitle(ENodeTitleType::Type TitleType) co
 	return CachedNodeTitle;
 }
 
-FString UK2Node_GetInputAxisValue::GetKeywords() const
-{
-	return TEXT("Get");
-}
-
 FText UK2Node_GetInputAxisValue::GetTooltipText() const
 {
 	if (CachedTooltip.IsOutOfDate(this))
diff --git a/Engine/Source/Editor/BlueprintGraph/Private/K2Node_GetInputVectorAxisValue.cpp b/Engine/Source/Editor/BlueprintGraph/Private/K2Node_GetInputVectorAxisValue.cpp
--- a/Engine/Source/Editor/BlueprintGraph/Private/K2Node_GetInputVectorAxisValue.cpp
+++ b/Engine/Source/Editor/BlueprintGraph/Private/K2Node_GetInputVectorAxisValue.cpp
@@ -19,11 +19,6 @@ void UK2Node_GetInputVectorAxisValue::Initialize(const FKey AxisKey)
 	SetFromFunction(AActor::StaticClass()->FindFunctionByName(TEXT("GetInputVectorAxisValue")));
 }
 
-FString UK2Node_GetInputVectorAxisValue::GetKeywords() const
-{
-	return TEXT("Get");
-}
-
 FText UK2Node_GetInputVectorAxisValue::GetTooltipText() const
 {
 	FFormatNamedArguments Args;
diff --git a/Engine/Source/Editor/BlueprintGraph/Private/K2Node_IfThenElse.cpp b/Engine/Source/Editor/BlueprintGraph/Private/K2Node_IfThenElse.cpp
--- a/Engine/Source/Editor/BlueprintGraph/Private/K2Node_IfThenElse.cpp
+++ b/Engine/Source/Editor/BlueprintGraph/Private/K2Node_IfThenElse.cpp
@@ -117,11 +117,6 @@ FText UK2Node_IfThenElse::GetTooltipText() const
 	return LOCTEXT("BrancStatement_Tooltip", "Branch Statement\nIf Condition is true, execution goes to True, otherwise it goes to False");
 }
 
-FString UK2Node_IfThenElse::GetKeywords() const
-{
-	return *LOCTEXT("BranchKeywords", "if bool").ToString();
-}
-
 UEdGraphPin* UK2Node_IfThenElse::GetThenPin() const
 {
 	const UEdGraphSchema_K2* K2Schema = GetDefault<UEdGraphSchema_K2>();
diff --git a/Engine/Source/Editor/BlueprintGraph/Private/K2Node_MacroInstance.cpp b/Engine/Source/Editor/BlueprintGraph/Private/K2Node_MacroInstance.cpp
--- a/Engine/Source/Editor/BlueprintGraph/Private/K2Node_MacroInstance.cpp
+++ b/Engine/Source/Editor/BlueprintGraph/Private/K2Node_MacroInstance.cpp
@@ -143,13 +143,13 @@ FText UK2Node_MacroInstance::GetTooltipText() const
 	return CachedTooltip;
 }
 
-FString UK2Node_MacroInstance::GetKeywords() const
+FText UK2Node_MacroInstance::GetKeywords() const
 {
 	// This might need to be refined, but seems better than providing no keywords:
 	UBlueprint* MacroBP = MacroGraphReference.GetBlueprint();
 	if (MacroBP)
 	{
-		return MacroBP->BlueprintDescription;
+		return FText::FromString(MacroBP->BlueprintDescription);
 	}
 	return UK2Node_Tunnel::GetKeywords();
 }
diff --git a/Engine/Source/Editor/BlueprintGraph/Private/K2Node_Select.cpp b/Engine/Source/Editor/BlueprintGraph/Private/K2Node_Select.cpp
--- a/Engine/Source/Editor/BlueprintGraph/Private/K2Node_Select.cpp
+++ b/Engine/Source/Editor/BlueprintGraph/Private/K2Node_Select.cpp
@@ -291,11 +291,6 @@ FText UK2Node_Select::GetTooltipText() const
 	return LOCTEXT("SelectNodeTooltip", "Return the option at Index, (first option is indexed at 0)");
 }
 
-FString UK2Node_Select::GetKeywords() const
-{
-	return TEXT("Ternary If");
-}
-
 FText UK2Node_Select::GetNodeTitle(ENodeTitleType::Type TitleType) const
 {
 	return LOCTEXT("Select", "Select");
diff --git a/Engine/Source/Editor/BlueprintGraph/Private/K2Node_Self.cpp b/Engine/Source/Editor/BlueprintGraph/Private/K2Node_Self.cpp
--- a/Engine/Source/Editor/BlueprintGraph/Private/K2Node_Self.cpp
+++ b/Engine/Source/Editor/BlueprintGraph/Private/K2Node_Self.cpp
@@ -50,9 +50,9 @@ FText UK2Node_Self::GetTooltipText() const
 	return NSLOCTEXT("K2Node", "GetSelfReference", "Gets a reference to this instance of the blueprint");
 }
 
-FString UK2Node_Self::GetKeywords() const
+FText UK2Node_Self::GetKeywords() const
 {
-	return TEXT("This");
+	return LOCTEXT("SelfKeywords", "This");
 }
 
 FText UK2Node_Self::GetNodeTitle(ENodeTitleType::Type TitleType) const
diff --git a/Engine/Source/Editor/BlueprintGraph/Public/BlueprintNodeSpawner.h b/Engine/Source/Editor/BlueprintGraph/Public/BlueprintNodeSpawner.h
--- a/Engine/Source/Editor/BlueprintGraph/Public/BlueprintNodeSpawner.h
+++ b/Engine/Source/Editor/BlueprintGraph/Public/BlueprintNodeSpawner.h
@@ -24,7 +24,7 @@ struct FBlueprintActionUiSpec
 	FText   MenuName;
 	FText   Category;
 	FText   Tooltip;
-	FString Keywords;
+	FText	Keywords;
 	FName   IconName;
 	FLinearColor IconTint;
 };
diff --git a/Engine/Source/Editor/GraphEditor/Private/SGraphActionMenu.cpp b/Engine/Source/Editor/GraphEditor/Private/SGraphActionMenu.cpp
--- a/Engine/Source/Editor/GraphEditor/Private/SGraphActionMenu.cpp
+++ b/Engine/Source/Editor/GraphEditor/Private/SGraphActionMenu.cpp
@@ -664,7 +664,7 @@ void SGraphActionMenu::GenerateFilteredItems(bool bPreserveExpansion)
 			// Combine the actions string, separate with \n so terms don't run into each other, and remove the spaces (incase the user is searching for a variable)
 			// In the case of groups containing multiple actions, they will have been created and added at the same place in the code, using the same description
 			// and keywords, so we only need to use the first one for filtering.
-			FString SearchText = CurrentAction.Actions[0]->MenuDescription.ToString() + LINE_TERMINATOR + CurrentAction.Actions[0]->GetSearchTitle().ToString() + LINE_TERMINATOR + CurrentAction.Actions[0]->Keywords + LINE_TERMINATOR +CurrentAction.Actions[0]->Category;
+			FString SearchText = CurrentAction.Actions[0]->GetSearchTitle().ToString() + LINE_TERMINATOR + CurrentAction.Actions[0]->GetSearchKeywords().ToString() + LINE_TERMINATOR +CurrentAction.Actions[0]->Category;
 			SearchText = SearchText.Replace( TEXT( " " ), TEXT( "" ) );
 			// Get the 'weight' of this in relation to the filter
 			EachWeight = GetActionFilteredWeight( CurrentAction, FilterTerms, SanitizedFilterTerms );
@@ -768,11 +768,11 @@ int32 SGraphActionMenu::GetActionFilteredWeight( const FGraphActionListBuilderBa
 		// Combine the actions string, separate with \n so terms don't run into each other, and remove the spaces (incase the user is searching for a variable)
 		// In the case of groups containing multiple actions, they will have been created and added at the same place in the code, using the same description
 		// and keywords, so we only need to use the first one for filtering.
-		FString SearchText = InCurrentAction.Actions[0]->MenuDescription.ToString() + LINE_TERMINATOR + InCurrentAction.Actions[Action]->GetSearchTitle().ToString() + LINE_TERMINATOR + InCurrentAction.Actions[Action]->Keywords + LINE_TERMINATOR +InCurrentAction.Actions[Action]->Category;
+		FString SearchText = InCurrentAction.Actions[Action]->GetSearchTitle().ToString() + LINE_TERMINATOR + InCurrentAction.Actions[Action]->GetSearchKeywords().ToString() + LINE_TERMINATOR +InCurrentAction.Actions[Action]->Category;
 		SearchText = SearchText.Replace( TEXT( " " ), TEXT( "" ) );
 
 		// First the keywords
-		InCurrentAction.Actions[Action]->Keywords.ParseIntoArray( EachEntry.Array, TEXT(" "), true );
+		InCurrentAction.Actions[Action]->GetSearchKeywords().ToString().ParseIntoArray( EachEntry.Array, TEXT(" "), true );
 		EachEntry.Weight = 10;
 		WeightedArrayList.Add( EachEntry );
 
diff --git a/Engine/Source/Editor/UnrealEd/Private/Commandlets/DumpBlueprintsInfoCommandlet.cpp b/Engine/Source/Editor/UnrealEd/Private/Commandlets/DumpBlueprintsInfoCommandlet.cpp
--- a/Engine/Source/Editor/UnrealEd/Private/Commandlets/DumpBlueprintsInfoCommandlet.cpp
+++ b/Engine/Source/Editor/UnrealEd/Private/Commandlets/DumpBlueprintsInfoCommandlet.cpp
@@ -1690,7 +1690,7 @@ static void DumpBlueprintInfoUtils::DumpActionMenuItem(uint32 Indent, FGraphActi
 		TooltipStr = TooltipStr.Replace(TEXT("\n"), *(IndentedNewline + BuildIndentString(TooltipFieldLabel.Len(), /*bUseSpaces =*/true)));
 
 		ActionEntry += IndentedNewline + TooltipFieldLabel + TooltipStr + "\",";
-		ActionEntry += IndentedNewline + "\"Keywords\"    : \"" + PrimeAction->Keywords + "\",";
+		ActionEntry += IndentedNewline + "\"Keywords\"    : \"" + PrimeAction->Keywords.ToString() + "\",";
 		ActionEntry += IndentedNewline + "\"SearchTitle\" : \"" + PrimeAction->GetSearchTitle().ToString() + "\",";
 		ActionEntry += IndentedNewline + FString::Printf(TEXT("\"Grouping\"    : %d"), PrimeAction->Grouping);
 		
diff --git a/Engine/Source/Runtime/Engine/Classes/EdGraph/EdGraphNode.h b/Engine/Source/Runtime/Engine/Classes/EdGraph/EdGraphNode.h
--- a/Engine/Source/Runtime/Engine/Classes/EdGraph/EdGraphNode.h
+++ b/Engine/Source/Runtime/Engine/Classes/EdGraph/EdGraphNode.h
@@ -309,7 +309,7 @@ public:
 	 *
 	 * @TODO: Should search keywords be localized? Probably.
 	 */
-	virtual FString GetKeywords() const { return TEXT(""); }
+	virtual FText GetKeywords() const;
 	 
 	/**
 	 * Returns the link used for external documentation for the graph node
diff --git a/Engine/Source/Runtime/Engine/Classes/EdGraph/EdGraphSchema.h b/Engine/Source/Runtime/Engine/Classes/EdGraph/EdGraphSchema.h
--- a/Engine/Source/Runtime/Engine/Classes/EdGraph/EdGraphSchema.h
+++ b/Engine/Source/Runtime/Engine/Classes/EdGraph/EdGraphSchema.h
@@ -70,7 +70,7 @@ struct ENGINE_API FEdGraphSchemaAction
 
 	/** This is just an arbitrary dump of extra text that search will match on, in addition to the description and tooltip, e.g., Add might have the keyword Math */
 	UPROPERTY()
-	FString Keywords;
+	FText Keywords;
 
 	/** This is a priority number for overriding alphabetical order in the action list (higher value  == higher in the list) */
 	UPROPERTY()
@@ -84,6 +84,10 @@ struct ENGINE_API FEdGraphSchemaAction
 	UPROPERTY()
 	FText CachedSearchTitle;
 
+	/** Search keywords for the action (doesn't have to be set when instantiated, will be constructed by GetSearchTitle() if left empty)*/
+	UPROPERTY()
+	FText CachedSearchKeywords;
+
 	FEdGraphSchemaAction() 
 		: Grouping(0)
 		, SectionID(0)
@@ -127,17 +131,27 @@ struct ENGINE_API FEdGraphSchemaAction
 	{
 		if(CachedSearchTitle.IsEmpty())
 		{
-			if(const FString* SourceString = FTextInspector::GetSourceString(MenuDescription))
-			{
-				FFormatNamedArguments Args;
-				Args.Add(TEXT("LocalizedTitle"), MenuDescription);
-				Args.Add(TEXT("SourceTitle"), FText::FromString(*SourceString));
-				CachedSearchTitle = FText::Format(FText::FromString("{LocalizedTitle} {SourceTitle}"), Args);
-			}
+			FFormatNamedArguments Args;
+			Args.Add(TEXT("LocalizedTitle"), MenuDescription);
+			Args.Add(TEXT("SourceTitle"), FText::FromString(MenuDescription.BuildSourceString()));
+			CachedSearchTitle = FText::Format(FText::FromString("{LocalizedTitle} {SourceTitle}"), Args);
 		}
 		return CachedSearchTitle;
 	}
 
+	/** Retrieves the full searchable title for this action */
+	FText GetSearchKeywords()
+	{
+		if(CachedSearchKeywords.IsEmpty())
+		{
+			FFormatNamedArguments Args;
+			Args.Add(TEXT("LocalizedKeywords"), Keywords);
+			Args.Add(TEXT("SourceKeywords"), FText::FromString(Keywords.BuildSourceString()));
+			CachedSearchKeywords = FText::Format(FText::FromString("{LocalizedKeywords} {SourceKeywords}"), Args);
+		}
+		return CachedSearchKeywords;
+	}
+
 	// GC.
 	virtual void AddReferencedObjects(FReferenceCollector& Collector) {}
 };
diff --git a/Engine/Source/Runtime/Engine/Classes/Materials/MaterialExpression.h b/Engine/Source/Runtime/Engine/Classes/Materials/MaterialExpression.h
--- a/Engine/Source/Runtime/Engine/Classes/Materials/MaterialExpression.h
+++ b/Engine/Source/Runtime/Engine/Classes/Materials/MaterialExpression.h
@@ -326,7 +326,7 @@ class ENGINE_API UMaterialExpression : public UObject
 
 #if WITH_EDITOR
 	/** Returns the keywords that should be used when searching for this expression */
-	virtual FString GetKeywords() const {return TEXT("");}
+	virtual FText GetKeywords() const {return FText::GetEmpty();}
 
 	/**
 	 * Recursively gets a list of all expressions that are connected to this
diff --git a/Engine/Source/Runtime/Engine/Classes/Materials/MaterialExpressionAdd.h b/Engine/Source/Runtime/Engine/Classes/Materials/MaterialExpressionAdd.h
--- a/Engine/Source/Runtime/Engine/Classes/Materials/MaterialExpressionAdd.h
+++ b/Engine/Source/Runtime/Engine/Classes/Materials/MaterialExpressionAdd.h
@@ -29,7 +29,7 @@ class UMaterialExpressionAdd : public UMaterialExpression
 	virtual int32 Compile(class FMaterialCompiler* Compiler, int32 OutputIndex, int32 MultiplexIndex) override;
 	virtual void GetCaption(TArray<FString>& OutCaptions) const override;
 #if WITH_EDITOR
-	virtual FString GetKeywords() const override {return TEXT("+");}
+	virtual FText GetKeywords() const override {return FText::FromString(TEXT("+"));}
 #endif // WITH_EDITOR
 	// End UMaterialExpression Interface
 
diff --git a/Engine/Source/Runtime/Engine/Classes/Materials/MaterialExpressionDivide.h b/Engine/Source/Runtime/Engine/Classes/Materials/MaterialExpressionDivide.h
--- a/Engine/Source/Runtime/Engine/Classes/Materials/MaterialExpressionDivide.h
+++ b/Engine/Source/Runtime/Engine/Classes/Materials/MaterialExpressionDivide.h
@@ -29,7 +29,7 @@ class UMaterialExpressionDivide : public UMaterialExpression
 	virtual int32 Compile(class FMaterialCompiler* Compiler, int32 OutputIndex, int32 MultiplexIndex) override;
 	virtual void GetCaption(TArray<FString>& OutCaptions) const override;
 #if WITH_EDITOR
-	virtual FString GetKeywords() const override {return TEXT("/");}
+	virtual FText GetKeywords() const override {return FText::FromString(TEXT("/"));}
 #endif // WITH_EDITOR
 	// End UMaterialExpression Interface
 };
diff --git a/Engine/Source/Runtime/Engine/Classes/Materials/MaterialExpressionFmod.h b/Engine/Source/Runtime/Engine/Classes/Materials/MaterialExpressionFmod.h
--- a/Engine/Source/Runtime/Engine/Classes/Materials/MaterialExpressionFmod.h
+++ b/Engine/Source/Runtime/Engine/Classes/Materials/MaterialExpressionFmod.h
@@ -21,7 +21,7 @@ class UMaterialExpressionFmod : public UMaterialExpression
 	virtual int32 Compile(class FMaterialCompiler* Compiler, int32 OutputIndex, int32 MultiplexIndex) override;
 	virtual void GetCaption(TArray<FString>& OutCaptions) const override;
 #if WITH_EDITOR
-	virtual FString GetKeywords() const override {return TEXT("%");}
+	virtual FText GetKeywords() const override {return FText::FromString(TEXT("%"));}
 #endif // WITH_EDITOR
 	// End UMaterialExpression Interface
 };
diff --git a/Engine/Source/Runtime/Engine/Classes/Materials/MaterialExpressionLinearInterpolate.h b/Engine/Source/Runtime/Engine/Classes/Materials/MaterialExpressionLinearInterpolate.h
--- a/Engine/Source/Runtime/Engine/Classes/Materials/MaterialExpressionLinearInterpolate.h
+++ b/Engine/Source/Runtime/Engine/Classes/Materials/MaterialExpressionLinearInterpolate.h
@@ -36,7 +36,7 @@ class UMaterialExpressionLinearInterpolate : public UMaterialExpression
 	virtual int32 Compile(class FMaterialCompiler* Compiler, int32 OutputIndex, int32 MultiplexIndex) override;
 	virtual void GetCaption(TArray<FString>& OutCaptions) const override;
 #if WITH_EDITOR
-	virtual FString GetKeywords() const override {return TEXT("lerp");}
+	virtual FText GetKeywords() const override {return FText::FromString(TEXT("lerp"));}
 #endif // WITH_EDITOR
 	// End UMaterialExpression Interface
 };
diff --git a/Engine/Source/Runtime/Engine/Classes/Materials/MaterialExpressionMultiply.h b/Engine/Source/Runtime/Engine/Classes/Materials/MaterialExpressionMultiply.h
--- a/Engine/Source/Runtime/Engine/Classes/Materials/MaterialExpressionMultiply.h
+++ b/Engine/Source/Runtime/Engine/Classes/Materials/MaterialExpressionMultiply.h
@@ -29,7 +29,7 @@ class UMaterialExpressionMultiply : public UMaterialExpression
 	virtual int32 Compile(class FMaterialCompiler* Compiler, int32 OutputIndex, int32 MultiplexIndex) override;
 	virtual void GetCaption(TArray<FString>& OutCaptions) const override;
 #if WITH_EDITOR
-	virtual FString GetKeywords() const override {return TEXT("*");}
+	virtual FText GetKeywords() const override {return FText::FromString(TEXT("*"));}
 #endif // WITH_EDITOR
 	// End UMaterialExpression Interface
 };
diff --git a/Engine/Source/Runtime/Engine/Classes/Materials/MaterialExpressionOneMinus.h b/Engine/Source/Runtime/Engine/Classes/Materials/MaterialExpressionOneMinus.h
--- a/Engine/Source/Runtime/Engine/Classes/Materials/MaterialExpressionOneMinus.h
+++ b/Engine/Source/Runtime/Engine/Classes/Materials/MaterialExpressionOneMinus.h
@@ -18,7 +18,7 @@ class UMaterialExpressionOneMinus : public UMaterialExpression
 	virtual int32 Compile(class FMaterialCompiler* Compiler, int32 OutputIndex, int32 MultiplexIndex) override;
 	virtual void GetCaption(TArray<FString>& OutCaptions) const override;
 #if WITH_EDITOR
-	virtual FString GetKeywords() const override {return TEXT("1-x");}
+	virtual FText GetKeywords() const override {return FText::FromString(TEXT("1-x"));}
 #endif // WITH_EDITOR
 	// End UMaterialExpression Interface
 
diff --git a/Engine/Source/Runtime/Engine/Classes/Materials/MaterialExpressionSquareRoot.h b/Engine/Source/Runtime/Engine/Classes/Materials/MaterialExpressionSquareRoot.h
--- a/Engine/Source/Runtime/Engine/Classes/Materials/MaterialExpressionSquareRoot.h
+++ b/Engine/Source/Runtime/Engine/Classes/Materials/MaterialExpressionSquareRoot.h
@@ -18,7 +18,7 @@ class UMaterialExpressionSquareRoot : public UMaterialExpression
 	virtual int32 Compile(class FMaterialCompiler* Compiler, int32 OutputIndex, int32 MultiplexIndex) override;
 	virtual void GetCaption(TArray<FString>& OutCaptions) const override;
 #if WITH_EDITOR
-	virtual FString GetKeywords() const override {return TEXT("sqrt");}
+	virtual FText GetKeywords() const override {return FText::FromString(TEXT("sqrt"));}
 #endif // WITH_EDITOR
 	// End UMaterialExpression Interface
 };
diff --git a/Engine/Source/Runtime/Engine/Classes/Materials/MaterialExpressionSubtract.h b/Engine/Source/Runtime/Engine/Classes/Materials/MaterialExpressionSubtract.h
--- a/Engine/Source/Runtime/Engine/Classes/Materials/MaterialExpressionSubtract.h
+++ b/Engine/Source/Runtime/Engine/Classes/Materials/MaterialExpressionSubtract.h
@@ -29,7 +29,7 @@ class UMaterialExpressionSubtract : public UMaterialExpression
 	virtual int32 Compile(class FMaterialCompiler* Compiler, int32 OutputIndex, int32 MultiplexIndex) override;
 	virtual void GetCaption(TArray<FString>& OutCaptions) const override;
 #if WITH_EDITOR
-	virtual FString GetKeywords() const override {return TEXT("-");}
+	virtual FText GetKeywords() const override {return FText::FromString(TEXT("-"));}
 #endif // WITH_EDITOR
 	// End UMaterialExpression Interface
 };
diff --git a/Engine/Source/Runtime/Engine/Private/EdGraph/EdGraphNode.cpp b/Engine/Source/Runtime/Engine/Private/EdGraph/EdGraphNode.cpp
--- a/Engine/Source/Runtime/Engine/Private/EdGraph/EdGraphNode.cpp
+++ b/Engine/Source/Runtime/Engine/Private/EdGraph/EdGraphNode.cpp
@@ -372,6 +372,11 @@ void UEdGraphNode::OnUpdateCommentText( const FString& NewComment )
 	NodeComment	= NewComment;
 }
 
+FText UEdGraphNode::GetKeywords() const
+{
+	return GetClass()->GetMetaDataText(TEXT("Keywords"), TEXT("UObjectKeywords"), GetClass()->GetFullGroupName(false));
+}
+
 #endif	//#if WITH_EDITOR
 
 /////////////////////////////////////////////////////
