WIP - Fixing behaviour tree culling and line drawing

diff --git a/Engine/Source/Editor/GraphEditor/Private/SGraphPanel.cpp b/Engine/Source/Editor/GraphEditor/Private/SGraphPanel.cpp
--- a/Engine/Source/Editor/GraphEditor/Private/SGraphPanel.cpp
+++ b/Engine/Source/Editor/GraphEditor/Private/SGraphPanel.cpp
@@ -384,15 +384,17 @@ int32 SGraphPanel::OnPaint( const FPaintArgs& Args, const FGeometry& AllottedGeo
 				TArray< TSharedRef<SWidget> > NodePins;
 				ChildNode->GetPins(NodePins);
 
-				const FVector2D NodeLoc = ChildNode->GetPosition();
-				const FGeometry SynthesizedNodeGeometry(GraphCoordToPanelCoord(NodeLoc), AllottedGeometry.AbsolutePosition, FVector2D::ZeroVector, 1.f);
+				const FVector2D NodeGraphLoc = ChildNode->GetPosition();
+				const FVector2D NodePanelLoc = GraphCoordToPanelCoord(NodeGraphLoc);
+				const FGeometry SynthesizedNodeGeometry(NodePanelLoc, AllottedGeometry.AbsolutePosition, ChildNode->GetDesiredSize(), 1.f);
 
 				for (TArray< TSharedRef<SWidget> >::TConstIterator NodePinIterator(NodePins); NodePinIterator; ++NodePinIterator)
 				{
 					const SGraphPin& PinWidget = static_cast<const SGraphPin&>((*NodePinIterator).Get());
-					FVector2D PinLoc = NodeLoc + PinWidget.GetNodeOffset();
 
-					const FGeometry SynthesizedPinGeometry(GraphCoordToPanelCoord(PinLoc), AllottedGeometry.AbsolutePosition, FVector2D::ZeroVector, 1.f);
+					FVector2D PinPanelLoc = NodePanelLoc + (PinWidget.GetCachedPinOffsetInNodeSpace() * ZoomFactor);
+					const FGeometry SynthesizedPinGeometry(PinPanelLoc, AllottedGeometry.AbsolutePosition, PinWidget.GetCachedPinSize() * ZoomFactor, 1.f);
+
 					PinGeometries.Add(*NodePinIterator, FArrangedWidget(*NodePinIterator, SynthesizedPinGeometry));
 				}
 
diff --git a/Engine/Source/Editor/GraphEditor/Private/SGraphPin.cpp b/Engine/Source/Editor/GraphEditor/Private/SGraphPin.cpp
--- a/Engine/Source/Editor/GraphEditor/Private/SGraphPin.cpp
+++ b/Engine/Source/Editor/GraphEditor/Private/SGraphPin.cpp
@@ -76,7 +76,7 @@ SGraphPin::SGraphPin()
 	, bShowLabel(true)
 	, bIsMovingLinks(false)
 	, PinColorModifier(FLinearColor::White)
-	, CachedNodeOffset(FVector2D::ZeroVector)
+	, CachedPinGeometry(FVector2D::ZeroVector, FVector2D::ZeroVector)
 {
 	IsEditable = true;
 
@@ -720,8 +720,9 @@ FReply SGraphPin::OnDrop( const FGeometry& MyGeometry, const FDragDropEvent& Dra
 
 void SGraphPin::Tick( const FGeometry& AllottedGeometry, const double InCurrentTime, const float InDeltaTime )
 {
-	CachedNodeOffset = AllottedGeometry.AbsolutePosition/AllottedGeometry.Scale - OwnerNodePtr.Pin()->GetUnscaledPosition();
-	CachedNodeOffset.Y += AllottedGeometry.Size.Y * 0.5f;
+	const FVector2D CachedNodeOffset = (AllottedGeometry.AbsolutePosition/AllottedGeometry.Scale) - OwnerNodePtr.Pin()->GetUnscaledPosition();
+	const FVector2D CachedPinSize = AllottedGeometry.Size;
+	CachedPinGeometry = FSlateRect(CachedNodeOffset, CachedNodeOffset + CachedPinSize);
 }
 
 UEdGraphPin* SGraphPin::GetPinObj() const
@@ -753,9 +754,14 @@ EVisibility SGraphPin::IsPinVisibleAsAdvanced() const
 	return (bIsAdvancedPin && bHideAdvancedPin && bCanBeHidden) ? EVisibility::Collapsed : EVisibility::Visible;
 }
 
-FVector2D SGraphPin::GetNodeOffset() const
+FVector2D SGraphPin::GetCachedPinOffsetInNodeSpace() const
 {
-	return CachedNodeOffset;
+	return CachedPinGeometry.GetTopLeft();
+}
+
+FVector2D SGraphPin::GetCachedPinSize() const
+{
+	return CachedPinGeometry.GetSize();
 }
 
 FText SGraphPin::GetPinLabel() const
diff --git a/Engine/Source/Editor/GraphEditor/Private/SNodePanel.cpp b/Engine/Source/Editor/GraphEditor/Private/SNodePanel.cpp
--- a/Engine/Source/Editor/GraphEditor/Private/SNodePanel.cpp
+++ b/Engine/Source/Editor/GraphEditor/Private/SNodePanel.cpp
@@ -186,6 +186,114 @@ void FGraphSelectionManager::ClickedOnNode(SelectedItemType Node, const FPointer
 	}
 }
 
+/**
+ * Manages the cached information about which nodes are currently visible in the graph
+ */
+class FNodePanelCullingCache
+{
+public:
+	typedef TSlotlessChildren<SNodePanel::SNode> FNodePanelChildren;
+
+	FNodePanelCullingCache()
+		: CachedFrameCounter(0)
+		, CachedZoomAmount(0.0f)
+		, NewlyAddedChildren()
+		, NewAndVisibleChildren()
+	{
+	}
+
+	void EndFrame()
+	{
+		//bCanUpdate = true;
+	}
+
+	void ResetCache()
+	{
+		CachedFrameCounter = 0;
+		CachedZoomAmount = 0.0f;
+		NewlyAddedChildren.Reset();
+		NewAndVisibleChildren.Empty();
+	}
+
+	void OnChildAdded(const TSharedRef<SNodePanel::SNode>& NewChild)
+	{
+		NewlyAddedChildren.Add(NewChild);
+	}
+
+	FNodePanelChildren* GetPotentiallyVisibleChildren(const float ZoomAmount, FNodePanelChildren& AllChildren, FNodePanelChildren& VisibleChildren)
+	{
+		return const_cast<FNodePanelChildren*>(GetPotentiallyVisibleChildrenImpl(ZoomAmount, AllChildren, VisibleChildren));
+	}
+
+	const FNodePanelChildren* GetPotentiallyVisibleChildren(const float ZoomAmount, const FNodePanelChildren& AllChildren, const FNodePanelChildren& VisibleChildren)
+	{
+		return GetPotentiallyVisibleChildrenImpl(ZoomAmount, AllChildren, VisibleChildren);
+	}
+
+private:
+	const FNodePanelChildren* GetPotentiallyVisibleChildrenImpl(const float ZoomAmount, const FNodePanelChildren& AllChildren, const FNodePanelChildren& VisibleChildren)
+	{
+		const bool bCanUpdate = CachedFrameCounter != GFrameCounter;
+		CachedFrameCounter = GFrameCounter;
+
+		// Only update if we've been ticked since the last time GetPotentiallyVisibleChildren was called
+		if (!bCanUpdate)
+		{
+			check(CachedPotentiallyVisibleChildren);
+			return CachedPotentiallyVisibleChildren;
+		}
+
+		//bCanUpdate = false;
+
+		// If the zoom amount has changed then all the children are assumed to be potentially visible
+		// This will essentially update the cached geometry for the entire graph and allow for accurate culling calculations
+		if (ZoomAmount != CachedZoomAmount)
+		{
+			CachedZoomAmount = ZoomAmount;
+			CachedPotentiallyVisibleChildren = &AllChildren;
+		}
+
+		// Newly added children are assumed to be visible until we've generated culling data for them
+		else if (NewlyAddedChildren.Num() > 0)
+		{
+			NewAndVisibleChildren = VisibleChildren;
+			for (auto& Child : NewlyAddedChildren)
+			{
+				NewAndVisibleChildren.Add(Child);
+			}
+			NewlyAddedChildren.Reset();
+			CachedPotentiallyVisibleChildren = &NewAndVisibleChildren;
+		}
+
+		// No special cases, so just use the visible children that were previously calculated
+		else
+		{
+			CachedPotentiallyVisibleChildren = &VisibleChildren;
+		}
+
+		check(CachedPotentiallyVisibleChildren);
+		return CachedPotentiallyVisibleChildren;
+	}
+
+	/** Set in EndFrame so that GetPotentiallyVisibleChildren only updates once per-frame */
+	//bool bCanUpdate;
+
+	/** Frame counter when the result of GetPotentiallyVisibleChildren was cached. To ensure that the cache only updates once per-frame */
+	uint64 CachedFrameCounter;
+
+	/** The zoom amount the last time GetPotentiallyVisibleChildren was called */
+	float CachedZoomAmount;
+
+	/** The cached result of the last call to GetPotentiallyVisibleChildren */
+	const FNodePanelChildren* CachedPotentiallyVisibleChildren;
+
+	/** An array of children that have been added since the last time GetPotentiallyVisibleChildren was called */
+	TArray<TSharedRef<SNodePanel::SNode>> NewlyAddedChildren;
+
+	/** Transient array updated by GetPotentiallyVisibleChildren so that we can avoid having to return a copy of the AllChildren and VisibleChildren arrays (which are the most common cases */
+	FNodePanelChildren NewAndVisibleChildren;
+};
+
 //////////////////////////////////////////////////////////////////////////
 // SNodePanel
 
@@ -200,8 +308,9 @@ namespace NodePanelDefs
 };
 
 SNodePanel::SNodePanel()
-: Children()
-, VisibleChildren()
+	: Children()
+	, VisibleChildren()
+	, CullingCache(MakeShareable(new FNodePanelCullingCache()))
 {
 }
 
@@ -212,7 +321,7 @@ void SNodePanel::OnArrangeChildren(const FGeometry& AllottedGeometry, FArrangedC
 
 void SNodePanel::ArrangeChildNodes(const FGeometry& AllottedGeometry, FArrangedChildren& ArrangedChildren) const
 {
-	const TSlotlessChildren<SNode>& ChildrenToArrange = ArrangedChildren.Accepts(EVisibility::Hidden) ? Children : VisibleChildren;
+	const TSlotlessChildren<SNode>& ChildrenToArrange = ArrangedChildren.Accepts(EVisibility::Hidden) ? Children : *CullingCache->GetPotentiallyVisibleChildren(GetZoomAmount(), Children, VisibleChildren);
 	// First pass nodes
 	for (int32 ChildIndex = 0; ChildIndex < ChildrenToArrange.Num(); ++ChildIndex)
 	{
@@ -244,7 +353,9 @@ FVector2D SNodePanel::ComputeDesiredSize( float ) const
 
 FChildren* SNodePanel::GetChildren()
 {
-	return &VisibleChildren;
+	// We return the potentially visible children here so that even nodes that may ultimately be culled due to being out-of-view can generate accurate cached geometry
+	return CullingCache->GetPotentiallyVisibleChildren(GetZoomAmount(), Children, VisibleChildren);
+	//return &VisibleChildren;
 }
 
 FChildren* SNodePanel::GetAllChildren()
@@ -1089,6 +1200,7 @@ void SNodePanel::AddGraphNode( const TSharedRef<SNodePanel::SNode>& NodeToAdd )
 {
 	Children.Add( NodeToAdd );
 	NodeToWidgetLookup.Add( NodeToAdd->GetObjectBeingDisplayed(), NodeToAdd );
+	CullingCache->OnChildAdded( NodeToAdd );
 }
 
 /** Remove all nodes from the panel */
@@ -1097,6 +1209,7 @@ void SNodePanel::RemoveAllNodes()
 	Children.Empty();
 	NodeToWidgetLookup.Empty();
 	VisibleChildren.Empty();
+	CullingCache->ResetCache();
 }
 
 void SNodePanel::PopulateVisibleChildren(const FGeometry& AllottedGeometry)
@@ -1369,6 +1482,21 @@ bool SNodePanel::IsNodeCulled(const TSharedRef<SNode>& Node, const FGeometry& Al
 {
 	if ( Node->ShouldAllowCulling() )
 	{
+		const FSlateRect ClipAreaRect(
+			//AllottedGeometry.GetDrawSize() * -NodePanelDefs::GuardBandArea, 
+			//AllottedGeometry.GetDrawSize() * ( 1.f + NodePanelDefs::GuardBandArea )
+			FVector2D::ZeroVector,
+			AllottedGeometry.GetLocalSize()
+			);
+		const FSlateRect NodeRect(
+			GraphCoordToPanelCoord( Node->GetPosition() ),
+			GraphCoordToPanelCoord( Node->GetPosition() + Node->GetDesiredSize() )
+			);
+
+		const bool bNodeIsVisible = FSlateRect::DoRectanglesIntersect( ClipAreaRect, NodeRect );
+		return !bNodeIsVisible;
+
+		/*
 		const FVector2D MinClipArea = AllottedGeometry.GetDrawSize() * -NodePanelDefs::GuardBandArea;
 		const FVector2D MaxClipArea = AllottedGeometry.GetDrawSize() * ( 1.f + NodePanelDefs::GuardBandArea );
 		const FVector2D NodeTopLeft = GraphCoordToPanelCoord( Node->GetPosition() );
@@ -1379,6 +1507,7 @@ bool SNodePanel::IsNodeCulled(const TSharedRef<SNode>& Node, const FGeometry& Al
 			NodeBottomRight.Y < MinClipArea.Y ||
 			NodeTopLeft.X > MaxClipArea.X ||
 			NodeTopLeft.Y > MaxClipArea.Y;
+		*/
 	}
 	else
 	{
diff --git a/Engine/Source/Editor/GraphEditor/Public/SGraphPin.h b/Engine/Source/Editor/GraphEditor/Public/SGraphPin.h
--- a/Engine/Source/Editor/GraphEditor/Public/SGraphPin.h
+++ b/Engine/Source/Editor/GraphEditor/Public/SGraphPin.h
@@ -105,8 +105,11 @@ public:
 	/** If pin in node is visible at all */
 	EVisibility IsPinVisibleAsAdvanced() const;
 
-	/** Gets Node Offset */
-	FVector2D GetNodeOffset() const;
+	/** Get the cached offset from owning node to our approximate position within that node so that we can guess the pin location when the node itself is culled from the graph */
+	FVector2D GetCachedPinOffsetInNodeSpace() const;
+
+	/** Get the cached size of the pin the last time it was available so that we can guess the pin size when the node itself is culled from the graph */
+	FVector2D GetCachedPinSize() const;
 
 protected:
 	FText GetPinLabel() const;
@@ -180,8 +183,12 @@ protected:
 	/** Color modifier for use by the connection drawing policy */
 	FLinearColor PinColorModifier;
 
-	/** Cached offset from owning node to approximate position of culled pins */
-	FVector2D CachedNodeOffset;
+	/**
+	 * Minimal cached geometry data used to approximate the pin location when the node itself is culled from the graph
+	 * The top-left of the rect provides the cached offset from owning node to our approximate position within that node
+	 * The remainder of the rect provides the cached size of the pin the last time it was available
+	 */
+	FSlateRect CachedPinGeometry;
 
 	TSet< TWeakObjectPtr<UEdGraphPin> > HoverPinSet;
 
diff --git a/Engine/Source/Editor/GraphEditor/Public/SNodePanel.h b/Engine/Source/Editor/GraphEditor/Public/SNodePanel.h
--- a/Engine/Source/Editor/GraphEditor/Public/SNodePanel.h
+++ b/Engine/Source/Editor/GraphEditor/Public/SNodePanel.h
@@ -231,6 +231,8 @@ namespace ENodeZone
 	};
 }
 
+class FNodePanelCullingCache;
+
 class GRAPHEDITOR_API SNodePanel : public SPanel
 {
 public:
@@ -798,6 +800,9 @@ protected:
 	TSlotlessChildren<SNode> Children;
 	TSlotlessChildren<SNode> VisibleChildren;
 
+	/** Manages the cached information about which nodes are currently visible in the graph */
+	TSharedRef<FNodePanelCullingCache> CullingCache;
+
 	/** The node that the user is dragging. Null when they are not dragging a node. */
 	TWeakPtr<SNode> NodeUnderMousePtr;
 
