Back out changelist 2519615

diff --git a/Engine/Source/Editor/GraphEditor/Private/SGraphPanel.cpp b/Engine/Source/Editor/GraphEditor/Private/SGraphPanel.cpp
--- a/Engine/Source/Editor/GraphEditor/Private/SGraphPanel.cpp
+++ b/Engine/Source/Editor/GraphEditor/Private/SGraphPanel.cpp
@@ -384,17 +384,15 @@ int32 SGraphPanel::OnPaint( const FPaintArgs& Args, const FGeometry& AllottedGeo
 				TArray< TSharedRef<SWidget> > NodePins;
 				ChildNode->GetPins(NodePins);
 
-				const FVector2D NodeGraphLoc = ChildNode->GetPosition();
-				const FVector2D NodePanelLoc = GraphCoordToPanelCoord(NodeGraphLoc);
-				const FGeometry SynthesizedNodeGeometry(NodePanelLoc, AllottedGeometry.AbsolutePosition, ChildNode->GetDesiredSize(), 1.f);
+				const FVector2D NodeLoc = ChildNode->GetPosition();
+				const FGeometry SynthesizedNodeGeometry(GraphCoordToPanelCoord(NodeLoc), AllottedGeometry.AbsolutePosition, FVector2D::ZeroVector, 1.f);
 
 				for (TArray< TSharedRef<SWidget> >::TConstIterator NodePinIterator(NodePins); NodePinIterator; ++NodePinIterator)
 				{
 					const SGraphPin& PinWidget = static_cast<const SGraphPin&>((*NodePinIterator).Get());
+					FVector2D PinLoc = NodeLoc + PinWidget.GetNodeOffset();
 
-					FVector2D PinPanelLoc = NodePanelLoc + (PinWidget.GetCachedPinOffsetInNodeSpace() * ZoomFactor);
-					const FGeometry SynthesizedPinGeometry(PinPanelLoc, AllottedGeometry.AbsolutePosition, PinWidget.GetCachedPinSize() * ZoomFactor, 1.f);
-
+					const FGeometry SynthesizedPinGeometry(GraphCoordToPanelCoord(PinLoc), AllottedGeometry.AbsolutePosition, FVector2D::ZeroVector, 1.f);
 					PinGeometries.Add(*NodePinIterator, FArrangedWidget(*NodePinIterator, SynthesizedPinGeometry));
 				}
 
diff --git a/Engine/Source/Editor/GraphEditor/Private/SGraphPin.cpp b/Engine/Source/Editor/GraphEditor/Private/SGraphPin.cpp
--- a/Engine/Source/Editor/GraphEditor/Private/SGraphPin.cpp
+++ b/Engine/Source/Editor/GraphEditor/Private/SGraphPin.cpp
@@ -76,7 +76,7 @@ SGraphPin::SGraphPin()
 	, bShowLabel(true)
 	, bIsMovingLinks(false)
 	, PinColorModifier(FLinearColor::White)
-	, CachedPinGeometry(FVector2D::ZeroVector, FVector2D::ZeroVector)
+	, CachedNodeOffset(FVector2D::ZeroVector)
 {
 	IsEditable = true;
 
@@ -720,9 +720,8 @@ FReply SGraphPin::OnDrop( const FGeometry& MyGeometry, const FDragDropEvent& Dra
 
 void SGraphPin::Tick( const FGeometry& AllottedGeometry, const double InCurrentTime, const float InDeltaTime )
 {
-	const FVector2D CachedNodeOffset = (AllottedGeometry.AbsolutePosition/AllottedGeometry.Scale) - OwnerNodePtr.Pin()->GetUnscaledPosition();
-	const FVector2D CachedPinSize = AllottedGeometry.Size;
-	CachedPinGeometry = FSlateRect(CachedNodeOffset, CachedNodeOffset + CachedPinSize);
+	CachedNodeOffset = AllottedGeometry.AbsolutePosition/AllottedGeometry.Scale - OwnerNodePtr.Pin()->GetUnscaledPosition();
+	CachedNodeOffset.Y += AllottedGeometry.Size.Y * 0.5f;
 }
 
 UEdGraphPin* SGraphPin::GetPinObj() const
@@ -754,14 +753,9 @@ EVisibility SGraphPin::IsPinVisibleAsAdvanced() const
 	return (bIsAdvancedPin && bHideAdvancedPin && bCanBeHidden) ? EVisibility::Collapsed : EVisibility::Visible;
 }
 
-FVector2D SGraphPin::GetCachedPinOffsetInNodeSpace() const
+FVector2D SGraphPin::GetNodeOffset() const
 {
-	return CachedPinGeometry.GetTopLeft();
-}
-
-FVector2D SGraphPin::GetCachedPinSize() const
-{
-	return CachedPinGeometry.GetSize();
+	return CachedNodeOffset;
 }
 
 FText SGraphPin::GetPinLabel() const
diff --git a/Engine/Source/Editor/GraphEditor/Private/SNodePanel.cpp b/Engine/Source/Editor/GraphEditor/Private/SNodePanel.cpp
--- a/Engine/Source/Editor/GraphEditor/Private/SNodePanel.cpp
+++ b/Engine/Source/Editor/GraphEditor/Private/SNodePanel.cpp
@@ -186,114 +186,6 @@ void FGraphSelectionManager::ClickedOnNode(SelectedItemType Node, const FPointer
 	}
 }
 
-/**
- * Manages the cached information about which nodes are currently visible in the graph
- */
-class FNodePanelCullingCache
-{
-public:
-	typedef TSlotlessChildren<SNodePanel::SNode> FNodePanelChildren;
-
-	FNodePanelCullingCache()
-		: CachedFrameCounter(0)
-		, CachedZoomAmount(0.0f)
-		, NewlyAddedChildren()
-		, NewAndVisibleChildren()
-	{
-	}
-
-	void EndFrame()
-	{
-		//bCanUpdate = true;
-	}
-
-	void ResetCache()
-	{
-		CachedFrameCounter = 0;
-		CachedZoomAmount = 0.0f;
-		NewlyAddedChildren.Reset();
-		NewAndVisibleChildren.Empty();
-	}
-
-	void OnChildAdded(const TSharedRef<SNodePanel::SNode>& NewChild)
-	{
-		NewlyAddedChildren.Add(NewChild);
-	}
-
-	FNodePanelChildren* GetPotentiallyVisibleChildren(const float ZoomAmount, FNodePanelChildren& AllChildren, FNodePanelChildren& VisibleChildren)
-	{
-		return const_cast<FNodePanelChildren*>(GetPotentiallyVisibleChildrenImpl(ZoomAmount, AllChildren, VisibleChildren));
-	}
-
-	const FNodePanelChildren* GetPotentiallyVisibleChildren(const float ZoomAmount, const FNodePanelChildren& AllChildren, const FNodePanelChildren& VisibleChildren)
-	{
-		return GetPotentiallyVisibleChildrenImpl(ZoomAmount, AllChildren, VisibleChildren);
-	}
-
-private:
-	const FNodePanelChildren* GetPotentiallyVisibleChildrenImpl(const float ZoomAmount, const FNodePanelChildren& AllChildren, const FNodePanelChildren& VisibleChildren)
-	{
-		const bool bCanUpdate = CachedFrameCounter != GFrameCounter;
-		CachedFrameCounter = GFrameCounter;
-
-		// Only update if we've been ticked since the last time GetPotentiallyVisibleChildren was called
-		if (!bCanUpdate)
-		{
-			check(CachedPotentiallyVisibleChildren);
-			return CachedPotentiallyVisibleChildren;
-		}
-
-		//bCanUpdate = false;
-
-		// If the zoom amount has changed then all the children are assumed to be potentially visible
-		// This will essentially update the cached geometry for the entire graph and allow for accurate culling calculations
-		if (ZoomAmount != CachedZoomAmount)
-		{
-			CachedZoomAmount = ZoomAmount;
-			CachedPotentiallyVisibleChildren = &AllChildren;
-		}
-
-		// Newly added children are assumed to be visible until we've generated culling data for them
-		else if (NewlyAddedChildren.Num() > 0)
-		{
-			NewAndVisibleChildren = VisibleChildren;
-			for (auto& Child : NewlyAddedChildren)
-			{
-				NewAndVisibleChildren.Add(Child);
-			}
-			NewlyAddedChildren.Reset();
-			CachedPotentiallyVisibleChildren = &NewAndVisibleChildren;
-		}
-
-		// No special cases, so just use the visible children that were previously calculated
-		else
-		{
-			CachedPotentiallyVisibleChildren = &VisibleChildren;
-		}
-
-		check(CachedPotentiallyVisibleChildren);
-		return CachedPotentiallyVisibleChildren;
-	}
-
-	/** Set in EndFrame so that GetPotentiallyVisibleChildren only updates once per-frame */
-	//bool bCanUpdate;
-
-	/** Frame counter when the result of GetPotentiallyVisibleChildren was cached. To ensure that the cache only updates once per-frame */
-	uint64 CachedFrameCounter;
-
-	/** The zoom amount the last time GetPotentiallyVisibleChildren was called */
-	float CachedZoomAmount;
-
-	/** The cached result of the last call to GetPotentiallyVisibleChildren */
-	const FNodePanelChildren* CachedPotentiallyVisibleChildren;
-
-	/** An array of children that have been added since the last time GetPotentiallyVisibleChildren was called */
-	TArray<TSharedRef<SNodePanel::SNode>> NewlyAddedChildren;
-
-	/** Transient array updated by GetPotentiallyVisibleChildren so that we can avoid having to return a copy of the AllChildren and VisibleChildren arrays (which are the most common cases */
-	FNodePanelChildren NewAndVisibleChildren;
-};
-
 //////////////////////////////////////////////////////////////////////////
 // SNodePanel
 
@@ -308,9 +200,8 @@ namespace NodePanelDefs
 };
 
 SNodePanel::SNodePanel()
-	: Children()
-	, VisibleChildren()
-	, CullingCache(MakeShareable(new FNodePanelCullingCache()))
+: Children()
+, VisibleChildren()
 {
 }
 
@@ -321,7 +212,7 @@ void SNodePanel::OnArrangeChildren(const FGeometry& AllottedGeometry, FArrangedC
 
 void SNodePanel::ArrangeChildNodes(const FGeometry& AllottedGeometry, FArrangedChildren& ArrangedChildren) const
 {
-	const TSlotlessChildren<SNode>& ChildrenToArrange = ArrangedChildren.Accepts(EVisibility::Hidden) ? Children : *CullingCache->GetPotentiallyVisibleChildren(GetZoomAmount(), Children, VisibleChildren);
+	const TSlotlessChildren<SNode>& ChildrenToArrange = ArrangedChildren.Accepts(EVisibility::Hidden) ? Children : VisibleChildren;
 	// First pass nodes
 	for (int32 ChildIndex = 0; ChildIndex < ChildrenToArrange.Num(); ++ChildIndex)
 	{
@@ -353,9 +244,7 @@ FVector2D SNodePanel::ComputeDesiredSize( float ) const
 
 FChildren* SNodePanel::GetChildren()
 {
-	// We return the potentially visible children here so that even nodes that may ultimately be culled due to being out-of-view can generate accurate cached geometry
-	return CullingCache->GetPotentiallyVisibleChildren(GetZoomAmount(), Children, VisibleChildren);
-	//return &VisibleChildren;
+	return &VisibleChildren;
 }
 
 FChildren* SNodePanel::GetAllChildren()
@@ -1200,7 +1089,6 @@ void SNodePanel::AddGraphNode( const TSharedRef<SNodePanel::SNode>& NodeToAdd )
 {
 	Children.Add( NodeToAdd );
 	NodeToWidgetLookup.Add( NodeToAdd->GetObjectBeingDisplayed(), NodeToAdd );
-	CullingCache->OnChildAdded( NodeToAdd );
 }
 
 /** Remove all nodes from the panel */
@@ -1209,7 +1097,6 @@ void SNodePanel::RemoveAllNodes()
 	Children.Empty();
 	NodeToWidgetLookup.Empty();
 	VisibleChildren.Empty();
-	CullingCache->ResetCache();
 }
 
 void SNodePanel::PopulateVisibleChildren(const FGeometry& AllottedGeometry)
@@ -1482,21 +1369,6 @@ bool SNodePanel::IsNodeCulled(const TSharedRef<SNode>& Node, const FGeometry& Al
 {
 	if ( Node->ShouldAllowCulling() )
 	{
-		const FSlateRect ClipAreaRect(
-			//AllottedGeometry.GetDrawSize() * -NodePanelDefs::GuardBandArea, 
-			//AllottedGeometry.GetDrawSize() * ( 1.f + NodePanelDefs::GuardBandArea )
-			FVector2D::ZeroVector,
-			AllottedGeometry.GetLocalSize()
-			);
-		const FSlateRect NodeRect(
-			GraphCoordToPanelCoord( Node->GetPosition() ),
-			GraphCoordToPanelCoord( Node->GetPosition() + Node->GetDesiredSize() )
-			);
-
-		const bool bNodeIsVisible = FSlateRect::DoRectanglesIntersect( ClipAreaRect, NodeRect );
-		return !bNodeIsVisible;
-
-		/*
 		const FVector2D MinClipArea = AllottedGeometry.GetDrawSize() * -NodePanelDefs::GuardBandArea;
 		const FVector2D MaxClipArea = AllottedGeometry.GetDrawSize() * ( 1.f + NodePanelDefs::GuardBandArea );
 		const FVector2D NodeTopLeft = GraphCoordToPanelCoord( Node->GetPosition() );
@@ -1507,7 +1379,6 @@ bool SNodePanel::IsNodeCulled(const TSharedRef<SNode>& Node, const FGeometry& Al
 			NodeBottomRight.Y < MinClipArea.Y ||
 			NodeTopLeft.X > MaxClipArea.X ||
 			NodeTopLeft.Y > MaxClipArea.Y;
-		*/
 	}
 	else
 	{
diff --git a/Engine/Source/Editor/GraphEditor/Public/SGraphPin.h b/Engine/Source/Editor/GraphEditor/Public/SGraphPin.h
--- a/Engine/Source/Editor/GraphEditor/Public/SGraphPin.h
+++ b/Engine/Source/Editor/GraphEditor/Public/SGraphPin.h
@@ -105,11 +105,8 @@ public:
 	/** If pin in node is visible at all */
 	EVisibility IsPinVisibleAsAdvanced() const;
 
-	/** Get the cached offset from owning node to our approximate position within that node so that we can guess the pin location when the node itself is culled from the graph */
-	FVector2D GetCachedPinOffsetInNodeSpace() const;
-
-	/** Get the cached size of the pin the last time it was available so that we can guess the pin size when the node itself is culled from the graph */
-	FVector2D GetCachedPinSize() const;
+	/** Gets Node Offset */
+	FVector2D GetNodeOffset() const;
 
 protected:
 	FText GetPinLabel() const;
@@ -183,12 +180,8 @@ protected:
 	/** Color modifier for use by the connection drawing policy */
 	FLinearColor PinColorModifier;
 
-	/**
-	 * Minimal cached geometry data used to approximate the pin location when the node itself is culled from the graph
-	 * The top-left of the rect provides the cached offset from owning node to our approximate position within that node
-	 * The remainder of the rect provides the cached size of the pin the last time it was available
-	 */
-	FSlateRect CachedPinGeometry;
+	/** Cached offset from owning node to approximate position of culled pins */
+	FVector2D CachedNodeOffset;
 
 	TSet< TWeakObjectPtr<UEdGraphPin> > HoverPinSet;
 
diff --git a/Engine/Source/Editor/GraphEditor/Public/SNodePanel.h b/Engine/Source/Editor/GraphEditor/Public/SNodePanel.h
--- a/Engine/Source/Editor/GraphEditor/Public/SNodePanel.h
+++ b/Engine/Source/Editor/GraphEditor/Public/SNodePanel.h
@@ -231,8 +231,6 @@ namespace ENodeZone
 	};
 }
 
-class FNodePanelCullingCache;
-
 class GRAPHEDITOR_API SNodePanel : public SPanel
 {
 public:
@@ -800,9 +798,6 @@ protected:
 	TSlotlessChildren<SNode> Children;
 	TSlotlessChildren<SNode> VisibleChildren;
 
-	/** Manages the cached information about which nodes are currently visible in the graph */
-	TSharedRef<FNodePanelCullingCache> CullingCache;
-
 	/** The node that the user is dragging. Null when they are not dragging a node. */
 	TWeakPtr<SNode> NodeUnderMousePtr;
 
