Graph Editor: Allowed interacting with splines when they overlap with a comment node [UE-13021]

diff --git a/Engine/Source/Editor/GraphEditor/Private/SGraphPanel.cpp b/Engine/Source/Editor/GraphEditor/Private/SGraphPanel.cpp
--- a/Engine/Source/Editor/GraphEditor/Private/SGraphPanel.cpp
+++ b/Engine/Source/Editor/GraphEditor/Private/SGraphPanel.cpp
@@ -464,12 +464,18 @@ int32 SGraphPanel::OnPaint( const FPaintArgs& Args, const FGeometry& AllottedGeo
 			{
 				OverlapData.ComputeBestPin();
 
-				// Only allow spline overlaps when there is no node under the cursor
+				// Only allow spline overlaps when there is no node under the cursor (unless it is a comment box)
 				const FVector2D PaintAbsoluteSpaceMousePos = AllottedGeometry.AbsolutePosition + SavedMousePosForOnPaintEventLocalSpace;
 				const int32 HoveredNodeIndex = SWidget::FindChildUnderPosition(ArrangedChildren, PaintAbsoluteSpaceMousePos);
 				if (HoveredNodeIndex != INDEX_NONE)
 				{
-					OverlapData = FGraphSplineOverlapResult();
+					TSharedRef<SGraphNode> HoveredNode = StaticCastSharedRef<SGraphNode>(ArrangedChildren[HoveredNodeIndex].Widget);
+					UEdGraphNode_Comment* CommentNode = Cast<UEdGraphNode_Comment>(HoveredNode->GetObjectBeingDisplayed());
+					if (CommentNode == nullptr)
+					{
+						// Wasn't a comment node, disallow the spline interaction
+						OverlapData = FGraphSplineOverlapResult();
+					}
 				}
 			}
 
