Editor should not stall when pasting large amounts of text into the Blueprint node context menu filter.

diff --git a/Engine/Source/Editor/GraphEditor/Private/SGraphActionMenu.cpp b/Engine/Source/Editor/GraphEditor/Private/SGraphActionMenu.cpp
--- a/Engine/Source/Editor/GraphEditor/Private/SGraphActionMenu.cpp
+++ b/Engine/Source/Editor/GraphEditor/Private/SGraphActionMenu.cpp
@@ -666,14 +666,20 @@ void SGraphActionMenu::GenerateFilteredItems(bool bPreserveExpansion)
 			// and keywords, so we only need to use the first one for filtering.
 			FString SearchText = CurrentAction.Actions[0]->GetSearchTitle().ToString() + LINE_TERMINATOR + CurrentAction.Actions[0]->GetSearchKeywords().ToString() + LINE_TERMINATOR +CurrentAction.Actions[0]->Category;
 			SearchText = SearchText.Replace( TEXT( " " ), TEXT( "" ) );
-			// Get the 'weight' of this in relation to the filter
-			EachWeight = GetActionFilteredWeight( CurrentAction, FilterTerms, SanitizedFilterTerms );
+
 			FString EachTermSanitized;
 			for (int32 FilterIndex = 0; (FilterIndex < FilterTerms.Num()) && bShowAction; ++FilterIndex)
 			{
 				const bool bMatchesTerm = ( SearchText.Contains( FilterTerms[FilterIndex] ) || ( SearchText.Contains( SanitizedFilterTerms[FilterIndex] ) == true ) );
 				bShowAction = bShowAction && bMatchesTerm;
 			}
+
+			// Only if we are going to show the action do we want to generate the weight of the filter text
+			if (bShowAction)
+			{
+				// Get the 'weight' of this in relation to the filter
+				EachWeight = GetActionFilteredWeight( CurrentAction, FilterTerms, SanitizedFilterTerms );
+			}
 		}
 
 		if (bShowAction)
