This modifies the comment balloons behavior on reroute nodes to be more inline other node comment balloons by enabling a toggle button and displaying by default when a comment is entered in the context menu.
#UE-12866 - Have node comments and reroute node comments act the same
#Branch UE4
#Proj Engine, GraphEditor

#includes some more general comment node / comment bubble fixes

diff --git a/Engine/Source/Editor/BehaviorTreeEditor/Private/SGraphNode_BehaviorTree.cpp b/Engine/Source/Editor/BehaviorTreeEditor/Private/SGraphNode_BehaviorTree.cpp
--- a/Engine/Source/Editor/BehaviorTreeEditor/Private/SGraphNode_BehaviorTree.cpp
+++ b/Engine/Source/Editor/BehaviorTreeEditor/Private/SGraphNode_BehaviorTree.cpp
@@ -497,6 +497,7 @@ void SGraphNode_BehaviorTree::UpdateGraphNode()
 	SAssignNew( CommentBubble, SCommentBubble )
 	.GraphNode( GraphNode )
 	.Text( this, &SGraphNode::GetNodeComment )
+	.OnTextCommitted( this, &SGraphNode::OnCommentTextCommitted )
 	.ColorAndOpacity( CommentColor )
 	.AllowPinning( true )
 	.EnableTitleBarBubble( true )
diff --git a/Engine/Source/Editor/EditorStyle/Private/SlateEditorStyle.cpp b/Engine/Source/Editor/EditorStyle/Private/SlateEditorStyle.cpp
--- a/Engine/Source/Editor/EditorStyle/Private/SlateEditorStyle.cpp
+++ b/Engine/Source/Editor/EditorStyle/Private/SlateEditorStyle.cpp
@@ -3459,19 +3459,19 @@ void FSlateEditorStyle::FStyle::SetupGraphEditorStyles()
 			.SetUncheckedImage( IMAGE_BRUSH( "Icons/icon_Blueprint_CommentBubbleOn_16x", FVector2D(10,10), FLinearColor(1.f, 1.f, 1.f, 0.5f)))
 			.SetUncheckedHoveredImage( IMAGE_BRUSH( "Icons/icon_Blueprint_CommentBubbleOn_16x", FVector2D(10,10), FLinearColor(1.f, 1.f, 1.f, 0.9f)))
 			.SetUncheckedPressedImage( IMAGE_BRUSH( "Icons/icon_Blueprint_CommentBubbleOn_16x", FVector2D(10,10), FLinearColor(1.f, 1.f, 1.f, 1.f)))
-			.SetCheckedImage( IMAGE_BRUSH( "Icons/icon_Blueprint_CommentBubbleOn_16x", FVector2D(10,10), FLinearColor(1.f, 1.f, 1.f, 0.7f)))
+			.SetCheckedImage( IMAGE_BRUSH( "Icons/icon_Blueprint_CommentBubbleOn_16x", FVector2D(10,10), FLinearColor(1.f, 1.f, 1.f, 0.8f)))
 			.SetCheckedHoveredImage( IMAGE_BRUSH( "Icons/icon_Blueprint_CommentBubbleOn_16x", FVector2D(10,10), FLinearColor(1.f, 1.f, 1.f, 1.f)))
 			.SetCheckedPressedImage( IMAGE_BRUSH( "Icons/icon_Blueprint_CommentBubbleOn_16x", FVector2D(10,10), FLinearColor(1.f, 1.f, 1.f, 0.6f)));
 		Set( "CommentBubbleButton", CommentBubbleButton );
 
 		const FCheckBoxStyle CommentBubblePin = FCheckBoxStyle()
 			.SetCheckBoxType( ESlateCheckBoxType::CheckBox )
-			.SetUncheckedImage( IMAGE_BRUSH( "Icons/icon_Blueprint_CommentBubbleUnPin_16x", FVector2D(10,10), FLinearColor::White))
-			.SetUncheckedHoveredImage( IMAGE_BRUSH( "Icons/icon_Blueprint_CommentBubbleUnPin_16x", FVector2D(10,10), FLinearColor::White))
-			.SetUncheckedPressedImage( IMAGE_BRUSH( "Icons/icon_Blueprint_CommentBubblePin_16x", FVector2D(10,10), FLinearColor::White))
-			.SetCheckedImage( IMAGE_BRUSH( "Icons/icon_Blueprint_CommentBubblePin_16x", FVector2D(10,10), FLinearColor::White))
-			.SetCheckedHoveredImage( IMAGE_BRUSH( "Icons/icon_Blueprint_CommentBubblePin_16x", FVector2D(10,10), FLinearColor::White))
-			.SetCheckedPressedImage( IMAGE_BRUSH( "Icons/icon_Blueprint_CommentBubbleUnPin_16x", FVector2D(10,10), FLinearColor::White));
+			.SetUncheckedImage( IMAGE_BRUSH( "Icons/icon_Blueprint_CommentBubbleUnPin_16x", FVector2D(10,10), FLinearColor(1.f, 1.f, 1.f, 0.5f)))
+			.SetUncheckedHoveredImage( IMAGE_BRUSH( "Icons/icon_Blueprint_CommentBubbleUnPin_16x", FVector2D(10,10), FLinearColor(1.f, 1.f, 1.f, 0.9f)))
+			.SetUncheckedPressedImage( IMAGE_BRUSH( "Icons/icon_Blueprint_CommentBubblePin_16x", FVector2D(10,10), FLinearColor(1.f, 1.f, 1.f, 1.f)))
+			.SetCheckedImage( IMAGE_BRUSH( "Icons/icon_Blueprint_CommentBubblePin_16x", FVector2D(10,10), FLinearColor(1.f, 1.f, 1.f, 0.8f)))
+			.SetCheckedHoveredImage( IMAGE_BRUSH( "Icons/icon_Blueprint_CommentBubblePin_16x", FVector2D(10,10), FLinearColor(1.f, 1.f, 1.f, 1.f)))
+			.SetCheckedPressedImage( IMAGE_BRUSH( "Icons/icon_Blueprint_CommentBubbleUnPin_16x", FVector2D(10,10), FLinearColor(1.f, 1.f, 1.f, 0.6f)));
 		Set( "CommentBubblePin", CommentBubblePin );
 
 
diff --git a/Engine/Source/Editor/EnvironmentQueryEditor/Private/SGraphNode_EnvironmentQuery.cpp b/Engine/Source/Editor/EnvironmentQueryEditor/Private/SGraphNode_EnvironmentQuery.cpp
--- a/Engine/Source/Editor/EnvironmentQueryEditor/Private/SGraphNode_EnvironmentQuery.cpp
+++ b/Engine/Source/Editor/EnvironmentQueryEditor/Private/SGraphNode_EnvironmentQuery.cpp
@@ -300,6 +300,7 @@ void SGraphNode_EnvironmentQuery::UpdateGraphNode()
 	SAssignNew(CommentBubble, SCommentBubble)
 		.GraphNode(GraphNode)
 		.Text(this, &SGraphNode::GetNodeComment)
+		.OnTextCommitted( this, &SGraphNode::OnCommentTextCommitted )
 		.ColorAndOpacity(CommentColor)
 		.AllowPinning(true)
 		.EnableTitleBarBubble(true)
diff --git a/Engine/Source/Editor/GraphEditor/Private/KismetNodes/SGraphNodeK2ArrayFunction.cpp b/Engine/Source/Editor/GraphEditor/Private/KismetNodes/SGraphNodeK2ArrayFunction.cpp
--- a/Engine/Source/Editor/GraphEditor/Private/KismetNodes/SGraphNodeK2ArrayFunction.cpp
+++ b/Engine/Source/Editor/GraphEditor/Private/KismetNodes/SGraphNodeK2ArrayFunction.cpp
@@ -166,11 +166,13 @@ void SGraphNodeK2ArrayFunction::UpdateGraphNode()
 
 		// Create comment bubble
 		TSharedPtr<SCommentBubble> CommentBubble;
+		const FSlateColor CommentColor = GetDefault<UGraphEditorSettings>()->DefaultCommentNodeTitleColor;
 
 		SAssignNew( CommentBubble, SCommentBubble )
 		.GraphNode( GraphNode )
 		.Text( this, &SGraphNode::GetNodeComment )
-		.ColorAndOpacity( this, &SGraphNodeK2ArrayFunction::GetCommentColor )
+		.OnTextCommitted( this, &SGraphNode::OnCommentTextCommitted )
+		.ColorAndOpacity( CommentColor )
 		.AllowPinning( true )
 		.EnableTitleBarBubble( true )
 		.EnableBubbleCtrls( true )
diff --git a/Engine/Source/Editor/GraphEditor/Private/KismetNodes/SGraphNodeK2Base.cpp b/Engine/Source/Editor/GraphEditor/Private/KismetNodes/SGraphNodeK2Base.cpp
--- a/Engine/Source/Editor/GraphEditor/Private/KismetNodes/SGraphNodeK2Base.cpp
+++ b/Engine/Source/Editor/GraphEditor/Private/KismetNodes/SGraphNodeK2Base.cpp
@@ -145,6 +145,7 @@ void SGraphNodeK2Base::UpdateCompactNode()
 	SAssignNew( CommentBubble, SCommentBubble )
 	.GraphNode( GraphNode )
 	.Text( this, &SGraphNode::GetNodeComment )
+	.OnTextCommitted( this, &SGraphNode::OnCommentTextCommitted )
 	.ColorAndOpacity( CommentColor )
 	.AllowPinning( true )
 	.EnableTitleBarBubble( true )
diff --git a/Engine/Source/Editor/GraphEditor/Private/KismetNodes/SGraphNodeK2Composite.cpp b/Engine/Source/Editor/GraphEditor/Private/KismetNodes/SGraphNodeK2Composite.cpp
--- a/Engine/Source/Editor/GraphEditor/Private/KismetNodes/SGraphNodeK2Composite.cpp
+++ b/Engine/Source/Editor/GraphEditor/Private/KismetNodes/SGraphNodeK2Composite.cpp
@@ -126,6 +126,7 @@ void SGraphNodeK2Composite::UpdateGraphNode()
 	SAssignNew( CommentBubble, SCommentBubble )
 	.GraphNode( GraphNode )
 	.Text( this, &SGraphNode::GetNodeComment )
+	.OnTextCommitted( this, &SGraphNode::OnCommentTextCommitted )
 	.ColorAndOpacity( CommentColor )
 	.AllowPinning( true )
 	.EnableTitleBarBubble( true )
diff --git a/Engine/Source/Editor/GraphEditor/Private/KismetNodes/SGraphNodeK2Var.cpp b/Engine/Source/Editor/GraphEditor/Private/KismetNodes/SGraphNodeK2Var.cpp
--- a/Engine/Source/Editor/GraphEditor/Private/KismetNodes/SGraphNodeK2Var.cpp
+++ b/Engine/Source/Editor/GraphEditor/Private/KismetNodes/SGraphNodeK2Var.cpp
@@ -288,6 +288,7 @@ void SGraphNodeK2Var::UpdateGraphNode()
 	SAssignNew( CommentBubble, SCommentBubble )
 	.GraphNode( GraphNode )
 	.Text( this, &SGraphNode::GetNodeComment )
+	.OnTextCommitted( this, &SGraphNode::OnCommentTextCommitted )
 	.ColorAndOpacity( CommentColor )
 	.AllowPinning( true )
 	.EnableTitleBarBubble( true )
diff --git a/Engine/Source/Editor/GraphEditor/Private/SCommentBubble.cpp b/Engine/Source/Editor/GraphEditor/Private/SCommentBubble.cpp
--- a/Engine/Source/Editor/GraphEditor/Private/SCommentBubble.cpp
+++ b/Engine/Source/Editor/GraphEditor/Private/SCommentBubble.cpp
@@ -16,6 +16,12 @@ namespace SCommentBubbleDefs
 	/** Height of the arrow connecting the bubble to the node */
 	static const float BubbleArrowHeight = 8.f;
 
+	/** Offset from the left edge to comment bubbles arrow center */
+	static const float ArrowCentreOffset = 12.f;
+
+	/** Offset from the left edge to comment bubbles toggle button center */
+	static const float ToggleButtonCentreOffset = 3.f;
+
 	/** Luminance CoEficients */
 	static const FLinearColor LuminanceCoEff( 0.2126f, 0.7152f, 0.0722f, 0.f );
 
@@ -38,6 +44,7 @@ void SCommentBubble::Construct( const FArguments& InArgs )
 
 	GraphNode				= InArgs._GraphNode;
 	CommentAttribute		= InArgs._Text;
+	OnTextCommittedDelegate	= InArgs._OnTextCommitted;
 	ColorAndOpacity			= InArgs._ColorAndOpacity;
 	bAllowPinning			= InArgs._AllowPinning;
 	bEnableTitleBarBubble	= InArgs._EnableTitleBarBubble;
@@ -47,13 +54,9 @@ void SCommentBubble::Construct( const FArguments& InArgs )
 	IsGraphNodeHovered		= InArgs._IsGraphNodeHovered;
 	HintText				= InArgs._HintText.IsSet() ? InArgs._HintText : NSLOCTEXT( "CommentBubble", "EditCommentHint", "Click to edit" );
 	OpacityValue			= SCommentBubbleDefs::FadeDelay;
-	// Create Default toggle and pin handlers
-	CommentBubbleToggleDelegate	= InArgs._OnCommentBubbleToggle.IsBound() ? InArgs._OnCommentBubbleToggle : 
-																			FOnCheckStateChanged::CreateSP( this, &SCommentBubble::OnCommentBubbleToggle );
-	CommentBubblePinnedDelegate = InArgs._OnCommentBubblePinned.IsBound() ? InArgs._OnCommentBubblePinned : 
-																			FOnCheckStateChanged::CreateSP( this, &SCommentBubble::OnPinStateToggle );
-	ToggleButtonCheck			= InArgs._ToggleButtonCheck.IsBound() ?		InArgs._ToggleButtonCheck : 
-																			TAttribute<ECheckBoxState>( this, &SCommentBubble::GetToggleButtonCheck );
+	// Create default delegate/attribute handlers if required
+	ToggleButtonCheck		= InArgs._ToggleButtonCheck.IsBound() ?	InArgs._ToggleButtonCheck : 
+																	TAttribute<ECheckBoxState>( this, &SCommentBubble::GetToggleButtonCheck );
 	// Ensue this value is set to something sensible
 	ForegroundColor = SCommentBubbleDefs::LightForegroundClr;
 
@@ -117,13 +120,19 @@ void SCommentBubble::Tick( const FGeometry& AllottedGeometry, const double InCur
 			}
 		}
 	}
-	if( CachedComment != GraphNode->NodeComment )
+	if( CachedComment != CommentAttribute.Get() )
 	{
-		CachedComment = GraphNode->NodeComment;
+		CachedComment = CommentAttribute.Get();
 		CachedCommentText = FText::FromString( CachedComment );
-		// Comment nodes should always have bCommentBubbleVisible set to true, they also are the nodes that set bInvertCulling.
-		GraphNode->bCommentBubbleVisible = bInvertLODCulling ? GraphNode->bCommentBubbleVisible : !CachedComment.IsEmpty();
-		UpdateBubble();
+		// Call text commit delegate
+		OnTextCommittedDelegate.ExecuteIfBound( CachedCommentText, ETextCommit::Default );
+		// Reflect changes to the Textblock because it doesn't update itself.
+		TextBlock->SetText( CachedCommentText );
+		// Toggle the comment on/off, provided it the parent isn't a comment node
+		if( !bInvertLODCulling )
+		{
+			OnCommentBubbleToggle( CachedComment.IsEmpty() ? ECheckBoxState::Unchecked : ECheckBoxState::Checked );
+		}
 	}
 }
 
@@ -157,7 +166,7 @@ void SCommentBubble::UpdateBubble()
 						SNew( SCheckBox )
 						.Style( &FEditorStyle::Get().GetWidgetStyle<FCheckBoxStyle>( "CommentBubblePin" ))
 						.IsChecked( this, &SCommentBubble::GetPinnedButtonCheck )
-						.OnCheckStateChanged( CommentBubblePinnedDelegate )
+						.OnCheckStateChanged( this, &SCommentBubble::OnPinStateToggle )
 						.ToolTipText( this, &SCommentBubble::GetScaleButtonTooltip )
 						.Cursor( EMouseCursor::Default )
 						.ForegroundColor( this, &SCommentBubble::GetForegroundColor )
@@ -171,7 +180,7 @@ void SCommentBubble::UpdateBubble()
 					SNew( SCheckBox )
 					.Style( &FEditorStyle::Get().GetWidgetStyle< FCheckBoxStyle >( "CommentBubbleButton" ))
 					.IsChecked( ToggleButtonCheck )
-					.OnCheckStateChanged( CommentBubbleToggleDelegate )
+					.OnCheckStateChanged( this, &SCommentBubble::OnCommentBubbleToggle )
 					.ToolTipText( NSLOCTEXT( "CommentBubble", "ToggleCommentTooltip", "Toggle Comment Bubble" ))
 					.Cursor( EMouseCursor::Default )
 					.ForegroundColor( FLinearColor::White )
@@ -189,7 +198,7 @@ void SCommentBubble::UpdateBubble()
 					SNew( SCheckBox )
 					.Style( &FEditorStyle::Get().GetWidgetStyle< FCheckBoxStyle >( "CommentBubbleButton" ))
 					.IsChecked( ToggleButtonCheck )
-					.OnCheckStateChanged( CommentBubbleToggleDelegate )
+					.OnCheckStateChanged( this, &SCommentBubble::OnCommentBubbleToggle )
 					.ToolTipText( NSLOCTEXT( "CommentBubble", "ToggleCommentTooltip", "Toggle Comment Bubble" ))
 					.Cursor( EMouseCursor::Default )
 					.ForegroundColor( FLinearColor::White )
@@ -225,12 +234,13 @@ void SCommentBubble::UpdateBubble()
 						.AutoWidth()
 						[
 							SAssignNew(TextBlock, SMultiLineEditableTextBox)
-							.Text( this, &SCommentBubble::GetCommentText )
+							.Text( CachedCommentText )
 							.HintText( NSLOCTEXT( "CommentBubble", "EditCommentHint", "Click to edit" ))
 							.IsReadOnly( this, &SCommentBubble::IsReadOnly )
 							.Font( FEditorStyle::GetFontStyle( TEXT("Graph.Node.CommentFont")))
 							.SelectAllTextWhenFocused( true )
 							.RevertTextOnEscape( true )
+							.ClearKeyboardFocusOnCommit( true )
 							.ModiferKeyForNewLine( EModifierKey::Control )
 							.ForegroundColor( this, &SCommentBubble::GetTextForegroundColor )
 							.BackgroundColor( this, &SCommentBubble::GetTextBackgroundColor )
@@ -280,7 +290,7 @@ void SCommentBubble::UpdateBubble()
 				SNew( SCheckBox )
 				.Style( &FEditorStyle::Get().GetWidgetStyle< FCheckBoxStyle >( "CommentTitleButton" ))
 				.IsChecked( ToggleButtonCheck )
-				.OnCheckStateChanged( CommentBubbleToggleDelegate )
+				.OnCheckStateChanged( this, &SCommentBubble::OnCommentBubbleToggle )
 				.ToolTipText( NSLOCTEXT( "CommentBubble", "ToggleCommentTooltip", "Toggle Comment Bubble" ))
 				.Cursor( EMouseCursor::Default )
 				.ForegroundColor( this, &SCommentBubble::GetToggleButtonColor )
@@ -302,6 +312,18 @@ FVector2D SCommentBubble::GetOffset() const
 	return FVector2D::ZeroVector;
 }
 
+float SCommentBubble::GetArrowCenterOffset() const
+{
+	float CentreOffset = GraphNode->bCommentBubbleVisible ? SCommentBubbleDefs::ArrowCentreOffset : SCommentBubbleDefs::ToggleButtonCentreOffset;
+	const bool bVisibleAndUnpinned = !GraphNode->bCommentBubblePinned && GraphNode->bCommentBubbleVisible;
+	if( bVisibleAndUnpinned && GraphNode->NodeWidget.IsValid() )
+	{
+		TSharedPtr<SGraphPanel> GraphPanel = GraphNode->NodeWidget.Pin()->GetOwnerPanel();
+		CentreOffset *= GraphPanel.IsValid() ? GraphPanel->GetZoomAmount() : 1.f;
+	}
+	return CentreOffset;
+}
+
 FVector2D SCommentBubble::GetSize() const
 {
 	return GetDesiredSize();
@@ -325,11 +347,6 @@ bool SCommentBubble::IsScalingAllowed() const
 	return !GraphNode->bCommentBubblePinned || !GraphNode->bCommentBubbleVisible;
 }
 
-FText SCommentBubble::GetCommentText() const
-{
-	return CachedCommentText;
-}
-
 FText SCommentBubble::GetScaleButtonTooltip() const
 {
 	if( GraphNode->bCommentBubblePinned )
@@ -350,19 +367,7 @@ FSlateColor SCommentBubble::GetToggleButtonColor() const
 
 FSlateColor SCommentBubble::GetBubbleColor() const
 {
-	bool bColorCommentBubble = true;
-
-	// For comment nodes, check if the color is enabled
-	if (UEdGraphNode_Comment* CommentBubble = Cast<UEdGraphNode_Comment>(GraphNode))
-	{
-		bColorCommentBubble = CommentBubble->bColorCommentBubble;
-	}
-
-	FLinearColor ReturnColor = FLinearColor::White;
-	if (bColorCommentBubble)
-	{
-		ReturnColor = ColorAndOpacity.Get().GetSpecifiedColor();
-	}
+	FLinearColor ReturnColor = ColorAndOpacity.Get().GetSpecifiedColor();
 
 	if(!GraphNode->bIsNodeEnabled)
 	{
@@ -383,10 +388,9 @@ FSlateColor SCommentBubble::GetTextForegroundColor() const
 
 void SCommentBubble::OnCommentTextCommitted( const FText& NewText, ETextCommit::Type CommitInfo )
 {
-	if( GraphNode )
-	{
-		GraphNode->OnUpdateCommentText( NewText.ToString() );
-	}
+	CachedComment = NewText.ToString();
+	CachedCommentText = NewText;
+	OnTextCommittedDelegate.ExecuteIfBound( CachedCommentText, CommitInfo );
 }
 
 EVisibility SCommentBubble::GetToggleButtonVisibility() const
@@ -421,7 +425,7 @@ void SCommentBubble::OnCommentBubbleToggle( ECheckBoxState State )
 	{
 		const FScopedTransaction Transaction( NSLOCTEXT( "CommentBubble", "BubbleVisibility", "Comment Bubble Visibilty" ) );
 		GraphNode->Modify();
-		GraphNode->bCommentBubbleVisible = !GraphNode->bCommentBubbleVisible;
+		GraphNode->bCommentBubbleVisible = State == ECheckBoxState::Checked;
 		OpacityValue = 0.f;
 		UpdateBubble();
 	}
diff --git a/Engine/Source/Editor/GraphEditor/Private/SGraphNode.cpp b/Engine/Source/Editor/GraphEditor/Private/SGraphNode.cpp
--- a/Engine/Source/Editor/GraphEditor/Private/SGraphNode.cpp
+++ b/Engine/Source/Editor/GraphEditor/Private/SGraphNode.cpp
@@ -166,6 +166,11 @@ void SGraphNode::SetTextCommittedEvent(FOnNodeTextCommitted InOnTextCommitted)
 	OnTextCommitted = InOnTextCommitted;
 }
 
+void SGraphNode::OnCommentTextCommitted(const FText& NewComment, ETextCommit::Type CommitInfo)
+{
+	GetNodeObj()->OnUpdateCommentText(NewComment.ToString());
+}
+
 void SGraphNode::SetDisallowedPinConnectionEvent(SGraphEditor::FOnDisallowedPinConnection InOnDisallowedPinConnection)
 {
 	OnDisallowedPinConnection = InOnDisallowedPinConnection;
@@ -845,6 +850,7 @@ void SGraphNode::UpdateGraphNode()
 	SAssignNew( CommentBubble, SCommentBubble )
 	.GraphNode( GraphNode )
 	.Text( this, &SGraphNode::GetNodeComment )
+	.OnTextCommitted( this, &SGraphNode::OnCommentTextCommitted )
 	.ColorAndOpacity( CommentColor )
 	.AllowPinning( true )
 	.EnableTitleBarBubble( true )
diff --git a/Engine/Source/Editor/GraphEditor/Private/SGraphNodeComment.cpp b/Engine/Source/Editor/GraphEditor/Private/SGraphNodeComment.cpp
--- a/Engine/Source/Editor/GraphEditor/Private/SGraphNodeComment.cpp
+++ b/Engine/Source/Editor/GraphEditor/Private/SGraphNodeComment.cpp
@@ -21,7 +21,7 @@ namespace SCommentNodeDefs
 	/** Minimum resize height for comment */
 	static const float MinHeight = 30.0;
 
-	/** TitleBarColor = CommnetColor * TitleBarColorMultiplier */
+	/** TitleBarColor = CommentColor * TitleBarColorMultiplier */
 	static const float TitleBarColorMultiplier = 0.6f;
 
 	/** Titlebar Offset - taken from the widget borders in UpdateGraphNode */
@@ -183,22 +183,23 @@ void SGraphNodeComment::UpdateGraphNode()
 	// Create comment bubble
 	TSharedPtr<SCommentBubble> CommentBubble;
 
-	SAssignNew( CommentBubble, SCommentBubble )
-	.GraphNode( GraphNode )
-	.Text( this, &SGraphNode::GetNodeComment )
-	.ColorAndOpacity( this, &SGraphNodeComment::GetCommentColor )
-	.AllowPinning( true )
-	.EnableTitleBarBubble( false )
-	.EnableBubbleCtrls( false )
-	.GraphLOD( this, &SGraphNode::GetCurrentLOD )
-	.InvertLODCulling( true )
-	.IsGraphNodeHovered( this, &SGraphNode::IsHovered );
-
-	GetOrAddSlot( ENodeZone::TopCenter )
-	.SlotOffset( TAttribute<FVector2D>( CommentBubble.Get(), &SCommentBubble::GetOffset ))
-	.SlotSize( TAttribute<FVector2D>( CommentBubble.Get(), &SCommentBubble::GetSize ))
-	.AllowScaling( TAttribute<bool>( CommentBubble.Get(), &SCommentBubble::IsScalingAllowed ))
-	.VAlign( VAlign_Top )
+	SAssignNew(CommentBubble, SCommentBubble)
+	.GraphNode(GraphNode)
+	.Text(this, &SGraphNodeComment::GetNodeComment)
+	.OnTextCommitted(this, &SGraphNodeComment::OnNameTextCommited)
+	.ColorAndOpacity(this, &SGraphNodeComment::GetCommentBubbleColor )
+	.AllowPinning(true)
+	.EnableTitleBarBubble(false)
+	.EnableBubbleCtrls(false)
+	.GraphLOD(this, &SGraphNode::GetCurrentLOD)
+	.InvertLODCulling(true)
+	.IsGraphNodeHovered(this, &SGraphNode::IsHovered);
+
+	GetOrAddSlot(ENodeZone::TopCenter)
+	.SlotOffset(TAttribute<FVector2D>(CommentBubble.Get(), &SCommentBubble::GetOffset))
+	.SlotSize(TAttribute<FVector2D>(CommentBubble.Get(), &SCommentBubble::GetSize))
+	.AllowScaling(TAttribute<bool>(CommentBubble.Get(), &SCommentBubble::IsScalingAllowed))
+	.VAlign(VAlign_Top)
 	[
 		CommentBubble.ToSharedRef()
 	];
@@ -209,6 +210,12 @@ FVector2D SGraphNodeComment::ComputeDesiredSize( float ) const
 	return UserSize;
 }
 
+FString SGraphNodeComment::GetNodeComment() const
+{
+	const FString Title = GetEditableNodeTitle();;
+	return Title;
+}
+
 FReply SGraphNodeComment::OnMouseButtonDoubleClick( const FGeometry& InMyGeometry, const FPointerEvent& InMouseEvent )
 {
 	// If user double-clicked in the title bar area
@@ -442,9 +449,10 @@ FSlateColor SGraphNodeComment::GetCommentTitleBarColor() const
 FSlateColor SGraphNodeComment::GetCommentBubbleColor() const
 {
 	UEdGraphNode_Comment* CommentNode = Cast<UEdGraphNode_Comment>(GraphNode);
-	if (CommentNode && CommentNode->bColorCommentBubble)
+	if (CommentNode)
 	{
-		const FLinearColor Color = CommentNode->CommentColor * SCommentNodeDefs::TitleBarColorMultiplier;
+		const FLinearColor Color = CommentNode->bColorCommentBubble ?	(CommentNode->CommentColor * SCommentNodeDefs::TitleBarColorMultiplier) :
+																		GetDefault<UGraphEditorSettings>()->DefaultCommentNodeTitleColor;
 		return FLinearColor(Color.R, Color.G, Color.B);
 	}
 	else
diff --git a/Engine/Source/Editor/GraphEditor/Private/SGraphNodeKnot.cpp b/Engine/Source/Editor/GraphEditor/Private/SGraphNodeKnot.cpp
--- a/Engine/Source/Editor/GraphEditor/Private/SGraphNodeKnot.cpp
+++ b/Engine/Source/Editor/GraphEditor/Private/SGraphNodeKnot.cpp
@@ -8,6 +8,17 @@
 #include "GenericCommands.h"
 #include "SInlineEditableTextBlock.h"
 
+namespace SKnotNodeDefinitions
+{
+	/** Offset from the left edge to display comment toggle button at */
+	static const float KnotCenterButtonAdjust = 3.f;
+
+	/** Offset from the left edge to display comment bubbles at */
+	static const float KnotCenterBubbleAdjust = 20.f;
+
+	/** Knot node spacer sizes */
+	static const FVector2D NodeSpacerSize(42.0f, 24.0f);
+}
 
 class FAmbivalentDirectionDragConnection : public FDragConnection
 {
@@ -254,13 +265,10 @@ FReply SGraphPinKnot::OnPinMouseDown(const FGeometry& SenderGeometry, const FPoi
 void SGraphNodeKnot::Construct(const FArguments& InArgs, UK2Node_Knot* InKnot)
 {
 	SGraphNodeDefault::Construct(SGraphNodeDefault::FArguments().GraphNodeObj(InKnot));
-	bCommentPinnedState = InKnot->bCommentBubbleVisible;
 }
 
 void SGraphNodeKnot::UpdateGraphNode()
 {
-	const FVector2D NodeSpacerSize(42.0f, 24.0f);
-
 	InputPins.Empty();
 	OutputPins.Empty();
 
@@ -289,7 +297,7 @@ void SGraphNodeKnot::UpdateGraphNode()
 			[
 				// Grab handle to be able to move the node
 				SNew(SSpacer)
-				.Size(NodeSpacerSize)
+				.Size(SKnotNodeDefinitions::NodeSpacerSize)
 				.Visibility(EVisibility::Visible)
 				.Cursor(EMouseCursor::CardinalCross)
 			]
@@ -323,20 +331,26 @@ void SGraphNodeKnot::UpdateGraphNode()
 			]
 		];
 	// Create comment bubble
-	SAssignNew( CommentBubble, SCommentBubble )
-	.GraphNode( GraphNode )
-	.Text( this, &SGraphNode::GetNodeComment )
-	.EnableBubbleCtrls( true )
-	.ColorAndOpacity( FLinearColor::White )
-	.GraphLOD( this, &SGraphNode::GetCurrentLOD )
-	.IsGraphNodeHovered( this, &SGraphNode::IsHovered )
-	.OnCommentBubbleToggle( this, &SGraphNodeKnot::OnCommentBubbleToggle );
-
-	GetOrAddSlot( ENodeZone::TopCenter )
-	.SlotOffset( TAttribute<FVector2D>( this, &SGraphNodeKnot::GetCommentOffset ))
-	.SlotSize( TAttribute<FVector2D>( CommentBubble.Get(), &SCommentBubble::GetSize ))
-	.AllowScaling( TAttribute<bool>( CommentBubble.Get(), &SCommentBubble::IsScalingAllowed ))
-	.VAlign( VAlign_Top )
+	const FSlateColor CommentColor = GetDefault<UGraphEditorSettings>()->DefaultCommentNodeTitleColor;
+
+	SAssignNew(CommentBubble, SCommentBubble)
+	.GraphNode(GraphNode)
+	.Text(this, &SGraphNode::GetNodeComment)
+	.OnTextCommitted(this, &SGraphNode::OnCommentTextCommitted)
+	.ToggleButtonCheck(this, &SGraphNodeKnot::GetBubbleCheckState)
+	.EnableTitleBarBubble(true)
+	.EnableBubbleCtrls(true)
+	.AllowPinning(true)
+	.ColorAndOpacity(CommentColor)
+	.GraphLOD(this, &SGraphNode::GetCurrentLOD)
+	.IsGraphNodeHovered(this, &SGraphNodeKnot::IsKnotHovered);
+	bHoveredCommentVisibility = false;
+
+	GetOrAddSlot(ENodeZone::TopCenter)
+	.SlotOffset(TAttribute<FVector2D>(this, &SGraphNodeKnot::GetCommentOffset))
+	.SlotSize( TAttribute<FVector2D>( CommentBubble.Get(), &SCommentBubble::GetSize))
+	.AllowScaling( TAttribute<bool>(CommentBubble.Get(), &SCommentBubble::IsScalingAllowed))
+	.VAlign(VAlign_Top)
 	[
 		CommentBubble.ToSharedRef()
 	];
@@ -387,34 +401,44 @@ void SGraphNodeKnot::AddPin(const TSharedRef<SGraphPin>& PinToAdd)
 
 FVector2D SGraphNodeKnot::GetCommentOffset() const
 {
-	static const FVector2D CenterAdjust( 8.f, 0.f );
-	return CommentBubble->GetOffset() + CenterAdjust;
+	const bool bBubbleVisible = GraphNode->bCommentBubbleVisible || bHoveredCommentVisibility;
+	const float ZoomAmount = GraphNode->bCommentBubblePinned && OwnerGraphPanelPtr.IsValid() ? OwnerGraphPanelPtr.Pin()->GetZoomAmount() : 1.f;
+	const float NodeWidthOffset = bBubbleVisible ?	SKnotNodeDefinitions::KnotCenterBubbleAdjust * ZoomAmount :
+													SKnotNodeDefinitions::KnotCenterButtonAdjust * ZoomAmount;
+
+	return FVector2D(NodeWidthOffset - CommentBubble->GetArrowCenterOffset(), -CommentBubble->GetDesiredSize().Y);
 }
 
 void SGraphNodeKnot::OnMouseEnter( const FGeometry& MyGeometry, const FPointerEvent& MouseEvent )
 {
-	if( !bCommentPinnedState && !GraphNode->NodeComment.IsEmpty() )
+	SGraphNode::OnMouseEnter(MyGeometry, MouseEvent);
+	if (!GraphNode->bCommentBubbleVisible && !GraphNode->NodeComment.IsEmpty())
 	{
-		GraphNode->bCommentBubbleVisible = true;
-		CommentBubble->UpdateBubble();
+		bHoveredCommentVisibility = true;
+		// Create the bubble widget while hovered
+		CommentBubble->OnCommentBubbleToggle(ECheckBoxState::Checked);
 	}
 }
 
-void SGraphNodeKnot::OnMouseLeave( const FPointerEvent& MouseEvent )
+void SGraphNodeKnot::OnMouseLeave(const FPointerEvent& MouseEvent)
 {
-	if( !bCommentPinnedState && !GraphNode->NodeComment.IsEmpty() )
+	SGraphNode::OnMouseLeave(MouseEvent);
+	if (bHoveredCommentVisibility)
 	{
-		GraphNode->bCommentBubbleVisible = false;
-		CommentBubble->UpdateBubble();
+		bHoveredCommentVisibility = false;
+		// Destroy the comment is visibility hasn't changed
+		CommentBubble->OnCommentBubbleToggle(ECheckBoxState::Unchecked);
 	}
 }
 
-void SGraphNodeKnot::OnCommentBubbleToggle( ECheckBoxState State )
+bool SGraphNodeKnot::IsKnotHovered() const
 {
-	if( CommentBubble.IsValid() )
-	{
-		bCommentPinnedState = !bCommentPinnedState;
-		GraphNode->bCommentBubbleVisible = bCommentPinnedState;
-		CommentBubble->UpdateBubble();
-	}
-}
\ No newline at end of file
+	const bool bIsVisible = bHoveredCommentVisibility || GraphNode->bCommentBubbleVisible;
+	return bIsVisible ? false : IsHovered();
+}
+
+ECheckBoxState SGraphNodeKnot::GetBubbleCheckState() const
+{
+	const bool bIsChecked = GraphNode->bCommentBubbleVisible && !bHoveredCommentVisibility;
+	return bIsChecked ? ECheckBoxState::Checked : ECheckBoxState::Unchecked;
+}
diff --git a/Engine/Source/Editor/GraphEditor/Private/SGraphNodeKnot.h b/Engine/Source/Editor/GraphEditor/Private/SGraphNodeKnot.h
--- a/Engine/Source/Editor/GraphEditor/Private/SGraphNodeKnot.h
+++ b/Engine/Source/Editor/GraphEditor/Private/SGraphNodeKnot.h
@@ -28,16 +28,21 @@ public:
 	virtual void OnMouseLeave( const FPointerEvent& MouseEvent ) override;
 	// End of SWidget interface
 
+protected:
+
+	/** Returns a custom hover state for the comment bubble */
+	bool IsKnotHovered() const;
+
 	/** Returns Offset to center comment on the node's only pin */
 	FVector2D GetCommentOffset() const;
 
-	/** Called when the comment bubble visibility/pinned state is changed */
-	void OnCommentBubbleToggle( ECheckBoxState State );
+	/** Called to override the comment bubble's current visibility */
+	ECheckBoxState GetBubbleCheckState() const;
 
 protected:
 
-	/** The pinned state for comment bubble */
-	bool bCommentPinnedState;
+	/** The hovered visibility state */
+	bool bHoveredCommentVisibility;
 	/** SharedPtr to comment bubble */
 	TSharedPtr<SCommentBubble> CommentBubble;
 
diff --git a/Engine/Source/Editor/GraphEditor/Public/SCommentBubble.h b/Engine/Source/Editor/GraphEditor/Public/SCommentBubble.h
--- a/Engine/Source/Editor/GraphEditor/Public/SCommentBubble.h
+++ b/Engine/Source/Editor/GraphEditor/Public/SCommentBubble.h
@@ -19,18 +19,15 @@ class GRAPHEDITOR_API SCommentBubble : public SCompoundWidget
 		/** the GraphNode this bubble should interact with */
 		SLATE_ARGUMENT( UEdGraphNode*, GraphNode )
 
-		/** Called when the comment bubble toggle button is pressed */
-		SLATE_EVENT( FOnCheckStateChanged, OnCommentBubbleToggle )
-
-		/** Called when the comment bubble pin button is pressed */
-		SLATE_EVENT( FOnCheckStateChanged, OnCommentBubblePinned )
-
-		/** Called when the comment bubble toggle button is pressed */
+		/** Attribute to determine the visibility buttons check state */
 		SLATE_ATTRIBUTE( ECheckBoxState, ToggleButtonCheck )
 
 		/** The comment text for the bubble */
 		SLATE_ATTRIBUTE( FString, Text )
 
+		/** Called when the comment text is committed */
+		SLATE_EVENT( FOnTextCommitted, OnTextCommitted )
+
 		/** The comment hint text for the bubble */
 		SLATE_ATTRIBUTE( FText, HintText )
 
@@ -73,6 +70,9 @@ class GRAPHEDITOR_API SCommentBubble : public SCompoundWidget
 	/** Returns the offset from the SNode center slot */
 	FVector2D GetOffset() const;
 
+	/** Returns the offset to the arrow center accounting for zoom on either the comment bubble or the title bar button based on current state */
+	float GetArrowCenterOffset() const;
+
 	/** Returns the bubble size */
 	FVector2D GetSize() const;
 
@@ -93,9 +93,6 @@ class GRAPHEDITOR_API SCommentBubble : public SCompoundWidget
 
 protected:
 
-	/** Returns the current comment text, converting from FString into FText */
-	FText GetCommentText() const;
-
 	/** Returns the current scale button tooltip */
 	FText GetScaleButtonTooltip() const;
 
@@ -150,10 +147,8 @@ protected:
 	/** Toggle button checked state  */
 	TAttribute<ECheckBoxState> ToggleButtonCheck;
 
-	/** Optional delegate to call when the comment is toggled */
-	FOnCheckStateChanged CommentBubbleToggleDelegate;
-	/** Optional delegate to call when the comment is toggled */
-	FOnCheckStateChanged CommentBubblePinnedDelegate;
+	/** Optional delegate to call when the comment text is committed */
+	FOnTextCommitted OnTextCommittedDelegate;
 	/** Delegate to determine if the graph node is currently hovered */
 	FIsGraphNodeHovered IsGraphNodeHovered;
 
@@ -170,7 +165,8 @@ protected:
 	/** Used to Control hover fade up/down for widgets */
 	float OpacityValue;
 	/** Cached comment */
-	mutable FString CachedComment;
-	/** Cached FText version of the comment */
-	mutable FText CachedCommentText;
+	FString CachedComment;
+	/** Cached FText comment */
+	FText CachedCommentText;
+
 };
diff --git a/Engine/Source/Editor/GraphEditor/Public/SGraphNode.h b/Engine/Source/Editor/GraphEditor/Public/SGraphNode.h
--- a/Engine/Source/Editor/GraphEditor/Public/SGraphNode.h
+++ b/Engine/Source/Editor/GraphEditor/Public/SGraphNode.h
@@ -177,6 +177,8 @@ public:
 	void SetTextCommittedEvent(FOnNodeTextCommitted InDelegate);
 	/** Set event when the user generates a warning tooltip because a connection was invalid */
 	void SetDisallowedPinConnectionEvent(SGraphEditor::FOnDisallowedPinConnection InOnDisallowedPinConnection);
+	/** called to replace this nodes comment text */
+	void OnCommentTextCommitted(const FText& NewComment, ETextCommit::Type CommitInfo);
 	/** returns true if a rename is pending on this node */
 	bool IsRenamePending() const { return bRenameIsPending; }
 
diff --git a/Engine/Source/Editor/GraphEditor/Public/SGraphNodeComment.h b/Engine/Source/Editor/GraphEditor/Public/SGraphNodeComment.h
--- a/Engine/Source/Editor/GraphEditor/Public/SGraphNodeComment.h
+++ b/Engine/Source/Editor/GraphEditor/Public/SGraphNodeComment.h
@@ -22,6 +22,7 @@ public:
 	virtual bool ShouldAllowCulling() const override { return false; }
 	virtual int32 GetSortDepth() const override;
 	virtual void EndUserInteraction() const override;
+	virtual FString GetNodeComment() const override;
 	// End SNodePanel::SNode interface
 
 	// Begin SPanel Interface
diff --git a/Engine/Source/Runtime/Engine/Private/EdGraph/EdGraphNode.cpp b/Engine/Source/Runtime/Engine/Private/EdGraph/EdGraphNode.cpp
--- a/Engine/Source/Runtime/Engine/Private/EdGraph/EdGraphNode.cpp
+++ b/Engine/Source/Runtime/Engine/Private/EdGraph/EdGraphNode.cpp
@@ -367,9 +367,12 @@ void UEdGraphNode::AddSearchMetaDataInfo(TArray<struct FSearchTagDataPair>& OutT
 
 void UEdGraphNode::OnUpdateCommentText( const FString& NewComment )
 {
-	const FScopedTransaction Transaction( LOCTEXT( "CommentCommitted", "Comment Changed" ) );
-	Modify();
-	NodeComment	= NewComment;
+	if( !NodeComment.Equals( NewComment ))
+	{
+		const FScopedTransaction Transaction( LOCTEXT( "CommentCommitted", "Comment Changed" ) );
+		Modify();
+		NodeComment	= NewComment;
+	}
 }
 
 FText UEdGraphNode::GetKeywords() const
