This changes the dark background on the documentation nodes to something lighter so the text can be read more easily.
#UE-14516 - Documentation Node Text Not Legible
#Branch UE4
#Proj GraphEditor

diff --git a/Engine/Source/Editor/GraphEditor/Private/SGraphNodeDocumentation.cpp b/Engine/Source/Editor/GraphEditor/Private/SGraphNodeDocumentation.cpp
--- a/Engine/Source/Editor/GraphEditor/Private/SGraphNodeDocumentation.cpp
+++ b/Engine/Source/Editor/GraphEditor/Private/SGraphNodeDocumentation.cpp
@@ -7,6 +7,7 @@
 #include "SLevelOfDetailBranchNode.h"
 #include "TutorialMetaData.h"
 #include "SInlineEditableTextBlock.h"
+#include "SSimpleGradient.h"
 
 #define LOCTEXT_NAMESPACE "SGraphNodeDocumentation"
 
@@ -32,6 +33,10 @@ namespace GraphNodeDocumentationDefs
 
 	/** Line wrap adjustment from node width, to account for scroll bar */
 	static const float LineWrapAdjustment = 20.f;
+
+	/** Documentation page gradient colors */
+	static const FLinearColor PageGradientStartColor( 0.85f, 0.85f, 0.85f, 1.f );
+	static const FLinearColor PageGradientEndColor( 0.75f, 0.75f, 0.75f, 1.f );
 }
 
 void SGraphNodeDocumentation::Construct( const FArguments& InArgs, UEdGraphNode* InNode )
@@ -200,7 +205,17 @@ TSharedPtr<SWidget> SGraphNodeDocumentation::CreateDocumentationPage()
 						SNew( SScrollBox )
 						+SScrollBox::Slot()
 						[
-							DesiredExcerpt.Content.ToSharedRef()
+							SNew(SOverlay)
+							+SOverlay::Slot()
+							[
+								SNew(SSimpleGradient)
+								.StartColor(GraphNodeDocumentationDefs::PageGradientStartColor) 
+								.EndColor(GraphNodeDocumentationDefs::PageGradientEndColor) 
+							]
+							+SOverlay::Slot()
+							[
+								DesiredExcerpt.Content.ToSharedRef()
+							]
 						]
 					]
 				]
