Graph Editor: Removed dead method ComputeSensibleGraphBounds from SNodePanel

diff --git a/Engine/Source/Editor/GraphEditor/Private/SNodePanel.cpp b/Engine/Source/Editor/GraphEditor/Private/SNodePanel.cpp
--- a/Engine/Source/Editor/GraphEditor/Private/SNodePanel.cpp
+++ b/Engine/Source/Editor/GraphEditor/Private/SNodePanel.cpp
@@ -1040,32 +1040,6 @@ void SNodePanel::ApplyMarqueeSelection( const FMarqueeOperation& InMarquee, cons
 	}
 }
 
-FSlateRect SNodePanel::ComputeSensibleGraphBounds() const
-{
-	float Left = 0.0f;
-	float Top = 0.0f;
-	float Right = 0.0f;
-	float Bottom = 0.0f;
-
-	// Find the bounds of the node positions
-	for (int32 ChildIndex = 0; ChildIndex < Children.Num(); ++ChildIndex)
-	{
-		const TSharedRef<SNode>& SomeChild = Children[ChildIndex];
-
-		FVector2D ChildPos = SomeChild->GetPosition();
-
-		Left = FMath::Min( Left, ChildPos.X );
-		Right = FMath::Max( Right, ChildPos.X );
-		Top = FMath::Min( Top, ChildPos.Y );
-		Bottom = FMath::Max( Bottom, ChildPos.Y );
-	}
-
-	// Pad it out in every direction, to roughly account for nodes being of non-zero extent
-	const float Padding = 100.0f;
-
-	return FSlateRect( Left - Padding, Top - Padding, Right + Padding, Bottom + Padding );
-}
-
 void SNodePanel::SelectAndCenterObject(const UObject* ObjectToSelect, bool bCenter)
 {
 	DeferredSelectionTargetObjects.Empty();
diff --git a/Engine/Source/Editor/GraphEditor/Public/SNodePanel.h b/Engine/Source/Editor/GraphEditor/Public/SNodePanel.h
--- a/Engine/Source/Editor/GraphEditor/Public/SNodePanel.h
+++ b/Engine/Source/Editor/GraphEditor/Public/SNodePanel.h
@@ -694,9 +694,6 @@ protected:
 	 */
 	static void ApplyMarqueeSelection(const FMarqueeOperation& InMarquee, const FGraphPanelSelectionSet& CurrentSelection, FGraphPanelSelectionSet& OutNewSelection);
 
-	/** @return a bounding rectangle around all the node locations; the graph bounds are padded out for the user's convenience */
-	FSlateRect ComputeSensibleGraphBounds() const;
-
 	/**
 	 * On the next tick, centers and selects the widget associated with the object if it exists
 	 *
