UE-14522 - Unlit materials no longer make use of the normal input thus making it impossible to use modified normals to modify refraction on unlit materials

The material side of this was actually working. When the refraction input is connected, the normal pin is active. The issue was just visual.
Material result node pin widgets were caching their colour on construction. Changed to recalculate each frame.

diff --git a/Engine/Source/Editor/GraphEditor/Private/MaterialPins/SGraphPinMaterialInput.cpp b/Engine/Source/Editor/GraphEditor/Private/MaterialPins/SGraphPinMaterialInput.cpp
--- a/Engine/Source/Editor/GraphEditor/Private/MaterialPins/SGraphPinMaterialInput.cpp
+++ b/Engine/Source/Editor/GraphEditor/Private/MaterialPins/SGraphPinMaterialInput.cpp
@@ -7,22 +7,21 @@
 
 void SGraphPinMaterialInput::Construct(const FArguments& InArgs, UEdGraphPin* InGraphPinObj)
 {
-	UMaterialGraph* MaterialGraph = CastChecked<UMaterialGraph>(InGraphPinObj->GetOwningNode()->GetGraph());
+	SGraphPin::Construct(SGraphPin::FArguments().UsePinColorForText(true), InGraphPinObj);
+}
+
+FSlateColor SGraphPinMaterialInput::GetPinColor() const
+{
+	check(GraphPinObj);
+	UMaterialGraph* MaterialGraph = CastChecked<UMaterialGraph>(GraphPinObj->GetOwningNode()->GetGraph());
 	const UMaterialGraphSchema* Schema = CastChecked<UMaterialGraphSchema>(MaterialGraph->GetSchema());
 
-	if (MaterialGraph->IsInputActive(InGraphPinObj))
+	if (MaterialGraph->IsInputActive(GraphPinObj))
 	{
-		CachedPinColor = Schema->ActivePinColor;
+		return Schema->ActivePinColor;
 	}
 	else
 	{
-		CachedPinColor = Schema->InactivePinColor;
+		return Schema->InactivePinColor;
 	}
-
-	SGraphPin::Construct(SGraphPin::FArguments().UsePinColorForText(true), InGraphPinObj);
-}
-
-FSlateColor SGraphPinMaterialInput::GetPinColor() const
-{
-	return CachedPinColor;
 }
diff --git a/Engine/Source/Editor/GraphEditor/Private/MaterialPins/SGraphPinMaterialInput.h b/Engine/Source/Editor/GraphEditor/Private/MaterialPins/SGraphPinMaterialInput.h
--- a/Engine/Source/Editor/GraphEditor/Private/MaterialPins/SGraphPinMaterialInput.h
+++ b/Engine/Source/Editor/GraphEditor/Private/MaterialPins/SGraphPinMaterialInput.h
@@ -15,7 +15,4 @@ protected:
 	virtual FSlateColor GetPinColor() const override;
 	// End SGraphPin interface
 
-private:
-	/** Cached Pin Color */
-	FSlateColor CachedPinColor;
 };
