Fixed not being able to rename graph nodes on mac
This addresses UE-9950.

diff --git a/Engine/Source/Editor/GraphEditor/Private/SGraphPanel.cpp b/Engine/Source/Editor/GraphEditor/Private/SGraphPanel.cpp
--- a/Engine/Source/Editor/GraphEditor/Private/SGraphPanel.cpp
+++ b/Engine/Source/Editor/GraphEditor/Private/SGraphPanel.cpp
@@ -205,7 +205,12 @@ int32 SGraphPanel::OnPaint( const FPaintArgs& Args, const FGeometry& AllottedGeo
 				// Handle Node renaming once the node is visible
 				if( bSelected && ChildNode->IsRenamePending() )
 				{
-					ChildNode->ApplyRename();
+					// Only open a rename when the window has focus
+					TSharedPtr<SWindow> OwnerWindow = FSlateApplication::Get().FindWidgetWindow(AsShared());
+					if (!OwnerWindow.IsValid() || FSlateApplication::Get().HasFocusedDescendants(OwnerWindow.ToSharedRef()))
+					{
+						ChildNode->ApplyRename();
+					}
 				}
 
 				// Draw the node's shadow.
