Changing the comment bubble newline modifier key to shift in order match other examples in the editor.
#Branch UE4
#Proj GraphEditor

diff --git a/Engine/Source/Editor/GraphEditor/Private/SCommentBubble.cpp b/Engine/Source/Editor/GraphEditor/Private/SCommentBubble.cpp
--- a/Engine/Source/Editor/GraphEditor/Private/SCommentBubble.cpp
+++ b/Engine/Source/Editor/GraphEditor/Private/SCommentBubble.cpp
@@ -241,7 +241,7 @@ void SCommentBubble::UpdateBubble()
 							.SelectAllTextWhenFocused( true )
 							.RevertTextOnEscape( true )
 							.ClearKeyboardFocusOnCommit( true )
-							.ModiferKeyForNewLine( EModifierKey::Control )
+							.ModiferKeyForNewLine( EModifierKey::Shift )
 							.ForegroundColor( this, &SCommentBubble::GetTextForegroundColor )
 							.BackgroundColor( this, &SCommentBubble::GetTextBackgroundColor )
 							.OnTextCommitted( this, &SCommentBubble::OnCommentTextCommitted )
