Improved diffing of behavior tree graphs

UETOOL-292 - Diffing of Behavior Trees doesn't work correctly

- Added support for diffing decorators and services.
- Ensured that you could see the properties of services and decorators in the diff tool.
    - This is now using common code shared with the behavior tree editor (to avoid this getting out of sync again).
- Fixed an issue where nodes containing changes could still be drawn faded out.
    - The "DiffR_" functions in GraphDiffControl weren't marking the node as changed since the diff results couldn't hold the information.
- Fixed SNodePanel::DeferredMovementTargetObject always moving the selection into view (rather than itself).
- Added a way to move to a node without also selecting it (SNodePanel::CenterObject).
    - This is needed as decorators and services don't exist as top-level nodes in the SNodePanel, so can't be centered on. We instead have to center on their container node.
- Added a way to control what items you want to diff.
    - This helps when checking for nodes that have been removed, as you can now avoid having to re-check all the other properties on the node, and just check for their existence (or lack thereof).

diff --git a/Engine/Source/Editor/AIGraph/Private/AIGraphNode.cpp b/Engine/Source/Editor/AIGraph/Private/AIGraphNode.cpp
--- a/Engine/Source/Editor/AIGraph/Private/AIGraphNode.cpp
+++ b/Engine/Source/Editor/AIGraph/Private/AIGraphNode.cpp
@@ -251,17 +251,20 @@ void UAIGraphNode::DiffProperties(UStruct* Struct, void* DataA, void* DataB, FDi
 
 		if (ValueStringA != ValueStringB)
 		{
-			if (Results)
+			// Only bother setting up the display data if we're storing the result
+			if(Results.CanStoreResults())
 			{
 				Diff.DisplayString = FText::Format(LOCTEXT("DIF_NodePropertyFmt", "Property Changed: {0} "), FText::FromString(Prop->GetName()));
-				Results.Add(Diff);
 			}
+			Results.Add(Diff);
 		}
 	}
 }
 
 void UAIGraphNode::FindDiffs(UEdGraphNode* OtherNode, FDiffResults& Results)
 {
+	Super::FindDiffs(OtherNode, Results);
+
 	FDiffSingleResult Diff;
 	Diff.Diff = EDiffType::NODE_PROPERTY;
 	Diff.Node1 = this;
diff --git a/Engine/Source/Editor/BehaviorTreeEditor/Classes/BehaviorTreeGraphNode.h b/Engine/Source/Editor/BehaviorTreeEditor/Classes/BehaviorTreeGraphNode.h
--- a/Engine/Source/Editor/BehaviorTreeEditor/Classes/BehaviorTreeGraphNode.h
+++ b/Engine/Source/Editor/BehaviorTreeEditor/Classes/BehaviorTreeGraphNode.h
@@ -24,6 +24,7 @@ class UBehaviorTreeGraphNode : public UAIGraphNode
 	virtual void AllocateDefaultPins() override;
 	virtual FText GetTooltipText() const override;
 	virtual bool CanCreateUnderSpecifiedSchema(const UEdGraphSchema* DesiredSchema) const override;
+	virtual void FindDiffs(class UEdGraphNode* OtherNode, struct FDiffResults& Results) override;
 	// End UEdGraphNode Interface
 
 	virtual FText GetDescription() const override;
diff --git a/Engine/Source/Editor/BehaviorTreeEditor/Private/BehaviorTreeEditor.cpp b/Engine/Source/Editor/BehaviorTreeEditor/Private/BehaviorTreeEditor.cpp
--- a/Engine/Source/Editor/BehaviorTreeEditor/Private/BehaviorTreeEditor.cpp
+++ b/Engine/Source/Editor/BehaviorTreeEditor/Private/BehaviorTreeEditor.cpp
@@ -23,6 +23,7 @@
 #include "BehaviorTreeEditorTabFactories.h"
 #include "BehaviorTreeEditorCommands.h"
 #include "BehaviorTreeEditorTabs.h"
+#include "BehaviorTreeEditorUtils.h"
 #include "BehaviorTree/BehaviorTree.h"
 #include "BehaviorTree/BlackboardData.h"
 #include "DetailCustomizations/BlackboardDataDetails.h"
@@ -678,72 +679,17 @@ TSharedRef<SWidget> FBehaviorTreeEditor::SpawnBlackboardDetails()
 
 void FBehaviorTreeEditor::OnSelectedNodesChanged(const TSet<class UObject*>& NewSelection)
 {
-	TArray<UObject*> Selection;
-	
-	UBehaviorTreeGraphNode_CompositeDecorator* FoundGraphNode_CompDecorator = NULL;
-	UBTDecorator* FoundDecorator = NULL;
-	bool bInjectedNode = false;
-
 	SelectedNodesCount = NewSelection.Num();
-	if(NewSelection.Num())
-	{
-		for(TSet<class UObject*>::TConstIterator SetIt(NewSelection);SetIt;++SetIt)
-		{
-			UBehaviorTreeGraphNode_Composite* GraphNode_Composite = Cast<UBehaviorTreeGraphNode_Composite>(*SetIt);
-			if (GraphNode_Composite)
-			{
-				Selection.Add(GraphNode_Composite->NodeInstance);
-				continue;
-			}
-
-			UBehaviorTreeGraphNode_Task* GraphNode_Task = Cast<UBehaviorTreeGraphNode_Task>(*SetIt);
-			if (GraphNode_Task)
-			{
-				Selection.Add(GraphNode_Task->NodeInstance);
-				continue;
-			}
-
-			UBehaviorTreeGraphNode_Decorator* GraphNode_Decorator1 = Cast<UBehaviorTreeGraphNode_Decorator>(*SetIt);
-			if (GraphNode_Decorator1)
-			{
-				Selection.Add(GraphNode_Decorator1->NodeInstance);
-				FoundDecorator = Cast<UBTDecorator>(GraphNode_Decorator1->NodeInstance);
-				bInjectedNode = bInjectedNode || GraphNode_Decorator1->bInjectedNode;
-				continue;
-			}
 
-			UBehaviorTreeDecoratorGraphNode_Decorator* GraphNode_Decorator2 = Cast<UBehaviorTreeDecoratorGraphNode_Decorator>(*SetIt);
-			if (GraphNode_Decorator2)
-			{
-				Selection.Add(GraphNode_Decorator2->NodeInstance);
-				bInjectedNode = bInjectedNode || !GraphNode_Decorator2->GetGraph()->bEditable;
-				continue;
-			}
-			
-			UBehaviorTreeGraphNode_Service* GraphNode_Service = Cast<UBehaviorTreeGraphNode_Service>(*SetIt);
-			if (GraphNode_Service)
-			{
-				Selection.Add(GraphNode_Service->NodeInstance);
-				continue;
-			}
-
-			UBehaviorTreeGraphNode_CompositeDecorator* GraphNode_CompDecorator = Cast<UBehaviorTreeGraphNode_CompositeDecorator>(*SetIt);
-			if (GraphNode_CompDecorator)
-			{
-				FoundGraphNode_CompDecorator = GraphNode_CompDecorator;
-				bInjectedNode = bInjectedNode || GraphNode_CompDecorator->bInjectedNode;
-			}
-
-			Selection.Add(*SetIt);
-		}
-	}
+	BehaviorTreeEditorUtils::FPropertySelectionInfo SelectionInfo;
+	TArray<UObject*> Selection = BehaviorTreeEditorUtils::GetSelectionForPropertyEditor(NewSelection, SelectionInfo);
 
 	UBehaviorTreeGraph* MyGraph = Cast<UBehaviorTreeGraph>(BehaviorTree->BTGraph);
 	FAbortDrawHelper Mode0, Mode1;
 	bShowDecoratorRangeLower = false;
 	bShowDecoratorRangeSelf = false;
-	bForceDisablePropertyEdit = bInjectedNode;
-	bSelectedNodeIsInjected = bInjectedNode;
+	bForceDisablePropertyEdit = SelectionInfo.bInjectedNode;
+	bSelectedNodeIsInjected = SelectionInfo.bInjectedNode;
 
 	if (Selection.Num() == 1)
 	{
@@ -752,13 +698,13 @@ void FBehaviorTreeEditor::OnSelectedNodesChanged(const TSet<class UObject*>& New
 			DetailsView->SetObjects(Selection);
 		}
 
-		if (FoundDecorator)
+		if (SelectionInfo.FoundDecorator)
 		{
-			GetAbortModePreview(FoundDecorator, Mode0, Mode1);
+			GetAbortModePreview(SelectionInfo.FoundDecorator, Mode0, Mode1);
 		}
-		else if (FoundGraphNode_CompDecorator)
+		else if (SelectionInfo.FoundGraphNode_CompDecorator)
 		{
-			GetAbortModePreview(FoundGraphNode_CompDecorator, Mode0, Mode1);
+			GetAbortModePreview(SelectionInfo.FoundGraphNode_CompDecorator, Mode0, Mode1);
 		}
 	}
 	else if (DetailsView.IsValid())
diff --git a/Engine/Source/Editor/BehaviorTreeEditor/Private/BehaviorTreeEditorUtils.cpp b/Engine/Source/Editor/BehaviorTreeEditor/Private/BehaviorTreeEditorUtils.cpp
new file mode 100644
--- /dev/null
+++ b/Engine/Source/Editor/BehaviorTreeEditor/Private/BehaviorTreeEditorUtils.cpp
@@ -0,0 +1,61 @@
+// Copyright 1998-2015 Epic Games, Inc. All Rights Reserved.
+
+#include "BehaviorTreeEditorPrivatePCH.h"
+#include "BehaviorTreeEditorUtils.h"
+
+TArray<UObject*> BehaviorTreeEditorUtils::GetSelectionForPropertyEditor(const TSet<UObject*>& InSelection, FPropertySelectionInfo& OutSelectionInfo)
+{
+	TArray<UObject*> Selection;
+
+	for (UObject* SelectionEntry : InSelection)
+	{
+		UBehaviorTreeGraphNode_Composite* GraphNode_Composite = Cast<UBehaviorTreeGraphNode_Composite>(SelectionEntry);
+		if (GraphNode_Composite)
+		{
+			Selection.Add(GraphNode_Composite->NodeInstance);
+			continue;
+		}
+
+		UBehaviorTreeGraphNode_Task* GraphNode_Task = Cast<UBehaviorTreeGraphNode_Task>(SelectionEntry);
+		if (GraphNode_Task)
+		{
+			Selection.Add(GraphNode_Task->NodeInstance);
+			continue;
+		}
+
+		UBehaviorTreeGraphNode_Decorator* GraphNode_Decorator1 = Cast<UBehaviorTreeGraphNode_Decorator>(SelectionEntry);
+		if (GraphNode_Decorator1)
+		{
+			Selection.Add(GraphNode_Decorator1->NodeInstance);
+			OutSelectionInfo.FoundDecorator = Cast<UBTDecorator>(GraphNode_Decorator1->NodeInstance);
+			OutSelectionInfo.bInjectedNode = OutSelectionInfo.bInjectedNode || GraphNode_Decorator1->bInjectedNode;
+			continue;
+		}
+
+		UBehaviorTreeDecoratorGraphNode_Decorator* GraphNode_Decorator2 = Cast<UBehaviorTreeDecoratorGraphNode_Decorator>(SelectionEntry);
+		if (GraphNode_Decorator2)
+		{
+			Selection.Add(GraphNode_Decorator2->NodeInstance);
+			OutSelectionInfo.bInjectedNode = OutSelectionInfo.bInjectedNode || !GraphNode_Decorator2->GetGraph()->bEditable;
+			continue;
+		}
+			
+		UBehaviorTreeGraphNode_Service* GraphNode_Service = Cast<UBehaviorTreeGraphNode_Service>(SelectionEntry);
+		if (GraphNode_Service)
+		{
+			Selection.Add(GraphNode_Service->NodeInstance);
+			continue;
+		}
+
+		UBehaviorTreeGraphNode_CompositeDecorator* GraphNode_CompDecorator = Cast<UBehaviorTreeGraphNode_CompositeDecorator>(SelectionEntry);
+		if (GraphNode_CompDecorator)
+		{
+			OutSelectionInfo.FoundGraphNode_CompDecorator = GraphNode_CompDecorator;
+			OutSelectionInfo.bInjectedNode = OutSelectionInfo.bInjectedNode || GraphNode_CompDecorator->bInjectedNode;
+		}
+
+		Selection.Add(SelectionEntry);
+	}
+
+	return Selection;
+}
diff --git a/Engine/Source/Editor/BehaviorTreeEditor/Private/BehaviorTreeEditorUtils.h b/Engine/Source/Editor/BehaviorTreeEditor/Private/BehaviorTreeEditorUtils.h
new file mode 100644
--- /dev/null
+++ b/Engine/Source/Editor/BehaviorTreeEditor/Private/BehaviorTreeEditorUtils.h
@@ -0,0 +1,23 @@
+// Copyright 1998-2015 Epic Games, Inc. All Rights Reserved.
+
+#pragma once
+
+namespace BehaviorTreeEditorUtils
+{
+	struct FPropertySelectionInfo
+	{
+		FPropertySelectionInfo()
+			: FoundGraphNode_CompDecorator(nullptr)
+			, FoundDecorator(nullptr)
+			, bInjectedNode(false)
+		{
+		}
+
+		UBehaviorTreeGraphNode_CompositeDecorator* FoundGraphNode_CompDecorator;
+		UBTDecorator* FoundDecorator;
+		bool bInjectedNode;
+	};
+
+	/** Given a selection of nodes, return the instances that should be selected be selected for editing in the property panel */
+	TArray<UObject*> GetSelectionForPropertyEditor(const TSet<UObject*>& InSelection, FPropertySelectionInfo& OutSelectionInfo);
+}
diff --git a/Engine/Source/Editor/BehaviorTreeEditor/Private/BehaviorTreeGraphNode.cpp b/Engine/Source/Editor/BehaviorTreeEditor/Private/BehaviorTreeGraphNode.cpp
--- a/Engine/Source/Editor/BehaviorTreeEditor/Private/BehaviorTreeGraphNode.cpp
+++ b/Engine/Source/Editor/BehaviorTreeEditor/Private/BehaviorTreeGraphNode.cpp
@@ -6,6 +6,7 @@
 #include "BlueprintNodeHelpers.h"
 #include "BehaviorTree/BehaviorTree.h"
 #include "AssetData.h"
+#include "GraphDiffControl.h"
 #include "BehaviorTree/Decorators/BTDecorator_BlueprintBase.h"
 #include "BehaviorTree/Services/BTService_BlueprintBase.h"
 #include "BehaviorTree/Tasks/BTTask_BlueprintBase.h"
@@ -108,6 +109,81 @@ bool UBehaviorTreeGraphNode::CanCreateUnderSpecifiedSchema(const UEdGraphSchema*
 	return DesiredSchema->GetClass()->IsChildOf(UEdGraphSchema_BehaviorTree::StaticClass());
 }
 
+void UBehaviorTreeGraphNode::FindDiffs(UEdGraphNode* OtherNode, FDiffResults& Results)
+{
+	Super::FindDiffs(OtherNode, Results);
+
+	UBehaviorTreeGraphNode* OtherBTGraphNode = Cast<UBehaviorTreeGraphNode>(OtherNode);
+	if (OtherBTGraphNode)
+	{
+		auto DiffSubNodes = [&Results](const FText& NodeTypeDisplayName, const TArray<UBehaviorTreeGraphNode*>& LhsSubNodes, const TArray<UBehaviorTreeGraphNode*>& RhsSubNodes)
+		{
+			TArray<FGraphDiffControl::FNodeMatch> NodeMatches;
+			TSet<const UEdGraphNode*> MatchedRhsNodes;
+
+			FGraphDiffControl::FNodeDiffContext AdditiveDiffContext;
+			AdditiveDiffContext.NodeTypeDisplayName = NodeTypeDisplayName;
+			AdditiveDiffContext.bIsRootNode = false;
+
+			// march through the all the nodes in the rhs and look for matches 
+			for (UEdGraphNode* RhsSubNode : RhsSubNodes)
+			{
+				FGraphDiffControl::FNodeMatch NodeMatch;
+				NodeMatch.NewNode = RhsSubNode;
+
+				for (UEdGraphNode* LhsSubNode : LhsSubNodes)
+				{
+					if (FGraphDiffControl::IsNodeMatch(LhsSubNode, RhsSubNode, &NodeMatches))
+					{
+						NodeMatch.OldNode = LhsSubNode;
+						break;
+					}
+				}
+
+				// if we found a corresponding node in the lhs graph, track it (so we can prevent future matches with the same nodes)
+				if (NodeMatch.IsValid())
+				{
+					NodeMatches.Add(NodeMatch);
+					MatchedRhsNodes.Add(NodeMatch.OldNode);
+				}
+
+				NodeMatch.Diff(AdditiveDiffContext, Results);
+			}
+
+			FGraphDiffControl::FNodeDiffContext SubtractiveDiffContext = AdditiveDiffContext;
+			SubtractiveDiffContext.DiffMode = FGraphDiffControl::EDiffMode::Subtractive;
+			SubtractiveDiffContext.DiffFlags = FGraphDiffControl::EDiffFlags::NodeExistance;
+
+			// go through the lhs nodes to catch ones that may have been missing from the rhs graph
+			for (UEdGraphNode* LhsSubNode : LhsSubNodes)
+			{
+				// if this node has already been matched, move on
+				if (!LhsSubNode || MatchedRhsNodes.Find(LhsSubNode))
+				{
+					continue;
+				}
+
+				FGraphDiffControl::FNodeMatch NodeMatch;
+				NodeMatch.NewNode = LhsSubNode;
+
+				for (UEdGraphNode* RhsSubNode : RhsSubNodes)
+				{
+					if (FGraphDiffControl::IsNodeMatch(LhsSubNode, RhsSubNode, &NodeMatches))
+					{
+						NodeMatch.OldNode = RhsSubNode;
+						break;
+					}
+				}
+
+				NodeMatch.Diff(SubtractiveDiffContext, Results);
+			}
+		};
+
+		DiffSubNodes(LOCTEXT("DecoratorDiffDisplayName", "Decorator"), Decorators, OtherBTGraphNode->Decorators);
+		DiffSubNodes(LOCTEXT("ServiceDiffDisplayName", "Service"), Services, OtherBTGraphNode->Services);
+	}
+}
+
 void UBehaviorTreeGraphNode::OnSubNodeAdded(UAIGraphNode* NodeTemplate)
 {
 	UBehaviorTreeGraphNode* BTGraphNode = Cast<UBehaviorTreeGraphNode>(NodeTemplate);
diff --git a/Engine/Source/Editor/BehaviorTreeEditor/Private/SBehaviorTreeDiff.cpp b/Engine/Source/Editor/BehaviorTreeEditor/Private/SBehaviorTreeDiff.cpp
--- a/Engine/Source/Editor/BehaviorTreeEditor/Private/SBehaviorTreeDiff.cpp
+++ b/Engine/Source/Editor/BehaviorTreeEditor/Private/SBehaviorTreeDiff.cpp
@@ -4,6 +4,7 @@
 #include "SBehaviorTreeDiff.h"
 #include "GraphDiffControl.h"
 #include "EdGraphUtilities.h"
+#include "BehaviorTreeEditorUtils.h"
 #include "Editor/PropertyEditor/Public/IDetailsView.h"
 #include "BehaviorTree/BehaviorTree.h"
 #include "GenericCommands.h"
@@ -322,30 +323,63 @@ void SBehaviorTreeDiff::OnSelectionChanged(FSharedDiffOnGraph Item, ESelectInfo:
 		PanelNew.GraphEditor.Pin()->ClearSelectionSet();
 		PanelOld.GraphEditor.Pin()->ClearSelectionSet();
 
-		if(Result.Pin1)
+		auto FocusPin = [this](UEdGraphPin* InPin)
 		{
-			LastPinTarget = Result.Pin1;
-			Result.Pin1->bIsDiffing = true;
-			GetGraphEditorForGraph(Result.Pin1->GetOwningNode()->GetGraph())->JumpToPin(Result.Pin1);
-		}
+			if (InPin)
+			{
+				LastPinTarget = InPin;
+				InPin->bIsDiffing = true;
 
-		if(Result.Pin2)
-		{
-			LastOtherPinTarget = Result.Pin2;
-			Result.Pin2->bIsDiffing = true;
-			GetGraphEditorForGraph(Result.Pin2->GetOwningNode()->GetGraph())->JumpToPin(Result.Pin2);
-		}
+				UEdGraph* NodeGraph = InPin->GetOwningNode()->GetGraph();
+				SGraphEditor* NodeGraphEditor = GetGraphEditorForGraph(NodeGraph);
+				NodeGraphEditor->JumpToPin(InPin);
+			}
+		};
+
+		FocusPin(Result.Pin1);
+		FocusPin(Result.Pin2);
 	}
 	else if(Result.Node1)
 	{
 		PanelNew.GraphEditor.Pin()->ClearSelectionSet();
 		PanelOld.GraphEditor.Pin()->ClearSelectionSet();
 
-		if(Result.Node2)
+		auto FocusNode = [this](UEdGraphNode* InNode)
 		{
-			GetGraphEditorForGraph(Result.Node2->GetGraph())->JumpToNode(Result.Node2, false);
-		}
-		GetGraphEditorForGraph(Result.Node1->GetGraph())->JumpToNode(Result.Node1, false);
+			if (InNode)
+			{
+				UEdGraph* NodeGraph = InNode->GetGraph();
+				SGraphEditor* NodeGraphEditor = GetGraphEditorForGraph(NodeGraph);
+
+				UBehaviorTreeGraphNode* BTNode = Cast<UBehaviorTreeGraphNode>(InNode);
+				if (BTNode && BTNode->bIsSubNode)
+				{
+					// This is a sub-node, we need to find our parent node in the graph
+					// todo: work out why BTNode->ParentNode is always null
+					UEdGraphNode** ParentNodePtr = NodeGraph->Nodes.FindByPredicate([BTNode](UEdGraphNode* PotentialParentNode) -> bool
+					{
+						UBehaviorTreeGraphNode* BTPotentialParentNode = Cast<UBehaviorTreeGraphNode>(PotentialParentNode);
+						return BTPotentialParentNode && (BTPotentialParentNode->Decorators.Contains(BTNode) || BTPotentialParentNode->Services.Contains(BTNode));
+					});
+
+					// We need to call JumpToNode on the parent node, and then SetNodeSelection on the sub-node
+					// as JumpToNode doesn't work for sub-nodes
+					if (ParentNodePtr)
+					{
+						check(InNode->GetGraph() == (*ParentNodePtr)->GetGraph());
+						NodeGraphEditor->JumpToNode(*ParentNodePtr, false, false);
+					}
+					NodeGraphEditor->SetNodeSelection(InNode, true);
+				}
+				else
+				{
+					NodeGraphEditor->JumpToNode(InNode, false);
+				}
+			}
+		};
+
+		FocusNode(Result.Node1);
+		FocusNode(Result.Node2);
 	}
 }
 
@@ -546,27 +580,8 @@ bool SBehaviorTreeDiff::FBehaviorTreeDiffPanel::CanCopyNodes() const
 
 void SBehaviorTreeDiff::FBehaviorTreeDiffPanel::OnSelectionChanged( const FGraphPanelSelectionSet& NewSelection )
 {
-	TArray<UObject*> Selection;
-	if(NewSelection.Num())
-	{
-		for(TSet<class UObject*>::TConstIterator SetIt(NewSelection);SetIt;++SetIt)
-		{
-			UBehaviorTreeGraphNode_Composite* GraphNode_Composite = Cast<UBehaviorTreeGraphNode_Composite>(*SetIt);
-			if (GraphNode_Composite)
-			{
-				Selection.Add(GraphNode_Composite->NodeInstance);
-				continue;
-			}
-
-			UBehaviorTreeGraphNode_Task* GraphNode_Task = Cast<UBehaviorTreeGraphNode_Task>(*SetIt);
-			if (GraphNode_Task)
-			{
-				Selection.Add(GraphNode_Task->NodeInstance);
-				continue;
-			}
-			Selection.Add(*SetIt);
-		}
-	}
+	BehaviorTreeEditorUtils::FPropertySelectionInfo SelectionInfo;
+	TArray<UObject*> Selection = BehaviorTreeEditorUtils::GetSelectionForPropertyEditor(NewSelection, SelectionInfo);
 
 	if (Selection.Num() == 1)
 	{
diff --git a/Engine/Source/Editor/GraphEditor/Private/GraphDiffControl.cpp b/Engine/Source/Editor/GraphEditor/Private/GraphDiffControl.cpp
--- a/Engine/Source/Editor/GraphEditor/Private/GraphDiffControl.cpp
+++ b/Engine/Source/Editor/GraphEditor/Private/GraphDiffControl.cpp
@@ -7,167 +7,158 @@
 #define LOCTEXT_NAMESPACE "GraphDiffControl"
 
 //If we are collecting all Diff results, keep going. If we just want to know if there *is* any diffs, we can early out
-#define KEEP_GOING_IF_RESULTS() 	bHasResult = true;			\
-if(!Results)					\
-{break;	}					
+#define KEEP_GOING_IF_RESULTS() 	bHasResult = true;	\
+if(!Results.CanStoreResults())		{ break; }
 
 
 /*******************************************************************************
 * Static helper functions
 *******************************************************************************/
 
-/** Determine if two Nodes are the same */
-static bool IsNodeMatch(class UEdGraphNode* Node1, class UEdGraphNode* Node2, TArray<FGraphDiffControl::FNodeMatch> const* Exclusions = NULL)
+/** Diff result when a node was added to the graph */
+static void DiffR_NodeAdded( const FGraphDiffControl::FNodeDiffContext& DiffContext, FDiffResults& Results, UEdGraphNode* Node )
 {
-	if(Node2->GetClass() != Node1->GetClass()) 
-	{
-		return false;
-	}
+	FDiffSingleResult Diff;
+	Diff.Diff = EDiffType::NODE_ADDED;
+	Diff.Node1 = Node;
 
-	if(Node1->NodeGuid == Node2->NodeGuid)
-	{
-		return true;
-	}
-
-	// we could be diffing two completely separate assets, this makes sure both 
-	// nodes historically belong to the same graph
-	bool bIsIntraAssetDiff = (Node1->GetGraph()->GraphGuid == Node2->GetGraph()->GraphGuid);
-
-	// if both nodes are from the same graph
-	if (bIsIntraAssetDiff)
-	{
-		return (Node1->GetFName() == Node2->GetFName());
-	}
-
-	if (Exclusions != NULL)
+	// Only bother setting up the display data if we're storing the result
+	if(Results.CanStoreResults())
 	{
-		// have to see if this node has already been matched with another
-		for (auto MatchIt(Exclusions->CreateConstIterator()); MatchIt; ++MatchIt)
-		{
-			FGraphDiffControl::FNodeMatch const& PriorMatch = *MatchIt;
-			if (!PriorMatch.IsValid())
-			{
-				continue;
-			}
-
-			// if one of these nodes has already been matched to a different node
-			if (((PriorMatch.OldNode == Node1) && (PriorMatch.NewNode != Node2)) ||
-				((PriorMatch.OldNode == Node2) && (PriorMatch.NewNode != Node1)) ||
-				((PriorMatch.NewNode == Node1) && (PriorMatch.OldNode != Node2)) ||
-				((PriorMatch.NewNode == Node2) && (PriorMatch.OldNode != Node1)))
-			{
-				return false;
-			}
-		}
+		FFormatNamedArguments Args;
+		Args.Add(TEXT("NodeType"), DiffContext.NodeTypeDisplayName);
+		Args.Add(TEXT("NodeTitle"), Node->GetNodeTitle(ENodeTitleType::ListView));
+		Diff.ToolTip =  FText::Format(LOCTEXT("DIF_AddNode", "Added {NodeType} '{NodeTitle}'"), Args);
+		Diff.DisplayString = Diff.ToolTip;
+		Diff.DisplayColor = FLinearColor(0.3f,1.0f,0.4f);
 	}
-
-	// the name hashes won't match for nodes from separate graph assets, so we 
-	// need to look for some kind of semblance between the two... title? (what's displayed to the user)
-	FText Title1 = Node1->GetNodeTitle(ENodeTitleType::FullTitle);
-	FText Title2 = Node2->GetNodeTitle(ENodeTitleType::FullTitle);
-
-	return Title1.CompareTo(Title2) == 0;
+	
+	Results.Add(Diff);
 }
 
-/** Diff result when a node was added to the graph */
-static void DiffR_NodeAdded( FDiffResults& Results, UEdGraphNode* Node )
+/** Diff result when a node was removed from the graph */
+static void DiffR_NodeRemoved( const FGraphDiffControl::FNodeDiffContext& DiffContext, FDiffResults& Results, UEdGraphNode* Node )
 {
-	if(Results)
-	{
-		FDiffSingleResult Diff;
-		Diff.Diff = EDiffType::NODE_ADDED;
-		Diff.Node1 = Node;
+	FDiffSingleResult Diff;
+	Diff.Diff = EDiffType::NODE_REMOVED;
+	Diff.Node1 = Node;
 
+	// Only bother setting up the display data if we're storing the result
+	if(Results.CanStoreResults())
+	{
 		FFormatNamedArguments Args;
+		Args.Add(TEXT("NodeType"), DiffContext.NodeTypeDisplayName);
 		Args.Add(TEXT("NodeTitle"), Node->GetNodeTitle(ENodeTitleType::ListView));
-		Diff.ToolTip =  FText::Format(LOCTEXT("DIF_AddNode", "Added Node '{NodeTitle}'"), Args);
+		Diff.ToolTip =  FText::Format(LOCTEXT("DIF_RemoveNode", "Removed {NodeType} '{NodeTitle}'"), Args);
 		Diff.DisplayString = Diff.ToolTip;
-		Diff.DisplayColor = FLinearColor(0.3f,1.0f,0.4f);
-		Results.Add(Diff);
+		Diff.DisplayColor = FLinearColor(1.f,0.4f,0.4f);
 	}
+
+	Results.Add(Diff);
 }
 
 /** Diff result when a node comment was changed */
-static void DiffR_NodeCommentChanged(FDiffResults& Results, class UEdGraphNode* Node, class UEdGraphNode* Node2)
+static void DiffR_NodeCommentChanged(const FGraphDiffControl::FNodeDiffContext& DiffContext, FDiffResults& Results, class UEdGraphNode* Node, class UEdGraphNode* Node2)
 {
-	if(Results)
-	{
-		FDiffSingleResult Diff;
-		Diff.Diff = EDiffType::NODE_COMMENT;
-		Diff.Node1 = Node;
-		Diff.Node2 = Node2;
+	FDiffSingleResult Diff;
+	Diff.Diff = EDiffType::NODE_COMMENT;
+	Diff.Node1 = Node;
+	Diff.Node2 = Node2;
 
+	// Only bother setting up the display data if we're storing the result
+	if(Results.CanStoreResults())
+	{
 		FFormatNamedArguments Args;
+		Args.Add(TEXT("NodeType"), DiffContext.NodeTypeDisplayName);
 		Args.Add(TEXT("NodeTitle"), Node->GetNodeTitle(ENodeTitleType::ListView));
-		Diff.ToolTip =  FText::Format(LOCTEXT("DIF_CommentModified", "Comment Modified Node '{NodeTitle}'"), Args);
-
+		Diff.ToolTip =  FText::Format(LOCTEXT("DIF_CommentModified", "Comment Modified {NodeType} '{NodeTitle}'"), Args);
 		Diff.DisplayString = Diff.ToolTip;
 		Diff.DisplayColor = FLinearColor(0.25f,0.4f,0.5f);
-		Results.Add(Diff);
 	}
+
+	Results.Add(Diff);
 }
 
 /** Diff result when a node was moved on the graph */
-static void DiffR_NodeMoved(FDiffResults& Results, UEdGraphNode* Node, class UEdGraphNode* OtherNode)
+static void DiffR_NodeMoved(const FGraphDiffControl::FNodeDiffContext& DiffContext, FDiffResults& Results, UEdGraphNode* Node, class UEdGraphNode* OtherNode)
 {
-	if(Results)
-	{
-		FDiffSingleResult Diff;
-		Diff.Diff = EDiffType::NODE_MOVED;
-		Diff.Node1 = Node;
-		Diff.Node2 = OtherNode;
+	FDiffSingleResult Diff;
+	Diff.Diff = EDiffType::NODE_MOVED;
+	Diff.Node1 = Node;
+	Diff.Node2 = OtherNode;
 
+	// Only bother setting up the display data if we're storing the result
+	if(Results.CanStoreResults())
+	{
 		FFormatNamedArguments Args;
+		Args.Add(TEXT("NodeType"), DiffContext.NodeTypeDisplayName);
 		Args.Add(TEXT("NodeTitle"), Node->GetNodeTitle(ENodeTitleType::ListView));
-		Diff.ToolTip = FText::Format(LOCTEXT("DIF_MoveNode", "Moved Node '{NodeTitle}'"), Args);
-
+		Diff.ToolTip = FText::Format(LOCTEXT("DIF_MoveNode", "Moved {NodeType} '{NodeTitle}'"), Args);
 		Diff.DisplayString = Diff.ToolTip;
 		Diff.DisplayColor = FLinearColor(0.9f, 0.84f, 0.43f);
-		Results.Add(Diff);
 	}
+
+	Results.Add(Diff);
 }
 
 /** Diff result when a pin type was changed */
 static void DiffR_PinTypeChanged(FDiffResults& Results, class UEdGraphPin* Pin2, class UEdGraphPin* Pin1)
 {
-	if(Results)
-	{
-		FEdGraphPinType Type1 = Pin1->PinType;
-		FEdGraphPinType Type2 = Pin2->PinType;
+	FEdGraphPinType Type1 = Pin1->PinType;
+	FEdGraphPinType Type2 = Pin2->PinType;
 
-		FDiffSingleResult Diff;
+	FDiffSingleResult Diff;
 
-		const UObject* T1Obj = Type1.PinSubCategoryObject.Get();
-		const UObject* T2Obj = Type2.PinSubCategoryObject.Get();
+	const UObject* T1Obj = Type1.PinSubCategoryObject.Get();
+	const UObject* T2Obj = Type2.PinSubCategoryObject.Get();
+
+	if(Type1.PinCategory != Type2.PinCategory)
+	{
+		Diff.Diff = EDiffType::PIN_TYPE_CATEGORY;
 
-		if(Type1.PinCategory != Type2.PinCategory)
+		// Only bother setting up the display data if we're storing the result
+		if(Results.CanStoreResults())
 		{
-			Diff.Diff = EDiffType::PIN_TYPE_CATEGORY;
 			Diff.ToolTip = FText::Format(LOCTEXT("DIF_PinCategoryToolTipFmt", "Pin '{0}' Category was '{1}', but is now '{2}"), FText::FromString(Pin2->PinName), FText::FromString(Pin1->PinType.PinCategory), FText::FromString(Pin2->PinType.PinCategory));
 			Diff.DisplayColor = FLinearColor(0.15f,0.53f,0.15f);
 			Diff.DisplayString = FText::Format(LOCTEXT("DIF_PinCategoryFmt", "Pin Category '{0}' ['{1}' -> '{2}']"), FText::FromString(Pin2->PinName), FText::FromString(Pin1->PinType.PinCategory), FText::FromString(Pin2->PinType.PinCategory));
 		}
-		else if(Type1.PinSubCategory != Type2.PinSubCategory)
+	}
+	else if(Type1.PinSubCategory != Type2.PinSubCategory)
+	{
+		Diff.Diff = EDiffType::PIN_TYPE_SUBCATEGORY;
+
+		// Only bother setting up the display data if we're storing the result
+		if(Results.CanStoreResults())
 		{
-			Diff.Diff = EDiffType::PIN_TYPE_SUBCATEGORY;
 			Diff.ToolTip = FText::Format(LOCTEXT("DIF_PinSubCategoryToolTipFmt", "Pin '{0}' SubCategory was '{1}', but is now '{2}"), FText::FromString(Pin2->PinName), FText::FromString(Pin1->PinType.PinSubCategory), FText::FromString(Pin2->PinType.PinSubCategory));
 			Diff.DisplayColor = FLinearColor(0.45f,0.53f,0.65f);
 			Diff.DisplayString = FText::Format(LOCTEXT("DIF_PinSubCategoryFmt", "Pin SubCategory '{0}'  ['{1}' -> '{2}']"), FText::FromString(Pin2->PinName), FText::FromString(Pin1->PinType.PinSubCategory), FText::FromString(Pin2->PinType.PinSubCategory));
 		}
-		else if(T1Obj != T2Obj && (T1Obj && T2Obj ) &&
-			(T1Obj->GetFName() != T2Obj->GetFName()))
+	}
+	else if(T1Obj != T2Obj && (T1Obj && T2Obj ) &&
+		(T1Obj->GetFName() != T2Obj->GetFName()))
+	{
+		Diff.Diff = EDiffType::PIN_TYPE_SUBCATEGORY_OBJECT;
+
+		// Only bother setting up the display data if we're storing the result
+		if(Results.CanStoreResults())
 		{
 			FString Obj1 = T1Obj->GetFName().ToString();
 			FString Obj2 = T2Obj->GetFName().ToString();
 
-			Diff.Diff = EDiffType::PIN_TYPE_SUBCATEGORY_OBJECT;
 			Diff.ToolTip = FText::Format(LOCTEXT("DIF_PinSubCategorObjToolTipFmt", "Pin '{0}' was SubCategoryObject '{1}', but is now '{2}"), FText::FromString(Pin2->PinName), FText::FromString(Obj1), FText::FromString(Obj2));
 			Diff.DisplayColor = FLinearColor(0.45f,0.13f,0.25f);
 			Diff.DisplayString = FText::Format(LOCTEXT("DIF_PinSubCategoryObjFmt", "Pin SubCategoryObject '{0}' ['{1}' -> '{2}']"), FText::FromString(Pin2->PinName), FText::FromString(Obj1), FText::FromString(Obj2));
 		}
-		else if(Type1.bIsArray != Type2.bIsArray)
+	}
+	else if(Type1.bIsArray != Type2.bIsArray)
+	{
+		Diff.Diff = EDiffType::PIN_TYPE_IS_ARRAY;
+
+		// Only bother setting up the display data if we're storing the result
+		if(Results.CanStoreResults())
 		{
-			Diff.Diff = EDiffType::PIN_TYPE_IS_ARRAY;
 			FText IsArray1 = Pin1->PinType.bIsArray ? LOCTEXT("true", "true") : LOCTEXT("false", "false");
 			FText IsArray2 = Pin2->PinType.bIsArray ? LOCTEXT("true", "true") : LOCTEXT("false", "false");
 
@@ -175,9 +166,14 @@ static void DiffR_PinTypeChanged(FDiffResults& Results, class UEdGraphPin* Pin2,
 			Diff.DisplayString = FText::Format(LOCTEXT("DIF_PinIsArrayFmt", "Pin IsArray '{0}' ['{1}' -> '{2}']"), FText::FromString(Pin2->PinName), IsArray1, IsArray2);
 			Diff.DisplayColor = FLinearColor(0.45f,0.33f,0.35f);
 		}
-		else if(Type1.bIsReference != Type2.bIsReference)
+	}
+	else if(Type1.bIsReference != Type2.bIsReference)
+	{
+		Diff.Diff = EDiffType::PIN_TYPE_IS_REF;
+
+		// Only bother setting up the display data if we're storing the result
+		if(Results.CanStoreResults())
 		{
-			Diff.Diff = EDiffType::PIN_TYPE_IS_REF;
 			FText IsRef1 = Pin1->PinType.bIsReference ? LOCTEXT("true", "true") : LOCTEXT("false", "false");
 			FText IsRef2 = Pin2->PinType.bIsReference ? LOCTEXT("true", "true") : LOCTEXT("false", "false");
 
@@ -185,22 +181,24 @@ static void DiffR_PinTypeChanged(FDiffResults& Results, class UEdGraphPin* Pin2,
 			Diff.DisplayColor = FLinearColor(0.25f,0.43f,0.35f);
 			Diff.DisplayString = FText::Format(LOCTEXT("DIF_PinIsRefFmt", "Pin IsReference '{0}' ['{1}' -> '{2}']"), FText::FromString(Pin2->PinName), IsRef1, IsRef2);
 		}
-
-		Diff.Pin1 = Pin1;
-		Diff.Pin2 = Pin2;
-		Results.Add(Diff);
 	}
+
+	Diff.Pin1 = Pin1;
+	Diff.Pin2 = Pin2;
+	Results.Add(Diff);
 }
 
 /** Diff result when the # of links to a pin was changed */
 static void DiffR_PinLinkCountChanged(FDiffResults& Results, class UEdGraphPin* Pin2, class UEdGraphPin* Pin1)
 {
-	if(Results)
+	FDiffSingleResult Diff;
+	Diff.Diff = Pin2->LinkedTo.Num() > Pin1->LinkedTo.Num()  ?  EDiffType::PIN_LINKEDTO_NUM_INC : EDiffType::PIN_LINKEDTO_NUM_DEC;
+	Diff.Pin2 = Pin2;
+	Diff.Pin1 = Pin1;
+
+	// Only bother setting up the display data if we're storing the result
+	if(Results.CanStoreResults())
 	{
-		FDiffSingleResult Diff;
-		Diff.Diff = Pin2->LinkedTo.Num() > Pin1->LinkedTo.Num()  ?  EDiffType::PIN_LINKEDTO_NUM_INC : EDiffType::PIN_LINKEDTO_NUM_DEC;
-		Diff.Pin2 = Pin2;
-		Diff.Pin1 = Pin1;
 		if(Diff.Diff == EDiffType::PIN_LINKEDTO_NUM_INC)
 		{
 			Diff.ToolTip = FText::Format(LOCTEXT("DIF_PinLinkCountIncToolTipFmt", "Pin '{0}' has more links (was {1} now {2})"), FText::FromString(Pin1->PinName), FText::AsNumber(Pin1->LinkedTo.Num()), FText::AsNumber(Pin2->LinkedTo.Num()));
@@ -213,23 +211,24 @@ static void DiffR_PinLinkCountChanged(FDiffResults& Results, class UEdGraphPin*
 			Diff.DisplayColor = FLinearColor(0.5f,0.3f,0.85f);
 			Diff.DisplayString = FText::Format(LOCTEXT("DIF_PinLinkCountDecFmt", "Removed Link to '{0}'"), FText::FromString(Pin1->PinName));
 		}
-
-		Results.Add(Diff);
 	}
+
+	Results.Add(Diff);
 }
 
 /** Diff result when a pin to relinked to a different node */
 static void DiffR_LinkedToNode(FDiffResults& Results, class UEdGraphPin* Pin1, class UEdGraphPin* Pin2,  class UEdGraphNode* Node1, class UEdGraphNode* Node2)
 {
-	if(Results)
+	FDiffSingleResult Diff;
+	Diff.Diff = EDiffType::PIN_LINKEDTO_NODE;
+	Diff.Pin2 = Pin2;
+	Diff.Pin1 = Pin1;
+	Diff.Node1 = Node1;
+	Diff.Node2 = Node2;
+
+	// Only bother setting up the display data if we're storing the result
+	if(Results.CanStoreResults())
 	{
-		FDiffSingleResult Diff;
-		Diff.Diff = EDiffType::PIN_LINKEDTO_NODE;
-		Diff.Pin2 = Pin2;
-		Diff.Pin1 = Pin1;
-		Diff.Node1 = Node1;
-		Diff.Node2 = Node2;
-
 		FText Node1Name = Node1->GetNodeTitle(ENodeTitleType::ListView);
 		FText Node2Name = Node2->GetNodeTitle(ENodeTitleType::ListView);
 
@@ -241,20 +240,22 @@ static void DiffR_LinkedToNode(FDiffResults& Results, class UEdGraphPin* Pin1, c
 
 		Diff.DisplayColor = FLinearColor(0.85f,0.71f,0.25f);
 		Diff.DisplayString = FText::Format(LOCTEXT("DIF_PinLinkMoved", "Link Moved  '{PinNameForNode1}' ['{NodeName1}' -> '{NodeName2}']"), Args);
-		Results.Add(Diff);
 	}
+
+	Results.Add(Diff);
 }
 
 /** Diff result when a pin default value was changed, and is in use*/
 static void DiffR_PinDefaultValueChanged(FDiffResults& Results, class UEdGraphPin* Pin2, class UEdGraphPin* Pin1)
 {
-	if(Results)
-	{
-		FDiffSingleResult Diff;
-		Diff.Diff = EDiffType::PIN_DEFAULT_VALUE;
-		Diff.Pin1 = Pin1;
-		Diff.Pin2 = Pin2;
+	FDiffSingleResult Diff;
+	Diff.Diff = EDiffType::PIN_DEFAULT_VALUE;
+	Diff.Pin1 = Pin1;
+	Diff.Pin2 = Pin2;
 
+	// Only bother setting up the display data if we're storing the result
+	if(Results.CanStoreResults())
+	{
 		FFormatNamedArguments Args;
 		Args.Add(TEXT("PinNameForValue1"), FText::FromString(Pin2->PinName));
 		Args.Add(TEXT("PinValue1"), FText::FromString(Pin1->GetDefaultAsString()));
@@ -262,8 +263,9 @@ static void DiffR_PinDefaultValueChanged(FDiffResults& Results, class UEdGraphPi
 		Diff.ToolTip = FText::Format(LOCTEXT("DIF_PinDefaultValueToolTip", "Pin '{PinNameForValue1}' Default Value was '{PinValue1}', but is now '{PinValue2}"), Args);
 		Diff.DisplayColor = FLinearColor(0.665f,0.13f,0.455f);
 		Diff.DisplayString = FText::Format(LOCTEXT("DIF_PinDefaultValue", "Pin Default '{PinNameForValue1}' '{PinValue1}' -> '{PinValue2}']"), Args);
-		Results.Add(Diff);
 	}
+
+	Results.Add(Diff);
 }
 
 /** 
@@ -277,16 +279,17 @@ static void DiffR_PinDefaultValueChanged(FDiffResults& Results, class UEdGraphPi
  */
 static void DiffR_NodePinCount(FDiffResults& Results, class UEdGraphNode* Node, class UEdGraphNode* Node2, const TArray<class UEdGraphPin*> &Pins1, const TArray<class UEdGraphPin*> &Pins2)
 {
-	if(Results)
+	FText NodeName = Node->GetNodeTitle(ENodeTitleType::ListView);
+	int32 OriginalCount = Pins2.Num();
+	int32 NewCount = Pins1.Num();
+	FDiffSingleResult Diff;
+	Diff.Diff = EDiffType::NODE_PIN_COUNT;
+	Diff.Node1 = Node;
+	Diff.Node2 = Node2;
+
+	// Only bother setting up the display data if we're storing the result
+	if(Results.CanStoreResults())
 	{
-		FText NodeName = Node->GetNodeTitle(ENodeTitleType::ListView);
-		int32 OriginalCount = Pins2.Num();
-		int32 NewCount = Pins1.Num();
-		FDiffSingleResult Diff;
-		Diff.Diff = EDiffType::NODE_PIN_COUNT;
-		Diff.Node1 = Node;
-		Diff.Node2 = Node2;
-
 		FFormatNamedArguments Args;
 		Args.Add(TEXT("NodeName"), NodeName);
 		Args.Add(TEXT("OriginalCount"), OriginalCount);
@@ -296,8 +299,9 @@ static void DiffR_NodePinCount(FDiffResults& Results, class UEdGraphNode* Node,
 		Diff.DisplayString = (OriginalCount < NewCount) 
 			? (FText::Format(LOCTEXT("DIF_NodePinCountIncreased", "Add Pin '{NodeName}'"), Args)) 
 			: (FText::Format(LOCTEXT("DIF_NodePinCountDecreased", "Removed Pin '{NodeName}'"), Args));
-		Results.Add(Diff);
 	}
+
+	Results.Add(Diff);
 }
 
 /** 
@@ -347,7 +351,7 @@ static class UEdGraphPin* FindOtherLink(TArray<class UEdGraphPin*>& Links2, int3
 	{
 		UEdGraphPin* Other = *It;
 		UEdGraphNode* Node2 = Other->GetOwningNode();
-		if(IsNodeMatch(Node1, Node2))
+		if(FGraphDiffControl::IsNodeMatch(Node1, Node2))
 		{
 			return Other;
 		}
@@ -367,7 +371,7 @@ static bool LinkedToDifferent(class UEdGraphPin* OriginalPin1, class UEdGraphPin
 
 		UEdGraphNode* Node1 = Pin1->GetOwningNode();
 		UEdGraphNode* Node2 = Pin2->GetOwningNode();
-		if(!IsNodeMatch(Node1, Node2))
+		if(!FGraphDiffControl::IsNodeMatch(Node1, Node2))
 		{
 			DiffR_LinkedToNode(Results, OriginalPin1, OriginalPin2, Node1, Node2);
 			KEEP_GOING_IF_RESULTS()
@@ -427,57 +431,77 @@ bool FGraphDiffControl::FNodeMatch::IsValid() const
 	return ((NewNode != NULL) && (OldNode != NULL));
 }
 
-bool FGraphDiffControl::FNodeMatch::Diff(TArray<FDiffSingleResult>* OptionalDiffsArray /* = NULL*/) const
+bool FGraphDiffControl::FNodeMatch::Diff(const FNodeDiffContext& DiffContext, TArray<FDiffSingleResult>* OptionalDiffsArray /* = NULL*/) const
 {
 	FDiffResults DiffsOut(OptionalDiffsArray);
+	return Diff(DiffContext, DiffsOut);
+}
+
+bool FGraphDiffControl::FNodeMatch::Diff(const FNodeDiffContext& DiffContext, FDiffResults& DiffsOut) const
+{
 	bool bIsDifferent = false;
 
 	if (IsValid())
 	{
 		//has comment changed?
-		if(NewNode->NodeComment != OldNode->NodeComment)
+		if((DiffContext.DiffFlags & EDiffFlags::NodeComment) && NewNode->NodeComment != OldNode->NodeComment)
 		{
-			DiffR_NodeCommentChanged(DiffsOut, NewNode, OldNode);
+			DiffR_NodeCommentChanged(DiffContext, DiffsOut, NewNode, OldNode);
 			bIsDifferent = true;
 		}
 
 		//has it moved?
-		if( (NewNode->NodePosX != OldNode->NodePosX) || (NewNode->NodePosY != OldNode->NodePosY) )
+		if( (DiffContext.DiffFlags & EDiffFlags::NodeMovement) && ( (NewNode->NodePosX != OldNode->NodePosX) || (NewNode->NodePosY != OldNode->NodePosY) ) )
 		{
 			//same node, different position--
-			DiffR_NodeMoved(DiffsOut, NewNode, OldNode);
+			DiffR_NodeMoved(DiffContext, DiffsOut, NewNode, OldNode);
 			bIsDifferent = true;
 		}
 
-		// Build arrays of pins that we care about
-		TArray< UEdGraphPin* > OldRelevantPins;
-		TArray< UEdGraphPin* > RelevantPins;
-		BuildArrayOfRelevantPins(OldNode->Pins, OldRelevantPins);
-		BuildArrayOfRelevantPins(NewNode->Pins, RelevantPins);
-
-		if(OldRelevantPins.Num() == RelevantPins.Num())
+		if(DiffContext.DiffFlags & EDiffFlags::NodePins)
 		{
-			//checks contents of pins
-			bIsDifferent |= ArePinsDifferent(OldRelevantPins, RelevantPins, DiffsOut);
-		}
-		else//# of pins changed
-		{
-			DiffR_NodePinCount(DiffsOut, NewNode, OldNode, RelevantPins, OldRelevantPins);
-			bIsDifferent = true;
+			// Build arrays of pins that we care about
+			TArray< UEdGraphPin* > OldRelevantPins;
+			TArray< UEdGraphPin* > RelevantPins;
+			BuildArrayOfRelevantPins(OldNode->Pins, OldRelevantPins);
+			BuildArrayOfRelevantPins(NewNode->Pins, RelevantPins);
+
+			if(OldRelevantPins.Num() == RelevantPins.Num())
+			{
+				//checks contents of pins
+				bIsDifferent |= ArePinsDifferent(OldRelevantPins, RelevantPins, DiffsOut);
+			}
+			else//# of pins changed
+			{
+				DiffR_NodePinCount(DiffsOut, NewNode, OldNode, RelevantPins, OldRelevantPins);
+				bIsDifferent = true;
+			}
 		}
 
 		//Find internal node diffs; skip this if we don't need the result data
-		if(!bIsDifferent || DiffsOut)
+		if((DiffContext.DiffFlags & EDiffFlags::NodeSpecificDiffs) && (!bIsDifferent || DiffsOut.CanStoreResults()))
 		{
 			OldNode->FindDiffs(NewNode, DiffsOut);
 			bIsDifferent |= DiffsOut.HasFoundDiffs();
 		}
 	}
-	else 
+	else if(DiffContext.DiffFlags & EDiffFlags::NodeExistance)
 	// one of the nodes is NULL
 	{
 		bIsDifferent = true;
-		DiffR_NodeAdded(DiffsOut, NewNode);
+		switch (DiffContext.DiffMode)
+		{
+		case EDiffMode::Additive:
+			DiffR_NodeAdded(DiffContext, DiffsOut, NewNode);
+			break;
+
+		case EDiffMode::Subtractive:
+			DiffR_NodeRemoved(DiffContext, DiffsOut, NewNode);
+			break;
+
+		default:
+			break;
+		}
 	}
 
 	return bIsDifferent;
@@ -514,6 +538,58 @@ FGraphDiffControl::FNodeMatch FGraphDiffControl::FindNodeMatch(class UEdGraph* G
 	return Match;
 }
 
+bool FGraphDiffControl::IsNodeMatch(class UEdGraphNode* Node1, class UEdGraphNode* Node2, TArray<FGraphDiffControl::FNodeMatch> const* Exclusions)
+{
+	if(Node2->GetClass() != Node1->GetClass()) 
+	{
+		return false;
+	}
+
+	if(Node1->NodeGuid == Node2->NodeGuid)
+	{
+		return true;
+	}
+
+	// we could be diffing two completely separate assets, this makes sure both 
+	// nodes historically belong to the same graph
+	bool bIsIntraAssetDiff = (Node1->GetGraph()->GraphGuid == Node2->GetGraph()->GraphGuid);
+
+	// if both nodes are from the same graph
+	if (bIsIntraAssetDiff)
+	{
+		return (Node1->GetFName() == Node2->GetFName());
+	}
+
+	if (Exclusions != NULL)
+	{
+		// have to see if this node has already been matched with another
+		for (auto MatchIt(Exclusions->CreateConstIterator()); MatchIt; ++MatchIt)
+		{
+			FGraphDiffControl::FNodeMatch const& PriorMatch = *MatchIt;
+			if (!PriorMatch.IsValid())
+			{
+				continue;
+			}
+
+			// if one of these nodes has already been matched to a different node
+			if (((PriorMatch.OldNode == Node1) && (PriorMatch.NewNode != Node2)) ||
+				((PriorMatch.OldNode == Node2) && (PriorMatch.NewNode != Node1)) ||
+				((PriorMatch.NewNode == Node1) && (PriorMatch.OldNode != Node2)) ||
+				((PriorMatch.NewNode == Node2) && (PriorMatch.OldNode != Node1)))
+			{
+				return false;
+			}
+		}
+	}
+
+	// the name hashes won't match for nodes from separate graph assets, so we 
+	// need to look for some kind of semblance between the two... title? (what's displayed to the user)
+	FText Title1 = Node1->GetNodeTitle(ENodeTitleType::FullTitle);
+	FText Title2 = Node2->GetNodeTitle(ENodeTitleType::FullTitle);
+
+	return Title1.CompareTo(Title2) == 0;
+}
+
 bool FGraphDiffControl::DiffGraphs(UEdGraph* const LhsGraph, UEdGraph* const RhsGraph, TArray<FDiffSingleResult>& DiffsOut)
 {
 	bool bFoundDifferences = false;
@@ -523,6 +599,9 @@ bool FGraphDiffControl::DiffGraphs(UEdGraph* const LhsGraph, UEdGraph* const Rhs
 		TArray<FGraphDiffControl::FNodeMatch> NodeMatches;
 		TSet<UEdGraphNode const*> MatchedRhsNodes;
 
+		FGraphDiffControl::FNodeDiffContext AdditiveDiffContext;
+		AdditiveDiffContext.NodeTypeDisplayName = LOCTEXT("NodeDiffDisplayName", "Node");
+
 		// march through the all the nodes in the rhs graph and look for matches 
 		for (auto NodeIt(RhsGraph->Nodes.CreateConstIterator()); NodeIt; ++NodeIt)
 		{
@@ -541,9 +620,13 @@ bool FGraphDiffControl::DiffGraphs(UEdGraph* const LhsGraph, UEdGraph* const Rhs
 				MatchedRhsNodes.Add(NodeMatch.OldNode);
 			}
 
-			bFoundDifferences |= NodeMatch.Diff(&DiffsOut);
+			bFoundDifferences |= NodeMatch.Diff(AdditiveDiffContext, &DiffsOut);
 		}
 
+		FGraphDiffControl::FNodeDiffContext SubtractiveDiffContext = AdditiveDiffContext;
+		SubtractiveDiffContext.DiffMode = EDiffMode::Subtractive;
+		SubtractiveDiffContext.DiffFlags = EDiffFlags::NodeExistance;
+
 		// go through the lhs nodes to catch ones that may have been missing from the rhs graph
 		for (auto NodeIt(LhsGraph->Nodes.CreateConstIterator()); NodeIt; ++NodeIt)
 		{
@@ -555,33 +638,7 @@ bool FGraphDiffControl::DiffGraphs(UEdGraph* const LhsGraph, UEdGraph* const Rhs
 			}
 
 			FGraphDiffControl::FNodeMatch NodeMatch = FGraphDiffControl::FindNodeMatch(RhsGraph, LhsNode, NodeMatches);
-
-			TArray<FDiffSingleResult> RhsDiffs;
-			if (NodeMatch.Diff(&RhsDiffs))
-			{
-				for (auto DiffIt(RhsDiffs.CreateIterator()); DiffIt; ++DiffIt)
-				{
-					FDiffSingleResult& Difference = *DiffIt;
-
-					// only looking for "adds" because all other diffs would have been caught earlier
-					if (Difference.Diff == EDiffType::NODE_ADDED) 
-					{
-						// flip the "add" to a "removed" because all the diffs are
-						// listed in a context of "what has changed from lhs to rhs?"
-						Difference.Diff	         = EDiffType::NODE_REMOVED;
-
-						FFormatNamedArguments Args;
-						Args.Add(TEXT("NodeTitle"), Difference.Node1->GetNodeTitle(ENodeTitleType::ListView));
-						Difference.ToolTip       = FText::Format(LOCTEXT("DIF_RemoveNode", "Removed Node '{NodeTitle}'"), Args);
-
-						Difference.DisplayString = Difference.ToolTip;
-						Difference.DisplayColor  = FLinearColor(1.f,0.4f,0.4f);
-
-						bFoundDifferences = true;
-						DiffsOut.Add(Difference);
-					}
-				}
-			}
+			bFoundDifferences |= NodeMatch.Diff(SubtractiveDiffContext, &DiffsOut);
 		}
 	}
 
diff --git a/Engine/Source/Editor/GraphEditor/Private/SGraphEditorImpl.cpp b/Engine/Source/Editor/GraphEditor/Private/SGraphEditorImpl.cpp
--- a/Engine/Source/Editor/GraphEditor/Private/SGraphEditorImpl.cpp
+++ b/Engine/Source/Editor/GraphEditor/Private/SGraphEditorImpl.cpp
@@ -22,9 +22,9 @@ bool SGraphEditorImpl::IsNodeTitleVisible( const UEdGraphNode* Node, bool bEnsur
 	return GraphPanel->IsNodeTitleVisible(Node, bEnsureVisible);
 }
 
-void SGraphEditorImpl::JumpToNode( const UEdGraphNode* JumpToMe, bool bRequestRename )
+void SGraphEditorImpl::JumpToNode( const UEdGraphNode* JumpToMe, bool bRequestRename, bool bSelectNode )
 {
-	GraphPanel->JumpToNode(JumpToMe, bRequestRename);
+	GraphPanel->JumpToNode(JumpToMe, bRequestRename, bSelectNode);
 	FocusLockedEditorHere();
 }
 
diff --git a/Engine/Source/Editor/GraphEditor/Private/SGraphEditorImpl.h b/Engine/Source/Editor/GraphEditor/Private/SGraphEditorImpl.h
--- a/Engine/Source/Editor/GraphEditor/Private/SGraphEditorImpl.h
+++ b/Engine/Source/Editor/GraphEditor/Private/SGraphEditorImpl.h
@@ -125,7 +125,7 @@ public:
 	virtual void SelectAllNodes() override;
 	virtual FVector2D GetPasteLocation() const override;
 	virtual bool IsNodeTitleVisible( const UEdGraphNode* Node, bool bRequestRename ) override;
-	virtual void JumpToNode( const UEdGraphNode* JumpToMe, bool bRequestRename = false ) override;
+	virtual void JumpToNode( const UEdGraphNode* JumpToMe, bool bRequestRename = false, bool bSelectNode = true ) override;
 	virtual void JumpToPin( const UEdGraphPin* JumpToMe ) override;
 	virtual UEdGraphPin* GetGraphPinForMenu() override;
 	virtual void ZoomToFit(bool bOnlySelection) override;
diff --git a/Engine/Source/Editor/GraphEditor/Private/SGraphPanel.cpp b/Engine/Source/Editor/GraphEditor/Private/SGraphPanel.cpp
--- a/Engine/Source/Editor/GraphEditor/Private/SGraphPanel.cpp
+++ b/Engine/Source/Editor/GraphEditor/Private/SGraphPanel.cpp
@@ -254,7 +254,7 @@ int32 SGraphPanel::OnPaint( const FPaintArgs& Args, const FGeometry& AllottedGeo
 					{
 						NodeMatches.Add(NodeMatch);
 					}
-					const bool bNodeIsDifferent = (!GraphObjToDiff || NodeMatch.Diff());
+					const bool bNodeIsDifferent = (!GraphObjToDiff || NodeMatch.Diff(FGraphDiffControl::FNodeDiffContext()));
 
 					/* When dragging off a pin, we want to duck the alpha of some nodes */
 					TSharedPtr< SGraphPin > OnlyStartPin = (1 == PreviewConnectorFromPins.Num()) ? PreviewConnectorFromPins[0].FindInGraphPanel(*this) : TSharedPtr< SGraphPin >();
@@ -1243,7 +1243,7 @@ bool SGraphPanel::JumpToRect(const FVector2D &TopLeft, const FVector2D &BottomRi
 	return true;
 }
 
-void SGraphPanel::JumpToNode(const UEdGraphNode* JumpToMe, bool bRequestRename)
+void SGraphPanel::JumpToNode(const UEdGraphNode* JumpToMe, bool bRequestRename, bool bSelectNode)
 {
 	if (JumpToMe != nullptr)
 	{
@@ -1257,8 +1257,16 @@ void SGraphPanel::JumpToNode(const UEdGraphNode* JumpToMe, bool bRequestRename)
 			}
 		}
 
-		// Select this node, and request that we jump to it.
-		SelectAndCenterObject(JumpToMe, true);
+		if (bSelectNode)
+		{
+			// Select this node, and request that we jump to it.
+			SelectAndCenterObject(JumpToMe, true);
+		}
+		else
+		{
+			// Jump to the node
+			CenterObject(JumpToMe);
+		}
 	}
 }
 
@@ -1266,7 +1274,7 @@ void SGraphPanel::JumpToPin(const UEdGraphPin* JumpToMe)
 {
 	if (JumpToMe != nullptr)
 	{
-		JumpToNode(JumpToMe->GetOwningNode(), false);
+		JumpToNode(JumpToMe->GetOwningNode(), false, true);
 	}
 }
 
diff --git a/Engine/Source/Editor/GraphEditor/Private/SGraphPanel.h b/Engine/Source/Editor/GraphEditor/Private/SGraphPanel.h
--- a/Engine/Source/Editor/GraphEditor/Private/SGraphPanel.h
+++ b/Engine/Source/Editor/GraphEditor/Private/SGraphPanel.h
@@ -126,7 +126,7 @@ public:
 	/** Focuses the view on rectangle, zooming if neccesary */
 	bool JumpToRect(const FVector2D &BottomLeft, const FVector2D &TopRight);
 
-	void JumpToNode(const class UEdGraphNode* JumpToMe, bool bRequestRename);
+	void JumpToNode(const class UEdGraphNode* JumpToMe, bool bRequestRename, bool bSelectNode);
 
 	void JumpToPin(const class UEdGraphPin* JumptToMe);
 
diff --git a/Engine/Source/Editor/GraphEditor/Private/SNodePanel.cpp b/Engine/Source/Editor/GraphEditor/Private/SNodePanel.cpp
--- a/Engine/Source/Editor/GraphEditor/Private/SNodePanel.cpp
+++ b/Engine/Source/Editor/GraphEditor/Private/SNodePanel.cpp
@@ -437,6 +437,7 @@ void SNodePanel::Tick( const FGeometry& AllottedGeometry, const double InCurrent
 {
 	CachedGeometry = AllottedGeometry;
 	bool bWasActiveTimerRegisteredThisFrame = false;
+	bool bCanMoveToTargetObjectThisFrame = true;
 
 	if(DeferredSelectionTargetObjects.Num() > 0)
 	{
@@ -453,22 +454,26 @@ void SNodePanel::Tick( const FGeometry& AllottedGeometry, const double InCurrent
 		{
 			SelectionManager.SetSelectionSet(NewSelectionSet);
 		}
-		DeferredSelectionTargetObjects.Empty();
 
+		DeferredSelectionTargetObjects.Empty();
 
-		// Since we want to move to a target object, do not zoom to extent. Panning and zoom will not begin until next tick however due to the nodes potentially not having a size yet.
-		if( DeferredMovementTargetObject )
-		{
-			bDeferredZoomToNodeExtents = false;
-		}
+		// Do not allow movement to happen this Tick as the selected nodes may not yet have a size set (if they're newly added)
+		bCanMoveToTargetObjectThisFrame = false;
 	}
-	else if(DeferredMovementTargetObject != nullptr && GetBoundsForNodes(true, ZoomTargetTopLeft, ZoomTargetBottomRight, ZoomPadding))
+	
+	if(DeferredMovementTargetObject)
 	{
-		DeferredMovementTargetObject = nullptr;
-		if (!ActiveTimerHandle.IsValid())
+		// Since we want to move to a target object, do not zoom to extent
+		bDeferredZoomToNodeExtents = false;
+
+		if (bCanMoveToTargetObjectThisFrame && GetBoundsForNode(DeferredMovementTargetObject, ZoomTargetTopLeft, ZoomTargetBottomRight, ZoomPadding))
 		{
-			ActiveTimerHandle = RegisterActiveTimer(0.f, FWidgetActiveTimerDelegate::CreateSP(this, &SNodePanel::HandleZoomToFit));
-			bWasActiveTimerRegisteredThisFrame = true;
+			DeferredMovementTargetObject = nullptr;
+			if (!ActiveTimerHandle.IsValid())
+			{
+				ActiveTimerHandle = RegisterActiveTimer(0.f, FWidgetActiveTimerDelegate::CreateSP(this, &SNodePanel::HandleZoomToFit));
+				bWasActiveTimerRegisteredThisFrame = true;
+			}
 		}
 	}
 	
@@ -1058,6 +1063,18 @@ void SNodePanel::SelectAndCenterObject(const UObject* ObjectToSelect, bool bCent
 	}
 }
 
+void SNodePanel::CenterObject(const UObject* ObjectToCenter)
+{
+	DeferredMovementTargetObject = ObjectToCenter;
+
+	auto PinnedActiveTimerHandle = ActiveTimerHandle.Pin();
+	if (PinnedActiveTimerHandle.IsValid())
+	{
+		// Stop zooming for now and let the a new zoom target be established
+		UnRegisterActiveTimer(PinnedActiveTimerHandle.ToSharedRef());
+	}
+}
+
 /** Add a slot to the CanvasPanel dynamically */
 void SNodePanel::AddGraphNode( const TSharedRef<SNodePanel::SNode>& NodeToAdd )
 {
@@ -1361,6 +1378,39 @@ bool SNodePanel::IsNodeCulled(const TSharedRef<SNode>& Node, const FGeometry& Al
 
 }
 
+bool SNodePanel::GetBoundsForNode(const UObject* InNode, /*out*/ FVector2D& MinCorner, /*out*/ FVector2D& MaxCorner, float Padding)
+{
+	MinCorner = FVector2D(MAX_FLT, MAX_FLT);
+	MaxCorner = FVector2D(-MAX_FLT, -MAX_FLT);
+
+	bool bValid = false;
+
+	TSharedRef<SNode>* pWidget = (InNode) ? NodeToWidgetLookup.Find(InNode) : nullptr;
+	if (pWidget)
+	{
+		SNode& Widget = pWidget->Get();
+		const FVector2D Lower = Widget.GetPosition();
+		const FVector2D Upper = Lower + Widget.GetDesiredSize();
+
+		MinCorner.X = FMath::Min(MinCorner.X, Lower.X);
+		MinCorner.Y = FMath::Min(MinCorner.Y, Lower.Y);
+		MaxCorner.X = FMath::Max(MaxCorner.X, Upper.X);
+		MaxCorner.Y = FMath::Max(MaxCorner.Y, Upper.Y);
+
+		bValid = true;
+	}
+
+	if (bValid)
+	{
+		MinCorner.X -= Padding;
+		MinCorner.Y -= Padding;
+		MaxCorner.X += Padding;
+		MaxCorner.Y += Padding;
+	}
+
+	return bValid;
+}
+
 bool SNodePanel::GetBoundsForNodes(bool bSelectionSetOnly, FVector2D& MinCorner, FVector2D& MaxCorner, float Padding)
 {
 	MinCorner = FVector2D(MAX_FLT, MAX_FLT);
diff --git a/Engine/Source/Editor/GraphEditor/Public/DiffResults.h b/Engine/Source/Editor/GraphEditor/Public/DiffResults.h
--- a/Engine/Source/Editor/GraphEditor/Public/DiffResults.h
+++ b/Engine/Source/Editor/GraphEditor/Public/DiffResults.h
@@ -105,9 +105,9 @@ struct FDiffResults
 	}
 
 	/** Test if it can store results*/
-	operator bool () const
+	bool CanStoreResults() const
 	{
-		return ResultArray != NULL;
+		return ResultArray != nullptr;
 	}
 
 	/** Get the number of diffs found*/
diff --git a/Engine/Source/Editor/GraphEditor/Public/GraphDiffControl.h b/Engine/Source/Editor/GraphEditor/Public/GraphDiffControl.h
--- a/Engine/Source/Editor/GraphEditor/Public/GraphDiffControl.h
+++ b/Engine/Source/Editor/GraphEditor/Public/GraphDiffControl.h
@@ -6,6 +6,59 @@
 class GRAPHEDITOR_API FGraphDiffControl
 {
 public:
+	/** What kind of diff are we performing */
+	enum class EDiffMode
+	{
+		/** An item present in 'RHS' but missing in 'LHS' has been added to the graph */
+		Additive,
+		/** An item present in 'RHS' but missing in 'LHS' has been removed from the graph */
+		Subtractive,
+	};
+
+	/** Flags controlling while operations should be included in the diff */
+	struct EDiffFlags
+	{
+		enum Type
+		{
+			/** Include node existence (adds and removes) in the diff */
+			NodeExistance = (1<<0),
+			/** Include node movement in the diff */
+			NodeMovement = (1<<1),
+			/** Include the node comment in the diff */
+			NodeComment = (1<<2),
+			/** Include the node pin changes in the diff */
+			NodePins = (1<<3),
+			/** Include any node specific items in the diff */
+			NodeSpecificDiffs = (1<<4),
+			/** Include all the things */
+			All = NodeExistance | NodeMovement | NodeComment | NodePins | NodeSpecificDiffs,
+		};
+	};
+
+	/** A struct holding the context data for a node diff */
+	struct GRAPHEDITOR_API FNodeDiffContext
+	{
+		FNodeDiffContext()
+			: DiffMode(EDiffMode::Additive)
+			, DiffFlags(EDiffFlags::All)
+			, NodeTypeDisplayName()
+			, bIsRootNode(true)
+		{
+		}
+
+		/** What kind of diff are we performing? */
+		EDiffMode DiffMode;
+
+		/** What items should be included in the diff? */
+		EDiffFlags::Type DiffFlags;
+
+		/** Display name used when showing the node type in messages */
+		FText NodeTypeDisplayName;
+
+		/** True if this node is a root node in the graph, false if this node is nested within another node */
+		bool bIsRootNode;
+	};
+
 	/** A struct to represent a found pair of nodes that match each other (for comparisons sake) */
 	struct GRAPHEDITOR_API FNodeMatch
 	{
@@ -25,7 +78,8 @@ public:
 		 * @param  DiffsArrayOut	If supplied, this will be filled out with all the differences that were found.
 		 * @return True if there were differences found, false if the two nodes are identical.
 		 */
-		bool Diff(TArray<FDiffSingleResult>* DiffsResultsOut = NULL) const;
+		bool Diff(const FNodeDiffContext& DiffContext, TArray<FDiffSingleResult>* DiffsResultsOut = NULL) const;
+		bool Diff(const FNodeDiffContext& DiffContext, FDiffResults& DiffsOut) const;
 
 		/**
 		 * Checks to see if this is a valid match.
@@ -47,6 +101,9 @@ public:
 	 */
 	static FNodeMatch FindNodeMatch(class UEdGraph* OldGraph, class UEdGraphNode* Node, TArray<FNodeMatch> const& PriorMatches);
 
+	/** Determine if the two Nodes are the same */
+	static bool IsNodeMatch(class UEdGraphNode* Node1, class UEdGraphNode* Node2, TArray<FGraphDiffControl::FNodeMatch> const* Exclusions = nullptr);
+
 	/**
 	 * Looks for node differences between the two supplied graphs. Diffs will be
 	 * returned in the context of "what has changed from OldGraph in NewGraph?" 
diff --git a/Engine/Source/Editor/GraphEditor/Public/SNodePanel.h b/Engine/Source/Editor/GraphEditor/Public/SNodePanel.h
--- a/Engine/Source/Editor/GraphEditor/Public/SNodePanel.h
+++ b/Engine/Source/Editor/GraphEditor/Public/SNodePanel.h
@@ -702,6 +702,13 @@ protected:
 	 */
 	void SelectAndCenterObject(const UObject* ObjectToSelect, bool bCenter);
 
+	/**
+	 * On the next tick, centers the widget associated with the object if it exists
+	 *
+	 * @param ObjectToCenter	The object to center
+	 */
+	void CenterObject(const UObject* ObjectToCenter);
+
 	/** Add a slot to the CanvasPanel dynamically */
 	virtual void AddGraphNode(const TSharedRef<SNode>& NodeToAdd);
 
@@ -740,6 +747,12 @@ protected:
 	virtual TSharedPtr<SWidget> OnSummonContextMenu(const FGeometry& MyGeometry, const FPointerEvent& MouseEvent) { return TSharedPtr<SWidget>(); }
 	virtual bool OnHandleLeftMouseRelease(const FGeometry& MyGeometry, const FPointerEvent& MouseEvent) { return false; }
 protected:
+	/**
+	 * Get the bounds of the given node
+	 * @return True if successful
+	 */
+	bool GetBoundsForNode(const UObject* InNode, /*out*/ FVector2D& MinCorner, /*out*/ FVector2D& MaxCorner, float Padding = 0.0f);
+
 	/**
 	 * Get the bounds of the selected nodes 
 	 * @param bSelectionSetOnly If true, limits the query to just the selected nodes.  Otherwise it does all nodes.
diff --git a/Engine/Source/Editor/UnrealEd/Public/GraphEditor.h b/Engine/Source/Editor/UnrealEd/Public/GraphEditor.h
--- a/Engine/Source/Editor/UnrealEd/Public/GraphEditor.h
+++ b/Engine/Source/Editor/UnrealEd/Public/GraphEditor.h
@@ -215,11 +215,11 @@ public:
 	}
 
 	/** Bring the specified node into view */
-	virtual void JumpToNode( const class UEdGraphNode* JumpToMe, bool bRequestRename )
+	virtual void JumpToNode( const class UEdGraphNode* JumpToMe, bool bRequestRename = false, bool bSelectNode = true )
 	{
 		if (Implementation.IsValid())
 		{
-			Implementation->JumpToNode(JumpToMe, bRequestRename);
+			Implementation->JumpToNode(JumpToMe, bRequestRename, bSelectNode);
 		}
 	}
 
