[UE-14700] Improving check message, CNR

diff --git a/Engine/Source/Editor/GraphEditor/Private/SGraphPin.cpp b/Engine/Source/Editor/GraphEditor/Private/SGraphPin.cpp
--- a/Engine/Source/Editor/GraphEditor/Private/SGraphPin.cpp
+++ b/Engine/Source/Editor/GraphEditor/Private/SGraphPin.cpp
@@ -161,7 +161,13 @@ void SGraphPin::Construct(const FArguments& InArgs, UEdGraphPin* InPin)
 	check(GraphPinObj != NULL);
 
 	const UEdGraphSchema* Schema = GraphPinObj->GetSchema();
-	check(Schema);
+	checkf(
+		Schema, 
+		TEXT("Missing schema for pin: %s with outer: %s of type %s"), 
+		*(GraphPinObj->GetName()),
+		GraphPinObj->GetOuter() ? *(GraphPinObj->GetOuter()->GetName()) : TEXT("NULL OUTER"), 
+		GraphPinObj->GetOuter() ? *(GraphPinObj->GetOuter()->GetClass()->GetName()) : TEXT("NULL OUTER")
+	);
 
 	const bool bCanConnectToPin = !GraphPinObj->bNotConnectable;
 	const bool bIsInput = (GetDirection() == EGPD_Input);
