Various graphs that looked disabled, will now correctly display as Read Only but can be panned and nodes are selectable.

Reworked how graphs were disabled to stop all out disabling of the SGraphPanel while still maintaining the visual display of being read-only.

Removed concept of TitleBarEnabledOnly from the graph editor and merged it with the IsEditable functionality.

diff --git a/Engine/Source/Editor/BehaviorTreeEditor/Private/SBehaviorTreeDiff.cpp b/Engine/Source/Editor/BehaviorTreeEditor/Private/SBehaviorTreeDiff.cpp
--- a/Engine/Source/Editor/BehaviorTreeEditor/Private/SBehaviorTreeDiff.cpp
+++ b/Engine/Source/Editor/BehaviorTreeEditor/Private/SBehaviorTreeDiff.cpp
@@ -457,7 +457,6 @@ void SBehaviorTreeDiff::FBehaviorTreeDiffPanel::GeneratePanel(UEdGraph* Graph, U
 			.GraphToEdit(Graph)
 			.GraphToDiff(GraphToDiff)
 			.IsEditable(false)
-			.TitleBarEnabledOnly(false)
 			.TitleBar(SNew(SBorder).HAlign(HAlign_Center)
 			[
 				SNew(STextBlock).Text(GetTitle())
diff --git a/Engine/Source/Editor/GraphEditor/Private/GraphEditorModule.cpp b/Engine/Source/Editor/GraphEditor/Private/GraphEditorModule.cpp
--- a/Engine/Source/Editor/GraphEditor/Private/GraphEditorModule.cpp
+++ b/Engine/Source/Editor/GraphEditor/Private/GraphEditorModule.cpp
@@ -59,7 +59,6 @@ TSharedRef<SGraphEditor> FGraphEditorModule::PRIVATE_MakeGraphEditor(
 	const TAttribute<bool>& InIsEmpty,
 	TAttribute<FGraphAppearanceInfo> Appearance,
 	TSharedPtr<SWidget> InTitleBar,
-	const TAttribute<bool>& InTitleBarEnabledOnly,
 	UEdGraph* InGraphToEdit,
 	SGraphEditor::FGraphEditorEvents InGraphEvents,
 	bool InAutoExpandActionMenu,
@@ -74,7 +73,6 @@ TSharedRef<SGraphEditor> FGraphEditorModule::PRIVATE_MakeGraphEditor(
 		.IsEditable(InIsEditable)
 		.Appearance(Appearance)
 		.TitleBar(InTitleBar)
-		.TitleBarEnabledOnly(InTitleBarEnabledOnly)
 		.GraphToEdit(InGraphToEdit)
 		.GraphEvents(InGraphEvents)
 		.AutoExpandActionMenu(InAutoExpandActionMenu)
diff --git a/Engine/Source/Editor/GraphEditor/Private/SGraphEditorImpl.cpp b/Engine/Source/Editor/GraphEditor/Private/SGraphEditorImpl.cpp
--- a/Engine/Source/Editor/GraphEditor/Private/SGraphEditorImpl.cpp
+++ b/Engine/Source/Editor/GraphEditor/Private/SGraphEditorImpl.cpp
@@ -161,7 +161,6 @@ void SGraphEditorImpl::Construct( const FArguments& InArgs )
 	Commands = MakeShareable( new FUICommandList() );
 	IsEditable = InArgs._IsEditable;
 	Appearance = InArgs._Appearance;
-	TitleBarEnabledOnly = InArgs._TitleBarEnabledOnly;
 	TitleBar	= InArgs._TitleBar;
 	bAutoExpandActionMenu = InArgs._AutoExpandActionMenu;
 	ShowGraphStateOverlay = InArgs._ShowGraphStateOverlay;
@@ -252,7 +251,6 @@ void SGraphEditorImpl::Construct( const FArguments& InArgs )
 			.IsEditable( this, &SGraphEditorImpl::IsGraphEditable )
 			.OnDropActor( InArgs._GraphEvents.OnDropActor )
 			.OnDropStreamingLevel( InArgs._GraphEvents.OnDropStreamingLevel )
-			.IsEnabled(this, &SGraphEditorImpl::GraphEd_OnGetGraphEnabled)
 			.OnVerifyTextCommit( InArgs._GraphEvents.OnVerifyTextCommit )
 			.OnTextCommitted( InArgs._GraphEvents.OnTextCommitted )
 			.OnSpawnNodeByShortcut( InArgs._GraphEvents.OnSpawnNodeByShortcut )
@@ -410,12 +408,6 @@ void SGraphEditorImpl::OnClosedActionMenu()
 	GraphPanel->OnStopMakingConnection(/*bForceStop=*/ true);
 }
 
-bool SGraphEditorImpl::GraphEd_OnGetGraphEnabled() const
-{
-	const bool bTitleBarOnly = TitleBarEnabledOnly.Get();
-	return !bTitleBarOnly;
-}
-
 FActionMenuContent SGraphEditorImpl::GraphEd_OnGetContextMenuFor(const FGraphContextMenuArguments& SpawnInfo)
 {
 	if (EdGraphObj != NULL)
diff --git a/Engine/Source/Editor/GraphEditor/Private/SGraphEditorImpl.h b/Engine/Source/Editor/GraphEditor/Private/SGraphEditorImpl.h
--- a/Engine/Source/Editor/GraphEditor/Private/SGraphEditorImpl.h
+++ b/Engine/Source/Editor/GraphEditor/Private/SGraphEditorImpl.h
@@ -26,7 +26,6 @@ public:
 		SLATE_ARGUMENT(TSharedPtr<FUICommandList>, AdditionalCommands)
 		SLATE_ATTRIBUTE( bool, IsEditable )
 		SLATE_ARGUMENT( TSharedPtr<SWidget>, TitleBar )
-		SLATE_ATTRIBUTE( bool, TitleBarEnabledOnly )
 		SLATE_ATTRIBUTE( FGraphAppearanceInfo, Appearance )
 		SLATE_ARGUMENT( UEdGraph*, GraphToEdit )
 		SLATE_ARGUMENT( UEdGraph*, GraphToDiff )
@@ -58,7 +57,6 @@ private:
 	SGraphEditor::FOnCreateActionMenu OnCreateActionMenu;
 
 	TAttribute<bool> IsEditable;
-	TAttribute<bool> TitleBarEnabledOnly;
 
 	bool bAutoExpandActionMenu;
 
@@ -101,8 +99,6 @@ public:
 
 	void OnClosedActionMenu();
 
-	bool GraphEd_OnGetGraphEnabled() const;
-
 	FActionMenuContent GraphEd_OnGetContextMenuFor(const FGraphContextMenuArguments& SpawnInfo);
 
 	//void GraphEd_OnPanelUpdated();
diff --git a/Engine/Source/Editor/GraphEditor/Private/SGraphPanel.cpp b/Engine/Source/Editor/GraphEditor/Private/SGraphPanel.cpp
--- a/Engine/Source/Editor/GraphEditor/Private/SGraphPanel.cpp
+++ b/Engine/Source/Editor/GraphEditor/Private/SGraphPanel.cpp
@@ -262,7 +262,7 @@ int32 SGraphPanel::OnPaint( const FPaintArgs& Args, const FGeometry& AllottedGeo
 					const FWidgetStyle& NodeStyleToUse = (bNodeIsDifferent && !bNodeIsNotUsableInCurrentContext)? InWidgetStyle : FadedStyle;
 
 					// Draw the node.O
-					CurWidgetsMaxLayerId = CurWidget.Widget->Paint( Args.WithNewParent(this), CurWidget.Geometry, MyClippingRect, OutDrawElements, ChildLayerId, NodeStyleToUse, ShouldBeEnabled( bParentEnabled ) );
+					CurWidgetsMaxLayerId = CurWidget.Widget->Paint( Args.WithNewParent(this), CurWidget.Geometry, MyClippingRect, OutDrawElements, ChildLayerId, NodeStyleToUse, IsEditable.Get() && ShouldBeEnabled( bParentEnabled ) );
 				}
 
 				// Draw the node's overlay, if it has one.
diff --git a/Engine/Source/Editor/GraphEditor/Public/GraphEditorModule.h b/Engine/Source/Editor/GraphEditor/Public/GraphEditorModule.h
--- a/Engine/Source/Editor/GraphEditor/Public/GraphEditorModule.h
+++ b/Engine/Source/Editor/GraphEditor/Public/GraphEditorModule.h
@@ -35,7 +35,6 @@ private:
 		const TAttribute<bool>& InIsEmpty,
 		TAttribute<FGraphAppearanceInfo> Appearance,
 		TSharedPtr<SWidget> InTitleBar,
-		const TAttribute<bool>& InTitleBarEnabledOnly,
 		UEdGraph* InGraphToEdit,
 		SGraphEditor::FGraphEditorEvents GraphEvents,
 		bool InAutoExpandActionMenu,
diff --git a/Engine/Source/Editor/Kismet/Private/BlueprintEditor.cpp b/Engine/Source/Editor/Kismet/Private/BlueprintEditor.cpp
--- a/Engine/Source/Editor/Kismet/Private/BlueprintEditor.cpp
+++ b/Engine/Source/Editor/Kismet/Private/BlueprintEditor.cpp
@@ -1173,12 +1173,10 @@ TSharedRef<SGraphEditor> FBlueprintEditor::CreateGraphEditorWidget(TSharedRef<FT
 	}
 
 
-	bool const bEnableTitleBarOnly = !IsGraphPanelEnabled(InGraph);
 	TSharedRef<SGraphEditor> Editor = SNew(SGraphEditor)
 		.AdditionalCommands(GraphEditorCommands)
 		.IsEditable(this, &FBlueprintEditor::IsEditable, InGraph)
 		.TitleBar(TitleBarWidget)
-		.TitleBarEnabledOnly(bEnableTitleBarOnly)
 		.Appearance(this, &FBlueprintEditor::GetGraphAppearance, InGraph)
 		.GraphToEdit(InGraph)
 		.GraphEvents(InEvents)
@@ -7537,7 +7535,7 @@ UEdGraph* FBlueprintEditor::GetFocusedGraph() const
 
 bool FBlueprintEditor::IsEditable(UEdGraph* InGraph) const
 {
-	return InEditingMode() && InGraph && InGraph->bEditable;
+	return InEditingMode() && InGraph && InGraph->bEditable && IsGraphPanelEnabled(InGraph);
 }
 
 bool FBlueprintEditor::IsGraphPanelEnabled(UEdGraph* InGraph) const
diff --git a/Engine/Source/Editor/Kismet/Private/SBlueprintDiff.cpp b/Engine/Source/Editor/Kismet/Private/SBlueprintDiff.cpp
--- a/Engine/Source/Editor/Kismet/Private/SBlueprintDiff.cpp
+++ b/Engine/Source/Editor/Kismet/Private/SBlueprintDiff.cpp
@@ -1106,7 +1106,6 @@ void FDiffPanel::GeneratePanel(UEdGraph* Graph, UEdGraph* GraphToDiff )
 			.GraphToEdit(Graph)
 			.GraphToDiff(GraphToDiff)
 			.IsEditable(false)
-			.TitleBarEnabledOnly(false)
 			.GraphEvents(InEvents);
 
 		GraphEditor = Editor;
diff --git a/Engine/Source/Editor/Kismet/Private/SGraphTitleBar.cpp b/Engine/Source/Editor/Kismet/Private/SGraphTitleBar.cpp
--- a/Engine/Source/Editor/Kismet/Private/SGraphTitleBar.cpp
+++ b/Engine/Source/Editor/Kismet/Private/SGraphTitleBar.cpp
@@ -43,18 +43,19 @@ FText SGraphTitleBar::GetTitleExtra() const
 {
 	check(EdGraphObj != NULL);
 
-	FString ExtraText = Kismet2Ptr.Pin()->GetGraphDecorationString(EdGraphObj).ToString();
+	FText ExtraText = Kismet2Ptr.Pin()->GetGraphDecorationString(EdGraphObj);
 
 	const bool bEditable = (Kismet2Ptr.Pin()->IsEditable(EdGraphObj));
 
 	if(!bEditable)
 	{
-		ExtraText += TEXT(" (");
-		ExtraText += LOCTEXT("ReadOnlyWarningText", "READ-ONLY").ToString();
-		ExtraText += TEXT(")");
+		FFormatNamedArguments Args;
+		Args.Add(TEXT("BaseText"), ExtraText);
+
+		ExtraText = FText::Format(LOCTEXT("ReadOnlyWarningText", "{BaseText} (READ-ONLY)"), Args);
 	}
 
-	return FText::FromString(ExtraText);
+	return ExtraText;
 }
 
 EVisibility SGraphTitleBar::IsGraphBlueprintNameVisible() const
diff --git a/Engine/Source/Editor/UnrealEd/Private/GraphEditor.cpp b/Engine/Source/Editor/UnrealEd/Private/GraphEditor.cpp
--- a/Engine/Source/Editor/UnrealEd/Private/GraphEditor.cpp
+++ b/Engine/Source/Editor/UnrealEd/Private/GraphEditor.cpp
@@ -22,7 +22,6 @@ void SGraphEditor::ConstructImplementation( const FArguments& InArgs )
 		InArgs._IsEmpty,
 		InArgs._Appearance,
 		InArgs._TitleBar,
-		InArgs._TitleBarEnabledOnly,
 		InArgs._GraphToEdit,
 		InArgs._GraphEvents,
 		InArgs._AutoExpandActionMenu,
diff --git a/Engine/Source/Editor/UnrealEd/Public/GraphEditor.h b/Engine/Source/Editor/UnrealEd/Public/GraphEditor.h
--- a/Engine/Source/Editor/UnrealEd/Public/GraphEditor.h
+++ b/Engine/Source/Editor/UnrealEd/Public/GraphEditor.h
@@ -112,7 +112,6 @@ public:
 		: _AdditionalCommands( static_cast<FUICommandList*>(NULL) )
 		, _IsEditable(true)
 		, _IsEmpty(false)
-		, _TitleBarEnabledOnly(false)
 		, _GraphToEdit(NULL)
 		, _GraphToDiff(NULL)
 		, _AutoExpandActionMenu(false)
@@ -124,7 +123,6 @@ public:
 		SLATE_ATTRIBUTE( bool, IsEmpty )	
 		SLATE_ARGUMENT( TSharedPtr<SWidget>, TitleBar )
 		SLATE_ATTRIBUTE( FGraphAppearanceInfo, Appearance )
-		SLATE_ATTRIBUTE( bool, TitleBarEnabledOnly )
 		SLATE_EVENT( FEdGraphEvent, OnGraphModuleReloaded )
 		SLATE_ARGUMENT( UEdGraph*, GraphToEdit )
 		SLATE_ARGUMENT( UEdGraph*, GraphToDiff )
