This adds a missed safety check in the comment bubbles before trying to set the text widget content.
#UE-15081 - Editor crashes without crash report or warning when entering a node comment on a blueprint node
#Branch UE4
#Proj GraphEditor

diff --git a/Engine/Source/Editor/GraphEditor/Private/SCommentBubble.cpp b/Engine/Source/Editor/GraphEditor/Private/SCommentBubble.cpp
--- a/Engine/Source/Editor/GraphEditor/Private/SCommentBubble.cpp
+++ b/Engine/Source/Editor/GraphEditor/Private/SCommentBubble.cpp
@@ -127,7 +127,10 @@ void SCommentBubble::Tick( const FGeometry& AllottedGeometry, const double InCur
 		// Call text commit delegate
 		OnTextCommittedDelegate.ExecuteIfBound( CachedCommentText, ETextCommit::Default );
 		// Reflect changes to the Textblock because it doesn't update itself.
-		TextBlock->SetText( CachedCommentText );
+		if( TextBlock.IsValid() )
+		{
+			TextBlock->SetText( CachedCommentText );
+		}
 		// Toggle the comment on/off, provided it the parent isn't a comment node
 		if( !bInvertLODCulling )
 		{
