Blueprints and Behavior Trees will no longer have their graphs display as disabled/read-only when debugging (in PIE).

Fallout from CL# 2535806

diff --git a/Engine/Source/Editor/GraphEditor/Private/GraphEditorModule.cpp b/Engine/Source/Editor/GraphEditor/Private/GraphEditorModule.cpp
--- a/Engine/Source/Editor/GraphEditor/Private/GraphEditorModule.cpp
+++ b/Engine/Source/Editor/GraphEditor/Private/GraphEditorModule.cpp
@@ -55,7 +55,8 @@ void FGraphEditorModule::ShutdownModule()
  */
 TSharedRef<SGraphEditor> FGraphEditorModule::PRIVATE_MakeGraphEditor( 
 	const TSharedPtr<FUICommandList>& InAdditionalCommands, 
-	const TAttribute<bool>& InIsEditable, 
+	const TAttribute<bool>& InIsEditable,
+	const TAttribute<bool>& InDisplayAsReadOnly,
 	const TAttribute<bool>& InIsEmpty,
 	TAttribute<FGraphAppearanceInfo> Appearance,
 	TSharedPtr<SWidget> InTitleBar,
@@ -71,6 +72,7 @@ TSharedRef<SGraphEditor> FGraphEditorModule::PRIVATE_MakeGraphEditor(
 		SNew(SGraphEditorImpl)
 		.AdditionalCommands(InAdditionalCommands)
 		.IsEditable(InIsEditable)
+		.DisplayAsReadOnly(InDisplayAsReadOnly)
 		.Appearance(Appearance)
 		.TitleBar(InTitleBar)
 		.GraphToEdit(InGraphToEdit)
diff --git a/Engine/Source/Editor/GraphEditor/Private/SGraphEditorImpl.cpp b/Engine/Source/Editor/GraphEditor/Private/SGraphEditorImpl.cpp
--- a/Engine/Source/Editor/GraphEditor/Private/SGraphEditorImpl.cpp
+++ b/Engine/Source/Editor/GraphEditor/Private/SGraphEditorImpl.cpp
@@ -160,6 +160,7 @@ void SGraphEditorImpl::Construct( const FArguments& InArgs )
 {
 	Commands = MakeShareable( new FUICommandList() );
 	IsEditable = InArgs._IsEditable;
+	DisplayAsReadOnly = InArgs._DisplayAsReadOnly;
 	Appearance = InArgs._Appearance;
 	TitleBar	= InArgs._TitleBar;
 	bAutoExpandActionMenu = InArgs._AutoExpandActionMenu;
@@ -249,6 +250,7 @@ void SGraphEditorImpl::Construct( const FArguments& InArgs )
 			.OnSelectionChanged( InArgs._GraphEvents.OnSelectionChanged )
 			.OnNodeDoubleClicked( InArgs._GraphEvents.OnNodeDoubleClicked )
 			.IsEditable( this, &SGraphEditorImpl::IsGraphEditable )
+			.DisplayAsReadOnly( this, &SGraphEditorImpl::DisplayGraphAsReadOnly )
 			.OnDropActor( InArgs._GraphEvents.OnDropActor )
 			.OnDropStreamingLevel( InArgs._GraphEvents.OnDropStreamingLevel )
 			.OnVerifyTextCommit( InArgs._GraphEvents.OnVerifyTextCommit )
@@ -560,6 +562,11 @@ bool SGraphEditorImpl::IsGraphEditable() const
 	return (EdGraphObj != NULL) && IsEditable.Get();
 }
 
+bool SGraphEditorImpl::DisplayGraphAsReadOnly() const
+{
+	return (EdGraphObj != NULL) && DisplayAsReadOnly.Get();
+}
+
 bool SGraphEditorImpl::IsLocked() const
 {
 	for( auto LockedGraph : LockedGraphs )
diff --git a/Engine/Source/Editor/GraphEditor/Private/SGraphEditorImpl.h b/Engine/Source/Editor/GraphEditor/Private/SGraphEditorImpl.h
--- a/Engine/Source/Editor/GraphEditor/Private/SGraphEditorImpl.h
+++ b/Engine/Source/Editor/GraphEditor/Private/SGraphEditorImpl.h
@@ -20,11 +20,13 @@ public:
 	SLATE_BEGIN_ARGS( SGraphEditorImpl )
 		: _AdditionalCommands( TSharedPtr<FUICommandList>() )
 		, _IsEditable(true)
+		, _DisplayAsReadOnly(false)
 		{}
 
 
 		SLATE_ARGUMENT(TSharedPtr<FUICommandList>, AdditionalCommands)
 		SLATE_ATTRIBUTE( bool, IsEditable )
+		SLATE_ATTRIBUTE( bool, DisplayAsReadOnly )
 		SLATE_ARGUMENT( TSharedPtr<SWidget>, TitleBar )
 		SLATE_ATTRIBUTE( FGraphAppearanceInfo, Appearance )
 		SLATE_ARGUMENT( UEdGraph*, GraphToEdit )
@@ -58,6 +60,9 @@ private:
 
 	TAttribute<bool> IsEditable;
 
+	/** Attribute for displaying the graph as read-only, which is a visual state only where IsEditable is a functional state */
+	TAttribute<bool> DisplayAsReadOnly;
+
 	bool bAutoExpandActionMenu;
 
 	/** Whether to show the state (read only / PIE etc) Overlay on the panel */
@@ -161,6 +166,9 @@ private:
 
 	bool IsGraphEditable() const;
 
+	/** Helper function to decide whether to display the graph in a read-only state */
+	bool DisplayGraphAsReadOnly() const;
+
 	bool IsLocked() const;
 
 private:
diff --git a/Engine/Source/Editor/GraphEditor/Private/SGraphPanel.cpp b/Engine/Source/Editor/GraphEditor/Private/SGraphPanel.cpp
--- a/Engine/Source/Editor/GraphEditor/Private/SGraphPanel.cpp
+++ b/Engine/Source/Editor/GraphEditor/Private/SGraphPanel.cpp
@@ -77,6 +77,7 @@ void SGraphPanel::Construct( const SGraphPanel::FArguments& InArgs )
 	this->GraphObjToDiff = InArgs._GraphObjToDiff;
 	this->SelectionManager.OnSelectionChanged = InArgs._OnSelectionChanged;
 	this->IsEditable = InArgs._IsEditable;
+	this->DisplayAsReadOnly = InArgs._DisplayAsReadOnly;
 	this->OnNodeDoubleClicked = InArgs._OnNodeDoubleClicked;
 	this->OnDropActor = InArgs._OnDropActor;
 	this->OnDropStreamingLevel = InArgs._OnDropStreamingLevel;
@@ -262,7 +263,7 @@ int32 SGraphPanel::OnPaint( const FPaintArgs& Args, const FGeometry& AllottedGeo
 					const FWidgetStyle& NodeStyleToUse = (bNodeIsDifferent && !bNodeIsNotUsableInCurrentContext)? InWidgetStyle : FadedStyle;
 
 					// Draw the node.O
-					CurWidgetsMaxLayerId = CurWidget.Widget->Paint( Args.WithNewParent(this), CurWidget.Geometry, MyClippingRect, OutDrawElements, ChildLayerId, NodeStyleToUse, IsEditable.Get() && ShouldBeEnabled( bParentEnabled ) );
+					CurWidgetsMaxLayerId = CurWidget.Widget->Paint( Args.WithNewParent(this), CurWidget.Geometry, MyClippingRect, OutDrawElements, ChildLayerId, NodeStyleToUse, !DisplayAsReadOnly.Get() && ShouldBeEnabled( bParentEnabled ) );
 				}
 
 				// Draw the node's overlay, if it has one.
diff --git a/Engine/Source/Editor/GraphEditor/Private/SGraphPanel.h b/Engine/Source/Editor/GraphEditor/Private/SGraphPanel.h
--- a/Engine/Source/Editor/GraphEditor/Private/SGraphPanel.h
+++ b/Engine/Source/Editor/GraphEditor/Private/SGraphPanel.h
@@ -44,6 +44,7 @@ public:
 		, _GraphObjToDiff( static_cast<UEdGraph*>(NULL) )
 		, _InitialZoomToFit( false )
 		, _IsEditable( true )
+		, _DisplayAsReadOnly( false )
 		, _ShowGraphStateOverlay(true)
 		, _OnUpdateGraphPanel()
 		{}
@@ -57,6 +58,7 @@ public:
 		SLATE_ARGUMENT( class UEdGraph*, GraphObjToDiff )
 		SLATE_ARGUMENT( bool, InitialZoomToFit )
 		SLATE_ATTRIBUTE( bool, IsEditable )
+		SLATE_ATTRIBUTE( bool, DisplayAsReadOnly )
 		/** Show overlay elements for the graph state such as the PIE and read-only borders and text */
 		SLATE_ATTRIBUTE(bool, ShowGraphStateOverlay)
 		SLATE_EVENT( FOnNodeVerifyTextCommit, OnVerifyTextCommit )
@@ -252,6 +254,9 @@ private:
 	/** Map of recently added nodes for the panel (maps from added nodes to UserActions indices) */
 	TMap<const class UEdGraphNode*, int32> UserAddedNodes;
 
+	/** Should the graph display all nodes in a read-only state (grayed)? This does not affect functionality of using them (IsEditable) */
+	TAttribute<bool> DisplayAsReadOnly;
+
 	FOnGraphChanged::FDelegate MyRegisteredGraphChangedDelegate;
 	FDelegateHandle            MyRegisteredGraphChangedDelegateHandle;
 private:
diff --git a/Engine/Source/Editor/GraphEditor/Public/GraphEditorModule.h b/Engine/Source/Editor/GraphEditor/Public/GraphEditorModule.h
--- a/Engine/Source/Editor/GraphEditor/Public/GraphEditorModule.h
+++ b/Engine/Source/Editor/GraphEditor/Public/GraphEditorModule.h
@@ -32,6 +32,7 @@ private:
 	virtual TSharedRef<SGraphEditor> PRIVATE_MakeGraphEditor(
 		const TSharedPtr<FUICommandList>& InAdditionalCommands,
 		const TAttribute<bool>& InIsEditable,
+		const TAttribute<bool>& InDisplayAsReadOnly,
 		const TAttribute<bool>& InIsEmpty,
 		TAttribute<FGraphAppearanceInfo> Appearance,
 		TSharedPtr<SWidget> InTitleBar,
diff --git a/Engine/Source/Editor/Kismet/Private/BlueprintEditor.cpp b/Engine/Source/Editor/Kismet/Private/BlueprintEditor.cpp
--- a/Engine/Source/Editor/Kismet/Private/BlueprintEditor.cpp
+++ b/Engine/Source/Editor/Kismet/Private/BlueprintEditor.cpp
@@ -1176,6 +1176,7 @@ TSharedRef<SGraphEditor> FBlueprintEditor::CreateGraphEditorWidget(TSharedRef<FT
 	TSharedRef<SGraphEditor> Editor = SNew(SGraphEditor)
 		.AdditionalCommands(GraphEditorCommands)
 		.IsEditable(this, &FBlueprintEditor::IsEditable, InGraph)
+		.DisplayAsReadOnly(this, &FBlueprintEditor::IsGraphReadOnly, InGraph)
 		.TitleBar(TitleBarWidget)
 		.Appearance(this, &FBlueprintEditor::GetGraphAppearance, InGraph)
 		.GraphToEdit(InGraph)
@@ -7520,7 +7521,12 @@ UEdGraph* FBlueprintEditor::GetFocusedGraph() const
 
 bool FBlueprintEditor::IsEditable(UEdGraph* InGraph) const
 {
-	return InEditingMode() && InGraph && InGraph->bEditable && IsGraphPanelEnabled(InGraph);
+	return InEditingMode() && !IsGraphReadOnly(InGraph);
+}
+
+bool FBlueprintEditor::IsGraphReadOnly(UEdGraph* InGraph) const
+{
+	return !(InGraph && InGraph->bEditable && IsGraphPanelEnabled(InGraph));
 }
 
 bool FBlueprintEditor::IsGraphPanelEnabled(UEdGraph* InGraph) const
diff --git a/Engine/Source/Editor/Kismet/Public/BlueprintEditor.h b/Engine/Source/Editor/Kismet/Public/BlueprintEditor.h
--- a/Engine/Source/Editor/Kismet/Public/BlueprintEditor.h
+++ b/Engine/Source/Editor/Kismet/Public/BlueprintEditor.h
@@ -265,6 +265,8 @@ public:
 	virtual bool IsEditable(UEdGraph* InGraph) const;
 	/** Determines if the graph's title bar should be the only interactable widget .*/
 	bool IsGraphPanelEnabled(UEdGraph* InGraph) const;
+	/** Determines if the graph is ReadOnly, this differs from editable in that it is never expected to be edited and is in a read-only state */
+	bool IsGraphReadOnly(UEdGraph* InGraph) const;
 
 	/** Used to determine the visibility of the graph's instruction text. */
 	float GetInstructionTextOpacity(UEdGraph* InGraph) const;
diff --git a/Engine/Source/Editor/UnrealEd/Private/GraphEditor.cpp b/Engine/Source/Editor/UnrealEd/Private/GraphEditor.cpp
--- a/Engine/Source/Editor/UnrealEd/Private/GraphEditor.cpp
+++ b/Engine/Source/Editor/UnrealEd/Private/GraphEditor.cpp
@@ -19,6 +19,7 @@ void SGraphEditor::ConstructImplementation( const FArguments& InArgs )
 	// Construct the implementation and make it the contents of this widget.
 	Implementation = GraphEdModule.PRIVATE_MakeGraphEditor( InArgs._AdditionalCommands, 
 		InArgs._IsEditable, 
+		InArgs._DisplayAsReadOnly, 
 		InArgs._IsEmpty,
 		InArgs._Appearance,
 		InArgs._TitleBar,
diff --git a/Engine/Source/Editor/UnrealEd/Public/GraphEditor.h b/Engine/Source/Editor/UnrealEd/Public/GraphEditor.h
--- a/Engine/Source/Editor/UnrealEd/Public/GraphEditor.h
+++ b/Engine/Source/Editor/UnrealEd/Public/GraphEditor.h
@@ -111,6 +111,7 @@ public:
 	SLATE_BEGIN_ARGS(SGraphEditor)
 		: _AdditionalCommands( static_cast<FUICommandList*>(NULL) )
 		, _IsEditable(true)
+		, _DisplayAsReadOnly(false)
 		, _IsEmpty(false)
 		, _GraphToEdit(NULL)
 		, _GraphToDiff(NULL)
@@ -120,6 +121,7 @@ public:
 
 		SLATE_ARGUMENT( TSharedPtr<FUICommandList>, AdditionalCommands )
 		SLATE_ATTRIBUTE( bool, IsEditable )		
+		SLATE_ATTRIBUTE( bool, DisplayAsReadOnly )		
 		SLATE_ATTRIBUTE( bool, IsEmpty )	
 		SLATE_ARGUMENT( TSharedPtr<SWidget>, TitleBar )
 		SLATE_ATTRIBUTE( FGraphAppearanceInfo, Appearance )
