Graph Editor: Allowed double-clicking to interact with splines when they overlap the body of a comment node (double-clicking on the header still activates comment editing) [UE-15209]

diff --git a/Engine/Source/Editor/GraphEditor/Private/SGraphNodeComment.cpp b/Engine/Source/Editor/GraphEditor/Private/SGraphNodeComment.cpp
--- a/Engine/Source/Editor/GraphEditor/Private/SGraphNodeComment.cpp
+++ b/Engine/Source/Editor/GraphEditor/Private/SGraphNodeComment.cpp
@@ -234,8 +234,8 @@ FReply SGraphNodeComment::OnMouseButtonDoubleClick( const FGeometry& InMyGeometr
 	}
 	else
 	{
-		// Otherwise let the base class handle it
-		return SGraphNode::OnMouseButtonDoubleClick(InMyGeometry, InMouseEvent);
+		// Otherwise let the graph handle it, to allow spline interactions to work when they overlap with a comment node
+		return FReply::Unhandled();
 	}
 }
 
