Fixed not being able to rename some graph node types on initial creation
This addresses UE-14873.

diff --git a/Engine/Source/Editor/AnimGraph/Private/AnimationStateMachineSchema.cpp b/Engine/Source/Editor/AnimGraph/Private/AnimationStateMachineSchema.cpp
--- a/Engine/Source/Editor/AnimGraph/Private/AnimationStateMachineSchema.cpp
+++ b/Engine/Source/Editor/AnimGraph/Private/AnimationStateMachineSchema.cpp
@@ -55,7 +55,7 @@ UEdGraphNode* FEdGraphSchemaAction_NewStateNode::PerformAction(class UEdGraph* P
 
 		// set outer to be the graph so it doesn't go away
 		NodeTemplate->Rename(NULL, ParentGraph);
-		ParentGraph->AddNode(NodeTemplate, false, bSelectNewNode);
+		ParentGraph->AddNode(NodeTemplate, true, bSelectNewNode);
 
 		NodeTemplate->CreateNewGuid();
 		NodeTemplate->PostPlacedNewNode();
diff --git a/Engine/Source/Editor/GraphEditor/Private/SGraphPanel.cpp b/Engine/Source/Editor/GraphEditor/Private/SGraphPanel.cpp
--- a/Engine/Source/Editor/GraphEditor/Private/SGraphPanel.cpp
+++ b/Engine/Source/Editor/GraphEditor/Private/SGraphPanel.cpp
@@ -1298,9 +1298,7 @@ void SGraphPanel::OnGraphChanged(const FEdGraphEditAction& EditAction)
 	{
 		if ((EditAction.Graph == GraphObj) &&
 			(EditAction.Nodes.Num() > 0) &&
-			// We do not want to mark it as a UserAddedNode for graphs that do not currently have focus,
-			// this causes each one to want to do the effects and rename, which causes problems.
-			(HasKeyboardFocus() || HasFocusedDescendants()))
+			EditAction.bUserInvoked)
 		{
 			int32 ActionIndex = UserActions.Num();
 			if (EditAction.Action & GRAPHACTION_AddNode)
