Selecting all nodes in a graph using keybindings will no longer assert due to null nodes being present in the graph list.

Added in some protection against nullptr to SGraphEditorImpl::SelectAllNodes

diff --git a/Engine/Source/Editor/GraphEditor/Private/SGraphEditorImpl.cpp b/Engine/Source/Editor/GraphEditor/Private/SGraphEditorImpl.cpp
--- a/Engine/Source/Editor/GraphEditor/Private/SGraphEditorImpl.cpp
+++ b/Engine/Source/Editor/GraphEditor/Private/SGraphEditorImpl.cpp
@@ -136,8 +136,11 @@ void SGraphEditorImpl::SelectAllNodes()
 	for (int32 NodeIndex = 0; NodeIndex < EdGraphObj->Nodes.Num(); ++NodeIndex)
 	{
 		UEdGraphNode* Node = EdGraphObj->Nodes[NodeIndex];
-		ensureMsg(Node->IsValidLowLevel(), TEXT("Node is invalid"));
-		NewSet.Add(Node);
+		if (Node)
+		{
+			ensureMsg(Node->IsValidLowLevel(), TEXT("Node is invalid"));
+			NewSet.Add(Node);
+		}
 	}
 	GraphPanel->SelectionManager.SetSelectionSet(NewSet);
 }
