Protect against crashes when querying information from a graph pin that has been reconstructed (and thus is in the transient package)

diff --git a/Engine/Source/Editor/GraphEditor/Private/SGraphPin.cpp b/Engine/Source/Editor/GraphEditor/Private/SGraphPin.cpp
--- a/Engine/Source/Editor/GraphEditor/Private/SGraphPin.cpp
+++ b/Engine/Source/Editor/GraphEditor/Private/SGraphPin.cpp
@@ -766,7 +766,11 @@ FVector2D SGraphPin::GetNodeOffset() const
 
 FText SGraphPin::GetPinLabel() const
 {
-	return GetPinObj()->GetOwningNode()->GetPinDisplayName(GetPinObj());
+	if (UEdGraphNode* GraphNode = GetPinObj()->GetOwningNodeUnchecked())
+	{
+		return GraphNode->GetPinDisplayName(GetPinObj());
+	}
+	return FText::GetEmpty();
 }
 
 /** @return whether this pin is incoming or outgoing */
@@ -878,24 +882,32 @@ FSlateColor SGraphPin::GetPinColor() const
 	{
 		return FSlateColor(FLinearColor(0.9f,0.2f,0.15f));
 	}
-	const UEdGraphSchema* Schema = GraphPinObj->GetSchema();
-	if(!GetPinObj()->GetOwningNode()->bIsNodeEnabled)
+	if (const UEdGraphSchema* Schema = GraphPinObj->GetSchema())
 	{
-		return Schema->GetPinTypeColor(GraphPinObj->PinType) * FLinearColor(1.0f, 1.0f, 1.0f, 0.5f);
+		if(!GetPinObj()->GetOwningNode()->bIsNodeEnabled)
+		{
+			return Schema->GetPinTypeColor(GraphPinObj->PinType) * FLinearColor(1.0f, 1.0f, 1.0f, 0.5f);
+		}
+
+		return Schema->GetPinTypeColor(GraphPinObj->PinType) * PinColorModifier;
 	}
 
-	return Schema->GetPinTypeColor(GraphPinObj->PinType) * PinColorModifier;
+	return FLinearColor::White;
 }
 
 FSlateColor SGraphPin::GetPinTextColor() const
 {
-	if(!GetPinObj()->GetOwningNode()->bIsNodeEnabled)
-	{
-		return FLinearColor(1.0f, 1.0f, 1.0f, 0.5f);
-	}
-	else if(bUsePinColorForText)
+	// If there is no schema there is no owning node (or basically this is a deleted node)
+	if (UEdGraphNode* GraphNode = GraphPinObj->GetOwningNodeUnchecked())
 	{
-		return GetPinColor();
+		if(!GraphNode->bIsNodeEnabled)
+		{
+			return FLinearColor(1.0f, 1.0f, 1.0f, 0.5f);
+		}
+		else if(bUsePinColorForText)
+		{
+			return GetPinColor();
+		}
 	}
 	return FLinearColor::White;
 }
@@ -905,16 +917,17 @@ const FSlateBrush* SGraphPin::GetPinStatusIcon() const
 {
 	UEdGraphPin* WatchedPin = ((GraphPinObj->Direction == EGPD_Input) && (GraphPinObj->LinkedTo.Num() > 0)) ? GraphPinObj->LinkedTo[0] : GraphPinObj;
 
-	UBlueprint* Blueprint = FBlueprintEditorUtils::FindBlueprintForNodeChecked(WatchedPin->GetOwningNode());
-
-	if (FKismetDebugUtilities::IsPinBeingWatched(Blueprint, WatchedPin) )
-	{
-		return FEditorStyle::GetBrush( TEXT("Graph.WatchedPinIcon_Pinned") );
-	}
-	else
+	if (UEdGraphNode* GraphNode = WatchedPin->GetOwningNodeUnchecked())
 	{
-		return NULL;
+		UBlueprint* Blueprint = FBlueprintEditorUtils::FindBlueprintForNodeChecked(GraphNode);
+
+		if (FKismetDebugUtilities::IsPinBeingWatched(Blueprint, WatchedPin) )
+		{
+			return FEditorStyle::GetBrush( TEXT("Graph.WatchedPinIcon_Pinned") );
+		}
 	}
+
+	return nullptr;
 }
 
 EVisibility SGraphPin::GetPinStatusIconVisibility() const
@@ -922,7 +935,7 @@ EVisibility SGraphPin::GetPinStatusIconVisibility() const
 	UEdGraphPin const* WatchedPin = ((GraphPinObj->Direction == EGPD_Input) && (GraphPinObj->LinkedTo.Num() > 0)) ? GraphPinObj->LinkedTo[0] : GraphPinObj;
 
 	UEdGraphSchema const* Schema = GraphPinObj->GetSchema();
-	return Schema->IsPinBeingWatched(WatchedPin) ? EVisibility::Visible : EVisibility::Collapsed;
+	return Schema && Schema->IsPinBeingWatched(WatchedPin) ? EVisibility::Visible : EVisibility::Collapsed;
 }
 
 FReply SGraphPin::ClickedOnPinStatusIcon()
@@ -939,7 +952,7 @@ EVisibility SGraphPin::GetDefaultValueVisibility() const
 {
 	// First ask schema
 	const UEdGraphSchema* Schema = GraphPinObj->GetSchema();
-	if (Schema->ShouldHidePinDefaultValue(GraphPinObj))
+	if (Schema == nullptr || Schema->ShouldHidePinDefaultValue(GraphPinObj))
 	{
 		return EVisibility::Collapsed;
 	}
