Keyboard navigation of the Blueprint context menu fixed.

Improved usability when nothing is filtered.

diff --git a/Engine/Source/Editor/GraphEditor/Private/SGraphActionMenu.cpp b/Engine/Source/Editor/GraphEditor/Private/SGraphActionMenu.cpp
--- a/Engine/Source/Editor/GraphEditor/Private/SGraphActionMenu.cpp
+++ b/Engine/Source/Editor/GraphEditor/Private/SGraphActionMenu.cpp
@@ -309,6 +309,9 @@ void SGraphActionMenu::Construct( const FArguments& InArgs, bool bIsReadOnly/* =
 		]
 	];
 
+	// When the search box has focus, we want first chance handling of any key down events so we can handle the up/down and escape keys the way we want
+	FilterTextBox->SetOnKeyDownHandler(FOnKeyDown::CreateSP(this, &SGraphActionMenu::OnKeyDown));
+
 	if (!InArgs._ShowFilterTextBox)
 	{
 		FilterTextBox->SetVisibility(EVisibility::Collapsed);
@@ -1224,7 +1227,7 @@ FReply SGraphActionMenu::OnKeyDown( const FGeometry& MyGeometry, const FKeyEvent
 	{
 		return TryToSpawnActiveSuggestion() ? FReply::Handled() : FReply::Unhandled();
 	}
-	else if (FilteredActionNodes.Num() > 0)
+	else if (FilteredActionNodes.Num() > 0 && !FilterTextBox->GetText().IsEmpty())
 	{
 		// Up and down move thru the filtered node list
 		if (KeyEvent.GetKey() == EKeys::Up)
@@ -1242,10 +1245,8 @@ FReply SGraphActionMenu::OnKeyDown( const FGeometry& MyGeometry, const FKeyEvent
 			if( SelectedSuggestion == INDEX_NONE )
 			{
 				SelectedSuggestion = (SelectedSuggestion + SelectionDelta + FilteredRootAction->Children.Num()) % FilteredRootAction->Children.Num();
-				TGuardValue<bool> PreventSelectionFromTriggeringCommit(bIgnoreUIUpdate, true);
-				TreeView->SetSelection(FilteredRootAction->Children[SelectedSuggestion], ESelectInfo::OnKeyPress);
-				TreeView->RequestScrollIntoView(FilteredRootAction->Children[SelectedSuggestion]);
-				return FReply::Handled().SetUserFocus(SharedThis(TreeView.Get()), EFocusCause::WindowActivate);
+				MarkActiveSuggestion();
+				return FReply::Handled();
 			}
 
 			//Move up or down one, wrapping around
@@ -1256,6 +1257,11 @@ FReply SGraphActionMenu::OnKeyDown( const FGeometry& MyGeometry, const FKeyEvent
 			return FReply::Handled();
 		}
 	}
+	else
+	{
+		// When all else fails, it means we haven't filtered the list and we want to handle it as if we were just scrolling through a normal tree view
+		return TreeView->OnKeyDown(FindChildGeometry(MyGeometry, TreeView.ToSharedRef()), KeyEvent);
+	}
 
 	return FReply::Unhandled();
 }
