On Blueprint enum pins, tooltips will appear when hovering over items in the drop down menu.

Fixed localization display issue for enums in the UI.

diff --git a/Engine/Source/Editor/GraphEditor/Private/KismetPins/SGraphPinEnum.cpp b/Engine/Source/Editor/GraphEditor/Private/KismetPins/SGraphPinEnum.cpp
--- a/Engine/Source/Editor/GraphEditor/Private/KismetPins/SGraphPinEnum.cpp
+++ b/Engine/Source/Editor/GraphEditor/Private/KismetPins/SGraphPinEnum.cpp
@@ -13,6 +13,7 @@ void SPinComboBox::Construct( const FArguments& InArgs )
 	OnSelectionChanged = InArgs._OnSelectionChanged;
 	VisibleText = InArgs._VisibleText;
 	OnGetDisplayName = InArgs._OnGetDisplayName;
+	OnGetTooltip = InArgs._OnGetTooltip;
 
 	this->ChildSlot
 	[
@@ -77,6 +78,7 @@ TSharedRef<ITableRow> SPinComboBox::OnGenerateComboWidget( TSharedPtr<int32> InC
 			[
 				SNew(STextBlock)
 				.Text( this, &SPinComboBox::GetRowString, RowIndex )
+				.ToolTipText( this, &SPinComboBox::GetRowTooltip, RowIndex )
 				.Font( FEditorStyle::GetFontStyle( TEXT("PropertyWindow.NormalFont") ) )
 			]
 		];
@@ -99,25 +101,37 @@ TSharedRef<SWidget>	SGraphPinEnum::GetDefaultValueWidget()
 		.VisibleText( this, &SGraphPinEnum::OnGetText )
 		.OnSelectionChanged( this, &SGraphPinEnum::ComboBoxSelectionChanged )
 		.Visibility( this, &SGraphPin::GetDefaultValueVisibility )
-		.OnGetDisplayName(this, &SGraphPinEnum::OnGetFriendlyName);
+		.OnGetDisplayName(this, &SGraphPinEnum::OnGetFriendlyName)
+		.OnGetTooltip(this, &SGraphPinEnum::OnGetTooltip);
 }
 
-FString SGraphPinEnum::OnGetFriendlyName(int32 EnumIndex)
+FText SGraphPinEnum::OnGetFriendlyName(int32 EnumIndex)
 {
 	UEnum* EnumPtr = Cast<UEnum>(GraphPinObj->PinType.PinSubCategoryObject.Get());
 
 	check(EnumPtr);
 	check(EnumIndex < EnumPtr->NumEnums());
 
-	FString EnumValueName = EnumPtr->GetDisplayNameText(EnumIndex).ToString();
-	if (EnumValueName.Len() == 0) 
+	FText EnumValueName = EnumPtr->GetDisplayNameText(EnumIndex);
+	if (EnumValueName.IsEmpty()) 
 	{
-		EnumValueName = EnumPtr->GetEnumName(EnumIndex);
+		EnumValueName = FText::FromString(EnumPtr->GetEnumName(EnumIndex));
 	}
 
 	return EnumValueName;
 }
 
+FText SGraphPinEnum::OnGetTooltip(int32 EnumIndex)
+{
+	UEnum* EnumPtr = Cast<UEnum>(GraphPinObj->PinType.PinSubCategoryObject.Get());
+
+	check(EnumPtr);
+	check(EnumIndex < EnumPtr->NumEnums());
+
+	FText EnumValueTooltip = EnumPtr->GetToolTipText(EnumIndex);
+	return EnumValueTooltip;
+}
+
 void SGraphPinEnum::ComboBoxSelectionChanged( TSharedPtr<int32> NewSelection, ESelectInfo::Type /*SelectInfo*/ )
 {
 	UEnum* EnumPtr = Cast<UEnum>(GraphPinObj->PinType.PinSubCategoryObject.Get());
diff --git a/Engine/Source/Editor/GraphEditor/Private/KismetPins/SGraphPinEnum.h b/Engine/Source/Editor/GraphEditor/Private/KismetPins/SGraphPinEnum.h
--- a/Engine/Source/Editor/GraphEditor/Private/KismetPins/SGraphPinEnum.h
+++ b/Engine/Source/Editor/GraphEditor/Private/KismetPins/SGraphPinEnum.h
@@ -48,7 +48,14 @@ private:
 	 *
 	 * @param EnumIndex	- The index of the enum to return the friendly name for
 	 */
-	FString OnGetFriendlyName(int32 EnumIndex);
+	FText OnGetFriendlyName(int32 EnumIndex);
+
+	/**
+	 * Returns the tooltip of the enum at index EnumIndex
+	 *
+	 * @param EnumIndex	- The index of the enum to return the tooltip for
+	 */
+	FText OnGetTooltip(int32 EnumIndex);
 
 	TSharedPtr<class SPinComboBox>	ComboBox;
 };
diff --git a/Engine/Source/Editor/GraphEditor/Public/SGraphPinComboBox.h b/Engine/Source/Editor/GraphEditor/Public/SGraphPinComboBox.h
--- a/Engine/Source/Editor/GraphEditor/Public/SGraphPinComboBox.h
+++ b/Engine/Source/Editor/GraphEditor/Public/SGraphPinComboBox.h
@@ -3,7 +3,7 @@
 #pragma once
 
 /** A delegate type invoked when a selection changes somewhere. */
-DECLARE_DELEGATE_RetVal_OneParam(FString, FGetComboItemDisplayString, int32);
+DECLARE_DELEGATE_RetVal_OneParam(FText, FGetComboItemDisplayString, int32);
 
 //Class implementation to create combo box
 class GRAPHEDITOR_API SPinComboBox : public SCompoundWidget
@@ -17,6 +17,7 @@ public:
 		SLATE_ATTRIBUTE( FString, VisibleText )
 		SLATE_EVENT( FOnSelectionChanged, OnSelectionChanged )
 		SLATE_EVENT( FGetComboItemDisplayString, OnGetDisplayName )
+		SLATE_EVENT( FGetComboItemDisplayString, OnGetTooltip )
 	SLATE_END_ARGS()
 
 	//Construct combo box using combo button and combo list
@@ -43,7 +44,13 @@ private:
 
 	FText GetRowString(int32 RowIndex) const
 	{
-		return FText::FromString(OnGetDisplayName.Execute(RowIndex));
+		return OnGetDisplayName.Execute(RowIndex);
+	}
+
+	/** Callback for Slate to retrieve the tooltip */
+	FText GetRowTooltip(int32 RowIndex) const
+	{
+		return OnGetTooltip.Execute(RowIndex);
 	}
 
 private:
@@ -57,4 +64,7 @@ private:
 	TWeakPtr<int32> CurrentSelection;
 
 	FGetComboItemDisplayString OnGetDisplayName;
+
+	/** Delegate for finding the tooltip */
+	FGetComboItemDisplayString OnGetTooltip;
 };
