Fixed the BP context menu having double expander arrows in its tree view

diff --git a/Engine/Source/Editor/GraphEditor/Private/SGraphActionMenu.cpp b/Engine/Source/Editor/GraphEditor/Private/SGraphActionMenu.cpp
--- a/Engine/Source/Editor/GraphEditor/Private/SGraphActionMenu.cpp
+++ b/Engine/Source/Editor/GraphEditor/Private/SGraphActionMenu.cpp
@@ -969,7 +969,7 @@ TSharedRef<ITableRow> SGraphActionMenu::MakeWidget( TSharedPtr<FGraphActionNode>
 	}
 
 	TSharedPtr<SHorizontalBox> RowContainer;
-	TableRow->SetContent
+	TableRow->SetRowContent
 	( 
 		SAssignNew(RowContainer, SHorizontalBox)
 	);
