Fixed crash (hard to reproduce), when Save Intermediate Build Products is used.

diff --git a/Engine/Source/Editor/GraphEditor/Private/KismetNodes/SGraphNodeK2Base.cpp b/Engine/Source/Editor/GraphEditor/Private/KismetNodes/SGraphNodeK2Base.cpp
--- a/Engine/Source/Editor/GraphEditor/Private/KismetNodes/SGraphNodeK2Base.cpp
+++ b/Engine/Source/Editor/GraphEditor/Private/KismetNodes/SGraphNodeK2Base.cpp
@@ -301,10 +301,10 @@ FText SGraphNodeK2Base::GetNodeCompactTitle() const
 /** Populate the brushes array with any overlay brushes to render */
 void SGraphNodeK2Base::GetOverlayBrushes(bool bSelected, const FVector2D WidgetSize, TArray<FOverlayBrushInfo>& Brushes) const
 {
-	UBlueprint* OwnerBlueprint = FBlueprintEditorUtils::FindBlueprintForNodeChecked(GraphNode);
+	UBlueprint* OwnerBlueprint = FBlueprintEditorUtils::FindBlueprintForNode(GraphNode);
 
 	// Search for an enabled or disabled breakpoint on this node
-	UBreakpoint* Breakpoint = FKismetDebugUtilities::FindBreakpointForNode(OwnerBlueprint, GraphNode);
+	UBreakpoint* Breakpoint = OwnerBlueprint ? FKismetDebugUtilities::FindBreakpointForNode(OwnerBlueprint, GraphNode) : nullptr;
 	if (Breakpoint != NULL)
 	{
 		FOverlayBrushInfo BreakpointOverlayInfo;
