Fixed the BP palette tree being wrong.

It needed to override SetRowContent as that's what's called on it now.

diff --git a/Engine/Source/Editor/GraphEditor/Private/SGraphActionMenu.cpp b/Engine/Source/Editor/GraphEditor/Private/SGraphActionMenu.cpp
--- a/Engine/Source/Editor/GraphEditor/Private/SGraphActionMenu.cpp
+++ b/Engine/Source/Editor/GraphEditor/Private/SGraphActionMenu.cpp
@@ -77,6 +77,11 @@ public:
 		ContentBorder->SetContent(InContent);
 	}
 
+	virtual void SetRowContent(TSharedRef< SWidget > InContent) override
+	{
+		ContentBorder->SetContent(InContent);
+	}
+
 private:
 	TSharedPtr<SBorder> ContentBorder;
 };
