UE-17211 Clicking connection wires between blueprints nodes causes crash - more robust code

diff --git a/Engine/Source/Editor/GraphEditor/Private/SGraphPin.cpp b/Engine/Source/Editor/GraphEditor/Private/SGraphPin.cpp
--- a/Engine/Source/Editor/GraphEditor/Private/SGraphPin.cpp
+++ b/Engine/Source/Editor/GraphEditor/Private/SGraphPin.cpp
@@ -341,11 +341,13 @@ FReply SGraphPin::OnPinMouseDown( const FGeometry& SenderGeometry, const FPointe
 				Schema->BreakPinLinks(*GraphPinObj, true);
 				return FReply::Handled();
 			}
-			else if (MouseEvent.IsControlDown() && (GraphPinObj->LinkedTo.Num() > 0))
+
+			auto OwnerNodePinned = OwnerNodePtr.Pin();
+			if (MouseEvent.IsControlDown() && (GraphPinObj->LinkedTo.Num() > 0))
 			{
 				// Get a reference to the owning panel widget
-				check(OwnerNodePtr.IsValid());
-				TSharedPtr<SGraphPanel> OwnerPanelPtr = OwnerNodePtr.Pin()->GetOwnerPanel();
+				check(OwnerNodePinned.IsValid());
+				TSharedPtr<SGraphPanel> OwnerPanelPtr = OwnerNodePinned->GetOwnerPanel();
 				check(OwnerPanelPtr.IsValid());
 
 				// Obtain the set of all pins within the panel
@@ -393,7 +395,7 @@ FReply SGraphPin::OnPinMouseDown( const FGeometry& SenderGeometry, const FPointe
 				Schema->BreakPinLinks(*GraphPinObj, true);
 
 				// Check to see if the panel has been invalidated by a graph change notification
-				if(!OwnerPanelPtr->Contains(OwnerNodePtr.Pin()->GetNodeObj()))
+				if (!OwnerPanelPtr->Contains(OwnerNodePinned->GetNodeObj()))
 				{
 					// Force the panel to update. This will cause node & pin widgets to be reinstanced to match any reconstructed node/pin object references.
 					OwnerPanelPtr->Update();
@@ -449,11 +451,16 @@ FReply SGraphPin::OnPinMouseDown( const FGeometry& SenderGeometry, const FPointe
 			}
 			
 			// Start a drag-drop on the pin
+			if (ensure(OwnerNodePinned.IsValid()))
 			{
 				TArray<TSharedRef<SGraphPin>> PinArray;
 				PinArray.Add(SharedThis(this));
 
-				return FReply::Handled().BeginDragDrop(SpawnPinDragEvent(this->OwnerNodePtr.Pin()->GetOwnerPanel().ToSharedRef(), PinArray, MouseEvent.IsShiftDown()));
+				return FReply::Handled().BeginDragDrop(SpawnPinDragEvent(OwnerNodePinned->GetOwnerPanel().ToSharedRef(), PinArray, MouseEvent.IsShiftDown()));
+			}
+			else
+			{
+				return FReply::Unhandled();
 			}
 		}
 		else
