Merging CL 2606013 UE4-Orion/... -> UE4/...
Delegate selector in "Create Event" node now reuses current window in order to work around scaling issues.

diff --git a/Engine/Source/Editor/GraphEditor/Private/KismetNodes/SGraphNodeK2CreateDelegate.cpp b/Engine/Source/Editor/GraphEditor/Private/KismetNodes/SGraphNodeK2CreateDelegate.cpp
--- a/Engine/Source/Editor/GraphEditor/Private/KismetNodes/SGraphNodeK2CreateDelegate.cpp
+++ b/Engine/Source/Editor/GraphEditor/Private/KismetNodes/SGraphNodeK2CreateDelegate.cpp
@@ -123,6 +123,7 @@ void SGraphNodeK2CreateDelegate::CreateBelowWidgetControls(TSharedPtr<SVerticalB
 			}
 
 			TSharedRef<SComboButton> SelectFunctionWidgetRef = SNew(SComboButton)
+				.Method(EPopupMethod::UseCurrentWindow)
 				.ButtonContent()
 				[
 					SNew(STextBlock)
