Blueprints: Moved KismetNodeInfoContext.h to be publically accessible outside of the GraphEditor module

diff --git a/Engine/Source/Editor/GraphEditor/Private/KismetNodes/KismetNodeInfoContext.h b/Engine/Source/Editor/GraphEditor/Private/KismetNodes/KismetNodeInfoContext.h
deleted file mode 100644
diff --git a/Engine/Source/Editor/GraphEditor/Public/KismetNodes/KismetNodeInfoContext.h b/Engine/Source/Editor/GraphEditor/Public/KismetNodes/KismetNodeInfoContext.h
new file mode 100644
--- /dev/null
+++ b/Engine/Source/Editor/GraphEditor/Public/KismetNodes/KismetNodeInfoContext.h
@@ -0,0 +1,47 @@
+// Copyright 1998-2015 Epic Games, Inc. All Rights Reserved.
+
+#pragma once
+
+// Context used to aid debugging displays for nodes
+struct FKismetNodeInfoContext : public FNodeInfoContext
+{
+public:
+	struct FObjectUUIDPair
+	{
+	public:
+		class UObject* Object;
+		int32 UUID;
+	public:
+		FObjectUUIDPair(class UObject* InObject, int32 InUUID) : Object(InObject), UUID(InUUID) {}
+
+		FString GetDisplayName() const
+		{
+			if (AActor* Actor = Cast<AActor>(Object))
+			{
+				return Actor->GetActorLabel();
+			}
+			else
+			{
+				return (Object != NULL) ? Object->GetName() : TEXT("(null)");
+			}
+		}
+	};
+
+	TMap<class UEdGraphNode*, TArray<FObjectUUIDPair> > NodesWithActiveLatentActions;
+
+	// Set of pins with watches
+	TSet<UEdGraphPin*> WatchedPinSet;
+
+	// Set of nodes with a pin that is being watched
+	TSet<UEdGraphNode*> WatchedNodeSet;
+
+	// Source blueprint for the graph
+	UBlueprint* SourceBlueprint;
+
+	// Object being debugged for the graph
+	UObject* ActiveObjectBeingDebugged;
+
+public:
+	FKismetNodeInfoContext(class UEdGraph* SourceGraph);
+};
+
