Graph Editor: Allowed graph nodes a chance to create a custom SGraphNode visual representation without having to manually register a factory (override the new method CreateVisualWidget)

diff --git a/Engine/Source/Editor/GraphEditor/Private/NodeFactory.cpp b/Engine/Source/Editor/GraphEditor/Private/NodeFactory.cpp
--- a/Engine/Source/Editor/GraphEditor/Private/NodeFactory.cpp
+++ b/Engine/Source/Editor/GraphEditor/Private/NodeFactory.cpp
@@ -82,6 +82,15 @@ TSharedPtr<SGraphNode> FNodeFactory::CreateNodeWidget(UEdGraphNode* InNode)
 {
 	check(InNode != NULL);
 
+	// First give a shot to the node itself
+	{
+		TSharedPtr<SGraphNode> NodeCreatedResult = InNode->CreateVisualWidget();
+		if (NodeCreatedResult.IsValid())
+		{
+			return NodeCreatedResult;
+		}
+	}
+
 	// First give a shot to the registered node factories
 	for (auto FactoryIt = FEdGraphUtilities::VisualNodeFactories.CreateIterator(); FactoryIt; ++FactoryIt)
 	{
diff --git a/Engine/Source/Runtime/Engine/Classes/EdGraph/EdGraphNode.h b/Engine/Source/Runtime/Engine/Classes/EdGraph/EdGraphNode.h
--- a/Engine/Source/Runtime/Engine/Classes/EdGraph/EdGraphNode.h
+++ b/Engine/Source/Runtime/Engine/Classes/EdGraph/EdGraphNode.h
@@ -439,6 +439,10 @@ public:
 
 	/** If the node has a subgraph, should they be merged into the main graph? */
 	virtual bool ShouldMergeChildGraphs() const { return true; }
+
+	/** Create a visual widget to represent this node in a graph editor or graph panel.  If not implemented, the default node factory will be used. */
+	virtual TSharedPtr<SGraphNode> CreateVisualWidget() { return TSharedPtr<SGraphNode>(); }
+
 #endif // WITH_EDITOR
 
 };
