Tweaked the tooltip for "Expand Nodes" menu option to better describe what the action will do.

diff --git a/Engine/Source/Editor/GraphEditor/Private/GraphEditorActions.cpp b/Engine/Source/Editor/GraphEditor/Private/GraphEditorActions.cpp
--- a/Engine/Source/Editor/GraphEditor/Private/GraphEditorActions.cpp
+++ b/Engine/Source/Editor/GraphEditor/Private/GraphEditorActions.cpp
@@ -38,7 +38,7 @@ void FGraphEditorCommandsImpl::RegisterCommands()
 	UI_COMMAND( CollapseNodes, "Collapse Nodes", "Collapses selected nodes into a single node", EUserInterfaceActionType::Button, FInputChord() )
 	UI_COMMAND( PromoteSelectionToFunction, "Promote to Function", "Promotes selected collapsed graphs to functions.", EUserInterfaceActionType::Button, FInputChord() )
 	UI_COMMAND( PromoteSelectionToMacro, "Promote to Macro", "Promotes selected collapsed graphs to macros.", EUserInterfaceActionType::Button, FInputChord() )
-	UI_COMMAND( ExpandNodes, "Expand Node", "Expands composite nodes", EUserInterfaceActionType::Button, FInputChord() )
+	UI_COMMAND( ExpandNodes, "Expand Node", "Expands the node's internal graph into the current graph and removes this node.", EUserInterfaceActionType::Button, FInputChord() )
 	UI_COMMAND( CollapseSelectionToFunction, "Collapse to Function", "Collapses selected nodes into a single function node.", EUserInterfaceActionType::Button, FInputChord() )
 	UI_COMMAND( CollapseSelectionToMacro, "Collapse to Macro", "Collapses selected nodes into a single macro node.", EUserInterfaceActionType::Button, FInputChord() )
 
