Blueprints: Changed display of FRotator values to be XYZ and in that order (RPY) instead of PYR

diff --git a/Engine/Config/BaseEngine.ini b/Engine/Config/BaseEngine.ini
--- a/Engine/Config/BaseEngine.ini
+++ b/Engine/Config/BaseEngine.ini
@@ -382,6 +382,10 @@ bUseStreamingPause=false
 +ActiveClassRedirects=(OldClassName="AnimNode_TwoBoneIK",NewClassName="/Script/AnimGraphRuntime.AnimNode_TwoBoneIK")
 +ActiveClassRedirects=(OldClassName="AnimNode_WheelHandler",NewClassName="/Script/AnimGraphRuntime.AnimNode_WheelHandler")
 
+; Renamed all user-exposed locations for rotator fields to display as XYZ (RPY) instead of YPR
++K2FieldRedirects=(OldFieldName="KismetMathLibrary.BreakRot", NewFieldName="KismetMathLibrary.BreakRotator")
++K2FieldRedirects=(OldFieldName="KismetMathLibrary.MakeRot", NewFieldName="KismetMathLibrary.MakeRotator")
+
 ;AIModule redirections
 +ActiveClassRedirects=(OldClassName="AIController",NewClassName="/Script/AIModule.AIController")
 +ActiveClassRedirects=(OldClassName="AIResourceInterface",NewClassName="/Script/AIModule.AIResourceInterface")
diff --git a/Engine/Source/Editor/BlueprintGraph/Private/EdGraphSchema_K2.cpp b/Engine/Source/Editor/BlueprintGraph/Private/EdGraphSchema_K2.cpp
--- a/Engine/Source/Editor/BlueprintGraph/Private/EdGraphSchema_K2.cpp
+++ b/Engine/Source/Editor/BlueprintGraph/Private/EdGraphSchema_K2.cpp
@@ -4973,7 +4973,7 @@ void UEdGraphSchema_K2::BackwardCompatibilityNodeConversion(UEdGraph* Graph, boo
 
 				if(OldMakeStructNode->StructType->GetName() == TEXT("Rotator"))
 				{
-					MakeNodeFunction = UKismetMathLibrary::StaticClass()->FindFunctionByName(GET_FUNCTION_NAME_CHECKED(UKismetMathLibrary, MakeRot));
+					MakeNodeFunction = UKismetMathLibrary::StaticClass()->FindFunctionByName(GET_FUNCTION_NAME_CHECKED(UKismetMathLibrary, MakeRotator));
 					OldPinToNewPinMap.Add(TEXT("Rotator"), TEXT("ReturnValue"));
 				}
 				else if(OldMakeStructNode->StructType->GetName() == TEXT("Vector"))
@@ -5054,7 +5054,7 @@ void UEdGraphSchema_K2::BackwardCompatibilityNodeConversion(UEdGraph* Graph, boo
 
 				if(OldBreakStructNode->StructType->GetName() == TEXT("Rotator"))
 				{
-					BreakNodeFunction = UKismetMathLibrary::StaticClass()->FindFunctionByName(GET_FUNCTION_NAME_CHECKED(UKismetMathLibrary, BreakRot));
+					BreakNodeFunction = UKismetMathLibrary::StaticClass()->FindFunctionByName(GET_FUNCTION_NAME_CHECKED(UKismetMathLibrary, BreakRotator));
 					OldPinToNewPinMap.Add(TEXT("Rotator"), TEXT("InRot"));
 				}
 				else if(OldBreakStructNode->StructType->GetName() == TEXT("Vector"))
diff --git a/Engine/Source/Editor/GraphEditor/Private/KismetPins/SGraphPinVector.cpp b/Engine/Source/Editor/GraphEditor/Private/KismetPins/SGraphPinVector.cpp
--- a/Engine/Source/Editor/GraphEditor/Private/KismetPins/SGraphPinVector.cpp
+++ b/Engine/Source/Editor/GraphEditor/Private/KismetPins/SGraphPinVector.cpp
@@ -27,6 +27,8 @@ public:
 	void Construct( const FArguments& InArgs, const bool bInIsRotator )
 	{
 		bIsRotator = bInIsRotator;
+		const bool bUseRPY = false;
+
 		VisibleText_0 = InArgs._VisibleText_0;
 		VisibleText_1 = InArgs._VisibleText_1;
 		VisibleText_2 = InArgs._VisibleText_2;
@@ -50,14 +52,14 @@ public:
 					[
 						SNew( STextBlock )
 						.Font( FEditorStyle::GetFontStyle( "Graph.VectorEditableTextBox" ) )
-						.Text( bIsRotator ? LOCTEXT("VectorNodeRollValueLabel", "R") : LOCTEXT("VectorNodeXAxisValueLabel", "X") )
+						.Text( bIsRotator && bUseRPY ? LOCTEXT("VectorNodeRollValueLabel", "R") : LOCTEXT("VectorNodeXAxisValueLabel", "X") )
 						.ColorAndOpacity( LabelClr )
 					]
 					.Value( this, &SVectorTextBox::GetTypeInValue_0 )
 					.OnValueCommitted( InArgs._OnFloatCommitted_Box_0 )
 					.Font( FEditorStyle::GetFontStyle( "Graph.VectorEditableTextBox" ) )
 					.UndeterminedString( LOCTEXT("MultipleValues", "Multiple Values") )
-					.ToolTipText( bIsRotator ? LOCTEXT("VectorNodeRollValueLabel_ToolTip", "Roll value") : LOCTEXT("VectorNodeXAxisValueLabel_ToolTip", "X value"))
+					.ToolTipText( bIsRotator ? LOCTEXT("VectorNodeRollValueLabel_ToolTip", "Roll value (around X)") : LOCTEXT("VectorNodeXAxisValueLabel_ToolTip", "X value"))
 					.EditableTextBoxStyle( &FEditorStyle::GetWidgetStyle<FEditableTextBoxStyle>( "Graph.VectorEditableTextBox" ))
 					.BorderForegroundColor( FLinearColor::White )
 					.BorderBackgroundColor( FLinearColor::White )
@@ -72,14 +74,14 @@ public:
 					[
 						SNew( STextBlock )
 						.Font( FEditorStyle::GetFontStyle( "Graph.VectorEditableTextBox" ) )
-						.Text( bIsRotator ? LOCTEXT("VectorNodePitchValueLabel", "P") : LOCTEXT("VectorNodeYAxisValueLabel", "Y") )
+						.Text( bIsRotator && bUseRPY ? LOCTEXT("VectorNodePitchValueLabel", "P") : LOCTEXT("VectorNodeYAxisValueLabel", "Y") )
 						.ColorAndOpacity( LabelClr )
 					]
 					.Value( this, &SVectorTextBox::GetTypeInValue_1 )
 					.OnValueCommitted( InArgs._OnFloatCommitted_Box_1 )
 					.Font( FEditorStyle::GetFontStyle( "Graph.VectorEditableTextBox" ) )
 					.UndeterminedString( LOCTEXT("MultipleValues", "Multiple Values") )
-					.ToolTipText(bIsRotator ? LOCTEXT("VectorNodePitchValueLabel_ToolTip", "Pitch value") : LOCTEXT("VectorNodeYAxisValueLabel_ToolTip", "Y value"))
+					.ToolTipText(bIsRotator ? LOCTEXT("VectorNodePitchValueLabel_ToolTip", "Pitch value (around Y)") : LOCTEXT("VectorNodeYAxisValueLabel_ToolTip", "Y value"))
 					.EditableTextBoxStyle( &FEditorStyle::GetWidgetStyle<FEditableTextBoxStyle>( "Graph.VectorEditableTextBox" ))
 					.BorderForegroundColor( FLinearColor::White )
 					.BorderBackgroundColor( FLinearColor::White )
@@ -94,14 +96,14 @@ public:
 					[
 						SNew( STextBlock )
 						.Font( FEditorStyle::GetFontStyle( "Graph.VectorEditableTextBox" ) )
-						.Text( bIsRotator ? LOCTEXT("VectorNodeYawValueLabel", "Y") : LOCTEXT("VectorNodeZAxisValueLabel", "Z") )
+						.Text( bIsRotator && bUseRPY ? LOCTEXT("VectorNodeYawValueLabel", "Y") : LOCTEXT("VectorNodeZAxisValueLabel", "Z") )
 						.ColorAndOpacity( LabelClr )
 					]
 					.Value( this, &SVectorTextBox::GetTypeInValue_2 )
 					.OnValueCommitted( InArgs._OnFloatCommitted_Box_2 )
 					.Font( FEditorStyle::GetFontStyle( "Graph.VectorEditableTextBox" ) )
 					.UndeterminedString( LOCTEXT("MultipleValues", "Multiple Values") )
-					.ToolTipText(bIsRotator ? LOCTEXT("VectorNodeYawValueLabel_Tooltip", "Yaw value") : LOCTEXT("VectorNodeZAxisValueLabel_ToolTip", "Z value"))
+					.ToolTipText(bIsRotator ? LOCTEXT("VectorNodeYawValueLabel_Tooltip", "Yaw value (around Z)") : LOCTEXT("VectorNodeZAxisValueLabel_ToolTip", "Z value"))
 					.EditableTextBoxStyle( &FEditorStyle::GetWidgetStyle<FEditableTextBoxStyle>( "Graph.VectorEditableTextBox" ))
 					.BorderForegroundColor( FLinearColor::White )
 					.BorderBackgroundColor( FLinearColor::White )
diff --git a/Engine/Source/Runtime/CoreUObject/Classes/Object.h b/Engine/Source/Runtime/CoreUObject/Classes/Object.h
--- a/Engine/Source/Runtime/CoreUObject/Classes/Object.h
+++ b/Engine/Source/Runtime/CoreUObject/Classes/Object.h
@@ -377,7 +377,7 @@ struct FPlane : public FVector
  * An orthogonal rotation in 3d space.
  * The full C++ class is located here: Engine\Source\Runtime\Core\Public\Math\Rotator.h
  */
-USTRUCT(immutable, noexport, BlueprintType, meta=(HasNativeMake="Engine.KismetMathLibrary.MakeRot", HasNativeBreak="Engine.KismetMathLibrary.BreakRot"))
+USTRUCT(immutable, noexport, BlueprintType, meta=(HasNativeMake="Engine.KismetMathLibrary.MakeRotator", HasNativeBreak="Engine.KismetMathLibrary.BreakRotator"))
 struct FRotator
 {
 	UPROPERTY(EditAnywhere, BlueprintReadWrite, Category=Rotator, SaveGame)
diff --git a/Engine/Source/Runtime/Engine/Classes/Kismet/KismetMathLibrary.h b/Engine/Source/Runtime/Engine/Classes/Kismet/KismetMathLibrary.h
--- a/Engine/Source/Runtime/Engine/Classes/Kismet/KismetMathLibrary.h
+++ b/Engine/Source/Runtime/Engine/Classes/Kismet/KismetMathLibrary.h
@@ -1214,10 +1214,19 @@ class ENGINE_API UKismetMathLibrary : public UBlueprintFunctionLibrary
 	UFUNCTION(BlueprintPure, Category="Math|Vector", meta=(Keywords="rotation rotate"))
 	static FVector GetUpVector(FRotator InRot);
 
-	/** Makes a rotator {Pitch, Yaw, Roll} from rotation values supplied in degrees */
+	/** Makes a rotator {Roll, Pitch, Yaw} from rotation values supplied in degrees */
 	UFUNCTION(BlueprintPure, Category="Math|Rotator", meta=(Keywords="construct build rotation rotate rotator makerotator", NativeMakeFunc))
-	static FRotator MakeRot(float Pitch, float Yaw, float Roll);
-	
+	static FRotator MakeRotator(
+		UPARAM(DisplayName="X (Roll)") float Roll,	
+		UPARAM(DisplayName="Y (Pitch)") float Pitch,
+		UPARAM(DisplayName="Z (Yaw)") float Yaw);
+
+	DEPRECATED(4.9, "Method renamed to MakeRotator as the parameter order changed")
+	static inline FRotator MakeRot(float Pitch, float Yaw, float Roll)
+	{
+		return MakeRotator(Roll, Pitch, Yaw);
+	}
+
 	/** Find a rotation for an object at Start location to point at Target location. */
 	UFUNCTION(BlueprintPure, Category="Math|Rotator", meta=(Keywords="rotation rotate"))
 	static FRotator FindLookAtRotation(const FVector& Start, const FVector& Target);
@@ -1258,9 +1267,18 @@ class ENGINE_API UKismetMathLibrary : public UBlueprintFunctionLibrary
 	UFUNCTION(BlueprintPure, Category="Math|Rotator", meta=(Keywords="construct build rotation rotate rotator makerotator"))
 	static FRotator MakeRotFromZY(const FVector& Z, const FVector& Y);
 
-	/** Breaks apart a rotator into Pitch, Yaw, Roll angles in degrees*/
-	UFUNCTION(BlueprintPure, Category="Math|Rotator", meta=(Keywords="rotation rotate rotator breakrotator", NativeBreakFunc))
-	static void BreakRot(FRotator InRot, float& Pitch, float& Yaw, float& Roll);
+	/** Breaks apart a rotator into {Roll, Pitch, Yaw} angles in degrees */
+	UFUNCTION(BlueprintPure, Category = "Math|Rotator", meta = (Keywords = "rotation rotate rotator breakrotator", NativeBreakFunc))
+	static void BreakRotator(FRotator InRot,
+		UPARAM(DisplayName="X (Roll)") float& Roll,
+		UPARAM(DisplayName="Y (Pitch)") float& Pitch,
+		UPARAM(DisplayName="Z (Yaw)") float& Yaw);
+
+	DEPRECATED(4.9, "Method renamed to BreakRotator as the parameter order changed")
+	static inline void BreakRot(FRotator InRot, float& Pitch, float& Yaw, float& Roll)
+	{
+		BreakRotator(InRot, Roll, Pitch, Yaw);
+	}
 
 	/** Breaks apart a rotator into its component axes */
 	UFUNCTION(BlueprintPure, Category="Math|Rotator", meta=(Keywords="rotation rotate rotator breakrotator"))
diff --git a/Engine/Source/Runtime/Engine/Private/KismetMathLibrary.cpp b/Engine/Source/Runtime/Engine/Private/KismetMathLibrary.cpp
--- a/Engine/Source/Runtime/Engine/Private/KismetMathLibrary.cpp
+++ b/Engine/Source/Runtime/Engine/Private/KismetMathLibrary.cpp
@@ -1804,7 +1804,7 @@ FVector UKismetMathLibrary::GetUpVector(FRotator InRot)
 	return FRotationMatrix(InRot).GetScaledAxis( EAxis::Z );
 }
 
-FRotator UKismetMathLibrary::MakeRot(float Pitch, float Yaw, float Roll)
+FRotator UKismetMathLibrary::MakeRotator(float Roll, float Pitch, float Yaw)
 {
 	return FRotator(Pitch,Yaw,Roll);
 }
@@ -1860,7 +1860,7 @@ FRotator UKismetMathLibrary::MakeRotFromZY(const FVector& Z, const FVector& Y)
 }
 
 
-void UKismetMathLibrary::BreakRot(FRotator InRot, float& Pitch, float& Yaw, float& Roll)
+void UKismetMathLibrary::BreakRotator(FRotator InRot, float& Roll, float& Pitch, float& Yaw)
 {
 	Pitch = InRot.Pitch;
 	Yaw = InRot.Yaw;
