UE-11709 - Cannot hide CustomUV pins on Break/Make MaterialAttributes nodes like you can with MaterialResults node

Custom UV pins in the material graph are hidden when they are on the results node.
They cannot be hidden in material functions however.

diff --git a/Engine/Source/Editor/GraphEditor/Private/MaterialNodes/SGraphNodeMaterialBase.cpp b/Engine/Source/Editor/GraphEditor/Private/MaterialNodes/SGraphNodeMaterialBase.cpp
--- a/Engine/Source/Editor/GraphEditor/Private/MaterialNodes/SGraphNodeMaterialBase.cpp
+++ b/Engine/Source/Editor/GraphEditor/Private/MaterialNodes/SGraphNodeMaterialBase.cpp
@@ -9,6 +9,8 @@
 #include "TutorialMetaData.h"
 #include "CanvasTypes.h"
 #include "CanvasItem.h"
+#include "Materials/MaterialExpressionMakeMaterialAttributes.h"
+#include "Materials/MaterialExpressionBreakMaterialAttributes.h"
 
 
 
@@ -237,7 +239,20 @@ void SGraphNodeMaterialBase::CreatePinWidgets()
 
 		const bool bPinHasConections = CurPin->LinkedTo.Num() > 0;
 
-		const bool bPinDesiresToBeHidden = CurPin->bHidden || (bHideNoConnectionPins && !bPinHasConections); 
+		bool bPinDesiresToBeHidden = CurPin->bHidden || (bHideNoConnectionPins && !bPinHasConections); 
+
+		UMaterialGraph* MaterialGraph = CastChecked<UMaterialGraph>(GraphNode->GetGraph());
+		if (MaterialNode && MaterialNode->MaterialExpression && MaterialGraph->MaterialFunction == nullptr)
+		{
+			bool bIsAMakeAttrNode = MaterialNode->MaterialExpression->IsA(UMaterialExpressionMakeMaterialAttributes::StaticClass());
+			bool bIsABreakAttrNode = MaterialNode->MaterialExpression->IsA(UMaterialExpressionBreakMaterialAttributes::StaticClass());
+		
+			if ((bIsABreakAttrNode && CurPin->Direction == EGPD_Output) || (bIsAMakeAttrNode && CurPin->Direction == EGPD_Input))
+			{
+				int32 ResultNodeIndex = bIsAMakeAttrNode ? MaterialNode->GetInputIndex(CurPin) : MaterialNode->GetOutputIndex(CurPin);
+				bPinDesiresToBeHidden |= !MaterialGraph->MaterialInputs[ResultNodeIndex].IsVisiblePin(MaterialGraph->Material);
+			}
+		}
 
 		if (!bPinDesiresToBeHidden)
 		{
diff --git a/Engine/Source/Editor/UnrealEd/Classes/MaterialGraph/MaterialGraphNode.h b/Engine/Source/Editor/UnrealEd/Classes/MaterialGraph/MaterialGraphNode.h
--- a/Engine/Source/Editor/UnrealEd/Classes/MaterialGraph/MaterialGraphNode.h
+++ b/Engine/Source/Editor/UnrealEd/Classes/MaterialGraph/MaterialGraphNode.h
@@ -39,7 +39,7 @@ public:
 	UNREALED_API void RecreateAndLinkNode();
 
 	/** Get the Material Expression output index from an output pin */
-	int32 GetOutputIndex(const UEdGraphPin* OutputPin);
+	UNREALED_API int32 GetOutputIndex(const UEdGraphPin* OutputPin);
 	/** Get the Material value type of an output pin */
 	uint32 GetOutputType(const UEdGraphPin* OutputPin);
 
@@ -68,7 +68,7 @@ public:
 	// UMaterialGraphNode_Base interface
 	virtual void CreateInputPins() override;
 	virtual void CreateOutputPins() override;
-	virtual int32 GetInputIndex(const UEdGraphPin* InputPin) const override;
+	virtual UNREALED_API int32 GetInputIndex(const UEdGraphPin* InputPin) const override;
 	virtual uint32 GetInputType(const UEdGraphPin* InputPin) const override;
 	// End of UMaterialGraphNode_Base interface
 
diff --git a/Engine/Source/Editor/UnrealEd/Private/MaterialGraph.cpp b/Engine/Source/Editor/UnrealEd/Private/MaterialGraph.cpp
--- a/Engine/Source/Editor/UnrealEd/Private/MaterialGraph.cpp
+++ b/Engine/Source/Editor/UnrealEd/Private/MaterialGraph.cpp
@@ -54,8 +54,10 @@ void UMaterialGraph::RebuildGraph()
 			MaterialInputs.Add( FMaterialInputInfo( FText::FromString(FString::Printf(TEXT("Customized UV%u"), UVIndex)), (EMaterialProperty)(MP_CustomizedUVs0 + UVIndex)) );
 		}
 
-		MaterialInputs.Add( FMaterialInputInfo( LOCTEXT("MaterialAttributes", "Material Attributes"), MP_MaterialAttributes) );
-		MaterialInputs.Add( FMaterialInputInfo( LOCTEXT("PixelDepthOffset", "Pixel Depth Offset"), MP_PixelDepthOffset ) );
+		MaterialInputs.Add(FMaterialInputInfo(LOCTEXT("PixelDepthOffset", "Pixel Depth Offset"), MP_PixelDepthOffset));
+
+		//^^^ New material properties go above here. ^^^^
+		MaterialInputs.Add(FMaterialInputInfo(LOCTEXT("MaterialAttributes", "Material Attributes"), MP_MaterialAttributes));
 
 		// Add Root Node
 		FGraphNodeCreator<UMaterialGraphNode_Root> NodeCreator(*this);
diff --git a/Engine/Source/Runtime/Engine/Classes/Materials/MaterialExpressionMakeMaterialAttributes.h b/Engine/Source/Runtime/Engine/Classes/Materials/MaterialExpressionMakeMaterialAttributes.h
--- a/Engine/Source/Runtime/Engine/Classes/Materials/MaterialExpressionMakeMaterialAttributes.h
+++ b/Engine/Source/Runtime/Engine/Classes/Materials/MaterialExpressionMakeMaterialAttributes.h
@@ -5,7 +5,7 @@
 #include "Materials/MaterialExpression.h"
 #include "MaterialExpressionMakeMaterialAttributes.generated.h"
 
-UCLASS(collapsecategories, hidecategories=Object)
+UCLASS(collapsecategories, hidecategories = Object, MinimalAPI)
 class UMaterialExpressionMakeMaterialAttributes : public UMaterialExpression
 {
 	GENERATED_UCLASS_BODY()
diff --git a/Engine/Source/Runtime/Engine/Public/SceneTypes.h b/Engine/Source/Runtime/Engine/Public/SceneTypes.h
--- a/Engine/Source/Runtime/Engine/Public/SceneTypes.h
+++ b/Engine/Source/Runtime/Engine/Public/SceneTypes.h
@@ -152,8 +152,10 @@ enum EMaterialProperty
 	MP_CustomizedUVs5,
 	MP_CustomizedUVs6,
 	MP_CustomizedUVs7,
-	MP_MaterialAttributes,
 	MP_PixelDepthOffset,
+
+	//^^^ New material properties go above here ^^^^
+	MP_MaterialAttributes,
 	MP_MAX,
 };
 
