Fixed all pins on make/break material attributes nodes being invisible when used alongside bUseMaterialAttriubtes.

diff --git a/Engine/Source/Editor/GraphEditor/Private/MaterialNodes/SGraphNodeMaterialBase.cpp b/Engine/Source/Editor/GraphEditor/Private/MaterialNodes/SGraphNodeMaterialBase.cpp
--- a/Engine/Source/Editor/GraphEditor/Private/MaterialNodes/SGraphNodeMaterialBase.cpp
+++ b/Engine/Source/Editor/GraphEditor/Private/MaterialNodes/SGraphNodeMaterialBase.cpp
@@ -250,7 +250,7 @@ void SGraphNodeMaterialBase::CreatePinWidgets()
 			if ((bIsABreakAttrNode && CurPin->Direction == EGPD_Output) || (bIsAMakeAttrNode && CurPin->Direction == EGPD_Input))
 			{
 				int32 ResultNodeIndex = bIsAMakeAttrNode ? MaterialNode->GetInputIndex(CurPin) : MaterialNode->GetOutputIndex(CurPin);
-				bPinDesiresToBeHidden |= !MaterialGraph->MaterialInputs[ResultNodeIndex].IsVisiblePin(MaterialGraph->Material);
+				bPinDesiresToBeHidden |= !MaterialGraph->MaterialInputs[ResultNodeIndex].IsVisiblePin(MaterialGraph->Material, true);
 			}
 		}
 
diff --git a/Engine/Source/Editor/UnrealEd/Classes/MaterialGraph/MaterialGraph.h b/Engine/Source/Editor/UnrealEd/Classes/MaterialGraph/MaterialGraph.h
--- a/Engine/Source/Editor/UnrealEd/Classes/MaterialGraph/MaterialGraph.h
+++ b/Engine/Source/Editor/UnrealEd/Classes/MaterialGraph/MaterialGraph.h
@@ -38,9 +38,9 @@ struct FMaterialInputInfo
 		return *Ret;
 	}
 
-	bool IsVisiblePin(const UMaterial* Material) const
+	bool IsVisiblePin(const UMaterial* Material, bool bIgnoreMaterialAttriubtes = false) const
 	{
-		if(Material->bUseMaterialAttributes)
+		if(Material->bUseMaterialAttributes && !bIgnoreMaterialAttriubtes)
 		{
 			return Property == MP_MaterialAttributes;
 		}
