[UE-16784] Various optimizations to graph action menu creation and searching/filtering

diff --git a/Engine/Source/Editor/GraphEditor/Private/GraphActionNode.cpp b/Engine/Source/Editor/GraphEditor/Private/GraphActionNode.cpp
--- a/Engine/Source/Editor/GraphEditor/Private/GraphActionNode.cpp
+++ b/Engine/Source/Editor/GraphEditor/Private/GraphActionNode.cpp
@@ -450,45 +450,47 @@ TSharedPtr<FGraphActionNode> FGraphActionNode::NewGroupDividerNode(TWeakPtr<FGra
 //------------------------------------------------------------------------------
 void FGraphActionNode::AddChildRecursively(TArray<FString>& CategoryStack, TSharedPtr<FGraphActionNode> NodeToAdd)
 {
-	TSharedPtr<FGraphActionNode> FoundSectionNode;
-	for ( TSharedPtr<FGraphActionNode> const& ChildNode : Children )
+	if (NodeToAdd->SectionID != INVALID_SECTION_ID)
 	{
-		if ( NodeToAdd->SectionID == ChildNode->SectionID && ChildNode->IsSectionHeadingNode() )
+		TSharedPtr<FGraphActionNode> FoundSectionNode;
+		for ( TSharedPtr<FGraphActionNode> const& ChildNode : Children )
 		{
-			FoundSectionNode = ChildNode;
-			break;
+			if ( NodeToAdd->SectionID == ChildNode->SectionID && ChildNode->IsSectionHeadingNode() )
+			{
+				FoundSectionNode = ChildNode;
+				break;
+			}
 		}
-	}
 
-	if ( FoundSectionNode.IsValid() )
-	{
-		FoundSectionNode->AddChildRecursively(CategoryStack, NodeToAdd);
+		if ( FoundSectionNode.IsValid() )
+		{
+			FoundSectionNode->AddChildRecursively(CategoryStack, NodeToAdd);
+			return;
+		}
 	}
-	else
+
+	if ( CategoryStack.Num() > 0 )
 	{
-		if ( CategoryStack.Num() > 0 )
-		{
-			FString CategorySection = CategoryStack[0];
-			CategoryStack.RemoveAt(0, 1);
+		FString CategorySection = CategoryStack[0];
+		CategoryStack.RemoveAt(0, 1);
 
-			// make sure we don't already have a child that this can nest under
-			TSharedPtr<FGraphActionNode> ExistingNode = FindMatchingParent(CategorySection, NodeToAdd);
-			if ( ExistingNode.IsValid() )
-			{
-				ExistingNode->AddChildRecursively(CategoryStack, NodeToAdd);
-			}
-			else
-			{
-				TSharedPtr<FGraphActionNode> CategoryNode = NewCategoryNode(CategorySection, NodeToAdd->Grouping, NodeToAdd->SectionID);
-				InsertChild(CategoryNode);
-				CategoryNode->AddChildRecursively(CategoryStack, NodeToAdd);
-			}
+		// make sure we don't already have a child that this can nest under
+		TSharedPtr<FGraphActionNode> ExistingNode = FindMatchingParent(CategorySection, NodeToAdd);
+		if ( ExistingNode.IsValid() )
+		{
+			ExistingNode->AddChildRecursively(CategoryStack, NodeToAdd);
 		}
 		else
 		{
-			InsertChild(NodeToAdd);
+			TSharedPtr<FGraphActionNode> CategoryNode = NewCategoryNode(CategorySection, NodeToAdd->Grouping, NodeToAdd->SectionID);
+			InsertChild(CategoryNode);
+			CategoryNode->AddChildRecursively(CategoryStack, NodeToAdd);
 		}
 	}
+	else
+	{
+		InsertChild(NodeToAdd);
+	}
 }
 
 //------------------------------------------------------------------------------
diff --git a/Engine/Source/Editor/GraphEditor/Private/SGraphActionMenu.cpp b/Engine/Source/Editor/GraphEditor/Private/SGraphActionMenu.cpp
--- a/Engine/Source/Editor/GraphEditor/Private/SGraphActionMenu.cpp
+++ b/Engine/Source/Editor/GraphEditor/Private/SGraphActionMenu.cpp
@@ -652,7 +652,11 @@ void SGraphActionMenu::GenerateFilteredItems(bool bPreserveExpansion)
 	// Tokenize the search box text into a set of terms; all of them must be present to pass the filter
 	TArray<FString> FilterTerms;
 	TrimmedFilterString.ParseIntoArray(FilterTerms, TEXT(" "), true);
-	
+	for (auto& String : FilterTerms)
+	{
+		String = String.ToLower();
+	}
+
 	// Generate a list of sanitized versions of the strings
 	TArray<FString> SanitizedFilterTerms;
 	for (int32 iFilters = 0; iFilters < FilterTerms.Num() ; iFilters++)
@@ -678,13 +682,12 @@ void SGraphActionMenu::GenerateFilteredItems(bool bPreserveExpansion)
 			// Combine the actions string, separate with \n so terms don't run into each other, and remove the spaces (incase the user is searching for a variable)
 			// In the case of groups containing multiple actions, they will have been created and added at the same place in the code, using the same description
 			// and keywords, so we only need to use the first one for filtering.
-			FString SearchText = CurrentAction.Actions[0]->GetSearchTitle().ToString() + LINE_TERMINATOR + CurrentAction.Actions[0]->GetSearchKeywords().ToString() + LINE_TERMINATOR +CurrentAction.Actions[0]->GetSearchCategory().ToString();
-			SearchText = SearchText.Replace( TEXT( " " ), TEXT( "" ) );
+			const FString& SearchText = CurrentAction.GetSearchTextForFirstAction();
 
 			FString EachTermSanitized;
 			for (int32 FilterIndex = 0; (FilterIndex < FilterTerms.Num()) && bShowAction; ++FilterIndex)
 			{
-				const bool bMatchesTerm = ( SearchText.Contains( FilterTerms[FilterIndex] ) || ( SearchText.Contains( SanitizedFilterTerms[FilterIndex] ) == true ) );
+				const bool bMatchesTerm = (SearchText.Contains(FilterTerms[FilterIndex], ESearchCase::CaseSensitive) || (SearchText.Contains(SanitizedFilterTerms[FilterIndex], ESearchCase::CaseSensitive) == true));
 				bShowAction = bShowAction && bMatchesTerm;
 			}
 
@@ -774,42 +777,38 @@ int32 SGraphActionMenu::GetActionFilteredWeight( const FGraphActionListBuilderBa
 	// Helper array
 	struct FArrayWithWeight
 	{
-		TArray< FString > Array;
+		FArrayWithWeight(const TArray< FString >* InArray, int32 InWeight)
+			: Array(InArray)
+			, Weight(InWeight)
+		{
+		}
+
+		const TArray< FString >* Array;
 		int32			  Weight;
 	};
 
 	// Setup an array of arrays so we can do a weighted search			
 	TArray< FArrayWithWeight > WeightedArrayList;
-	FArrayWithWeight EachEntry;
-
+	
 	int32 Action = 0;
 	if( InCurrentAction.Actions[Action].IsValid() == true )
 	{
 		// Combine the actions string, separate with \n so terms don't run into each other, and remove the spaces (incase the user is searching for a variable)
 		// In the case of groups containing multiple actions, they will have been created and added at the same place in the code, using the same description
 		// and keywords, so we only need to use the first one for filtering.
-		FString SearchText = InCurrentAction.Actions[Action]->GetSearchTitle().ToString() + LINE_TERMINATOR + InCurrentAction.Actions[Action]->GetSearchKeywords().ToString() + LINE_TERMINATOR + InCurrentAction.Actions[Action]->GetSearchCategory().ToString();
-		SearchText = SearchText.Replace( TEXT( " " ), TEXT( "" ) );
+		const FString& SearchText = InCurrentAction.GetSearchTextForFirstAction();
 
 		// First the keywords
-		InCurrentAction.Actions[Action]->GetSearchKeywords().ToString().ParseIntoArray( EachEntry.Array, TEXT(" "), true );
-		EachEntry.Weight = 10;
-		WeightedArrayList.Add( EachEntry );
+		WeightedArrayList.Add(FArrayWithWeight(&InCurrentAction.GetSearchKeywordsArrayForFirstAction(), 10));
 
 		// The description
-		InCurrentAction.Actions[Action]->MenuDescription.ToString().ParseIntoArray( EachEntry.Array, TEXT(" "), true );
-		EachEntry.Weight = DescriptionWeight;
-		WeightedArrayList.Add( EachEntry );
+		WeightedArrayList.Add(FArrayWithWeight(&InCurrentAction.GetMenuDescriptionArrayForFirstAction(), DescriptionWeight));
 
 		// The node search title weight
-		InCurrentAction.Actions[Action]->GetSearchTitle().ToString().ParseIntoArray( EachEntry.Array, TEXT(" "), true );
-		EachEntry.Weight = NodeTitleWeight;
-		WeightedArrayList.Add( EachEntry );
+		WeightedArrayList.Add(FArrayWithWeight(&InCurrentAction.GetSearchTitleArrayForFirstAction(), NodeTitleWeight));
 
 		// The category
-		InCurrentAction.Actions[Action]->GetSearchCategory().ToString().ParseIntoArray( EachEntry.Array, TEXT(" "), true );
-		EachEntry.Weight = CategoryWeight;
-		WeightedArrayList.Add( EachEntry );
+		WeightedArrayList.Add(FArrayWithWeight(&InCurrentAction.GetSearchCategoryArrayForFirstAction(), CategoryWeight));
 
 		// Now iterate through all the filter terms and calculate a 'weight' using the values and multipliers
 		FString EachTerm;
@@ -818,11 +817,11 @@ int32 SGraphActionMenu::GetActionFilteredWeight( const FGraphActionListBuilderBa
 		{
 			EachTerm = InFilterTerms[FilterIndex];
 			EachTermSanitized = InSanitizedFilterTerms[FilterIndex];
-			if( SearchText.Contains( EachTerm ) )
+			if( SearchText.Contains( EachTerm, ESearchCase::CaseSensitive ) )
 			{
 				TotalWeight += 2;
 			}
-			else if( SearchText.Contains( EachTermSanitized ) )
+			else if (SearchText.Contains(EachTermSanitized, ESearchCase::CaseSensitive))
 			{
 				TotalWeight++;
 			}		
@@ -830,27 +829,27 @@ int32 SGraphActionMenu::GetActionFilteredWeight( const FGraphActionListBuilderBa
 			for (int32 iFindCount = 0; iFindCount < WeightedArrayList.Num() ; iFindCount++)
 			{
 				int32 WeightPerList = 0;
-				TArray<FString>& KeywordArray = WeightedArrayList[iFindCount].Array;
+				const TArray<FString>& KeywordArray = *WeightedArrayList[iFindCount].Array;
 				int32 EachWeight = WeightedArrayList[iFindCount].Weight;
 				int32 WholeMatchCount = 0;
 				for (int32 iEachWord = 0; iEachWord < KeywordArray.Num() ; iEachWord++)
 				{
 					// If we get an exact match weight the find count to get exact matches higher priority
-					if( KeywordArray[ iEachWord ] == EachTerm )
+					if (KeywordArray[iEachWord] == EachTerm)
 					{
 						WeightPerList += EachWeight * WholeMatchWeightMultiplier;
 						WholeMatchCount++;
 					}
-					else if( KeywordArray[ iEachWord ].Contains( EachTerm ) )
+					else if (KeywordArray[iEachWord].Contains(EachTerm, ESearchCase::CaseSensitive))
 					{
 						WeightPerList += EachWeight;
 					}
-					else if( KeywordArray[ iEachWord ] == EachTermSanitized )
+					else if (KeywordArray[iEachWord] == EachTermSanitized)
 					{
 						WeightPerList += ( EachWeight * WholeMatchWeightMultiplier ) / 2;
 						WholeMatchCount++;
 					}
-					else if( KeywordArray[ iEachWord ].Contains( EachTermSanitized ) )
+					else if (KeywordArray[iEachWord].Contains(EachTermSanitized, ESearchCase::CaseSensitive))
 					{
 						WeightPerList += EachWeight / 2;
 					}
diff --git a/Engine/Source/Editor/Kismet/Private/BlueprintActionMenuBuilder.cpp b/Engine/Source/Editor/Kismet/Private/BlueprintActionMenuBuilder.cpp
--- a/Engine/Source/Editor/Kismet/Private/BlueprintActionMenuBuilder.cpp
+++ b/Engine/Source/Editor/Kismet/Private/BlueprintActionMenuBuilder.cpp
@@ -110,7 +110,7 @@ TSharedPtr<FBlueprintActionMenuItem> FBlueprintActionMenuItemFactory::MakeAction
 	FBlueprintActionMenuItem* NewMenuItem = new FBlueprintActionMenuItem(Action, UiSignature);
 
 	NewMenuItem->Grouping = MenuGrouping;
-	NewMenuItem->Category = FText::Format(FText::FromString(TEXT("{0}|{1}")), RootCategory, NewMenuItem->Category);
+	NewMenuItem->Category = FText::FromString(RootCategory.ToString() + TEXT('|') + NewMenuItem->Category.ToString());
 
 	return MakeShareable(NewMenuItem);
 }
@@ -121,7 +121,7 @@ TSharedPtr<FBlueprintDragDropMenuItem> FBlueprintActionMenuItemFactory::MakeDrag
 	// FBlueprintDragDropMenuItem takes care of its own menu MenuDescription, etc.
 	FBlueprintDragDropMenuItem* NewMenuItem = new FBlueprintDragDropMenuItem(Context, SampleAction, MenuGrouping);
 
-	NewMenuItem->Category = FText::Format(FText::FromString(TEXT("{0}|{1}")), RootCategory, NewMenuItem->Category);
+	NewMenuItem->Category = FText::FromString(RootCategory.ToString() + TEXT('|') + NewMenuItem->Category.ToString());
 	return MakeShareable(NewMenuItem);
 }
 
@@ -134,7 +134,7 @@ TSharedPtr<FBlueprintActionMenuItem> FBlueprintActionMenuItemFactory::MakeBoundM
 	FBlueprintActionMenuItem* NewMenuItem = new FBlueprintActionMenuItem(Action, UiSignature, ActionInfo.GetBindings());
 
 	NewMenuItem->Grouping = MenuGrouping;
-	NewMenuItem->Category = FText::Format(FText::FromString(TEXT("{0}|{1}")), RootCategory, NewMenuItem->Category);
+	NewMenuItem->Category = FText::FromString( RootCategory.ToString() + TEXT('|') + NewMenuItem->Category.ToString());
 
 	return MakeShareable(NewMenuItem);
 }
diff --git a/Engine/Source/Editor/UnrealEd/Private/Commandlets/DumpBlueprintsInfoCommandlet.cpp b/Engine/Source/Editor/UnrealEd/Private/Commandlets/DumpBlueprintsInfoCommandlet.cpp
--- a/Engine/Source/Editor/UnrealEd/Private/Commandlets/DumpBlueprintsInfoCommandlet.cpp
+++ b/Engine/Source/Editor/UnrealEd/Private/Commandlets/DumpBlueprintsInfoCommandlet.cpp
@@ -1690,7 +1690,7 @@ static void DumpBlueprintInfoUtils::DumpActionMenuItem(uint32 Indent, FGraphActi
 
 		ActionEntry += IndentedNewline + TooltipFieldLabel + TooltipStr + "\",";
 		ActionEntry += IndentedNewline + "\"Keywords\"    : \"" + PrimeAction->Keywords.ToString() + "\",";
-		ActionEntry += IndentedNewline + "\"SearchTitle\" : \"" + PrimeAction->GetSearchTitle().ToString() + "\",";
+		ActionEntry += IndentedNewline + "\"SearchTitle\" : \"" + PrimeAction->GetSearchTitle() + "\",";
 		ActionEntry += IndentedNewline + FString::Printf(TEXT("\"Grouping\"    : %d"), PrimeAction->Grouping);
 		
 		// Get action node type info
diff --git a/Engine/Source/Editor/UnrealEd/Private/EditorCategoryUtils.cpp b/Engine/Source/Editor/UnrealEd/Private/EditorCategoryUtils.cpp
--- a/Engine/Source/Editor/UnrealEd/Private/EditorCategoryUtils.cpp
+++ b/Engine/Source/Editor/UnrealEd/Private/EditorCategoryUtils.cpp
@@ -298,7 +298,7 @@ void FEditorCategoryUtils::GetClassHideCategories(UClass const* Class, TArray<FS
 		
 		for (FString& Category : CategoriesOut)
 		{
-			Category = GetCategoryDisplayString(FText::FromString(Category)).ToString();
+			Category = GetCategoryDisplayString(Category);
 		}
 	}
 }
diff --git a/Engine/Source/Runtime/Core/Private/UObject/UnrealNames.cpp b/Engine/Source/Runtime/Core/Private/UObject/UnrealNames.cpp
--- a/Engine/Source/Runtime/Core/Private/UObject/UnrealNames.cpp
+++ b/Engine/Source/Runtime/Core/Private/UObject/UnrealNames.cpp
@@ -218,7 +218,7 @@ FCriticalSection* FName::GetCriticalSection()
 FString FName::NameToDisplayString( const FString& InDisplayName, const bool bIsBool )
 {
 	// Copy the characters out so that we can modify the string in place
-	const TArray< TCHAR > Chars = InDisplayName.GetCharArray();
+	const TArray< TCHAR >& Chars = InDisplayName.GetCharArray();
 
 	// This is used to indicate that we are in a run of uppercase letter and/or digits.  The code attempts to keep
 	// these characters together as breaking them up often looks silly (i.e. "Draw Scale 3 D" as opposed to "Draw Scale 3D"
@@ -226,6 +226,7 @@ FString FName::NameToDisplayString( const FString& InDisplayName, const bool bIs
 	bool bWasSpace = false;
 	bool bWasOpenParen = false;
 	FString OutDisplayName;
+	OutDisplayName.GetCharArray().Reserve(Chars.Num());
 	for( int32 CharIndex = 0 ; CharIndex < Chars.Num() ; ++CharIndex )
 	{
 		TCHAR ch = Chars[CharIndex];
diff --git a/Engine/Source/Runtime/Engine/Classes/EdGraph/EdGraphSchema.h b/Engine/Source/Runtime/Engine/Classes/EdGraph/EdGraphSchema.h
--- a/Engine/Source/Runtime/Engine/Classes/EdGraph/EdGraphSchema.h
+++ b/Engine/Source/Runtime/Engine/Classes/EdGraph/EdGraphSchema.h
@@ -82,15 +82,15 @@ struct ENGINE_API FEdGraphSchemaAction
 
 	/** Search title for the action (doesn't have to be set when instantiated, will be constructed by GetSearchTitle() if left empty). */
 	UPROPERTY()
-	FText CachedSearchTitle;
+	FString CachedSearchTitle;
 
 	/** Search keywords for the action (doesn't have to be set when instantiated, will be constructed by GetSearchTitle() if left empty). */
 	UPROPERTY()
-	FText CachedSearchKeywords;
+	FString CachedSearchKeywords;
 
 	/** Search categories for the action (doesn't have to be set when instantiated, will be constructed by GetSearchTitle() if left empty). */
 	UPROPERTY()
-	FText CachedSearchCategories;
+	FString CachedSearchCategories;
 
 	FEdGraphSchemaAction() 
 		: Grouping(0)
@@ -131,40 +131,40 @@ struct ENGINE_API FEdGraphSchemaAction
 	}
 
 	/** Retrieves the full searchable title for this action. */
-	FText GetSearchTitle()
+	const FString& GetSearchTitle()
 	{
 		if(CachedSearchTitle.IsEmpty())
 		{
-			FFormatNamedArguments Args;
-			Args.Add(TEXT("LocalizedTitle"), MenuDescription);
-			Args.Add(TEXT("SourceTitle"), FText::FromString(MenuDescription.BuildSourceString()));
-			CachedSearchTitle = FText::Format(FText::FromString("{LocalizedTitle} {SourceTitle}"), Args);
+			CachedSearchTitle = MenuDescription.ToString();
+			CachedSearchTitle.AppendChar(TEXT(' '));
+			CachedSearchTitle.Append(MenuDescription.BuildSourceString());
+			CachedSearchTitle = CachedSearchTitle.ToLower();
 		}
 		return CachedSearchTitle;
 	}
 
 	/** Retrieves the full searchable keywords for this action. */
-	FText GetSearchKeywords()
+	const FString& GetSearchKeywords()
 	{
 		if(CachedSearchKeywords.IsEmpty())
 		{
-			FFormatNamedArguments Args;
-			Args.Add(TEXT("LocalizedKeywords"), Keywords);
-			Args.Add(TEXT("SourceKeywords"), FText::FromString(Keywords.BuildSourceString()));
-			CachedSearchKeywords = FText::Format(FText::FromString("{LocalizedKeywords} {SourceKeywords}"), Args);
+			CachedSearchKeywords = Keywords.ToString();
+			CachedSearchKeywords.AppendChar(TEXT(' '));
+			CachedSearchKeywords.Append(Keywords.BuildSourceString());
+			CachedSearchKeywords = CachedSearchKeywords.ToLower();
 		}
 		return CachedSearchKeywords;
 	}
 
 	/** Retrieves the full searchable categories for this action. */
-	FText GetSearchCategory()
+	const FString& GetSearchCategory()
 	{
 		if(CachedSearchCategories.IsEmpty())
 		{
-			FFormatNamedArguments Args;
-			Args.Add(TEXT("LocalizedCategories"), Category);
-			Args.Add(TEXT("SourceCategories"), FText::FromString(Category.BuildSourceString()));
-			CachedSearchCategories = FText::Format(FText::FromString("{LocalizedCategories} {SourceCategories}"), Args);
+			CachedSearchCategories = Category.ToString();
+			CachedSearchCategories.AppendChar(TEXT(' '));
+			CachedSearchCategories.Append(Category.BuildSourceString());
+			CachedSearchCategories = CachedSearchCategories.ToLower();
 		}
 		return CachedSearchCategories;
 	}
@@ -297,6 +297,15 @@ public:
 		/** Constructor accepting multiple actions */
 		ActionGroup( const TArray< TSharedPtr<FEdGraphSchemaAction> >& InActions, FString const& RootCategory = TEXT("") );
 
+		/** Move constructor and move assignment operator */
+		ENGINE_API ActionGroup(ActionGroup && Other);
+		ENGINE_API ActionGroup& operator=(ActionGroup && Other);
+
+		/** Copy constructor and assignment operator */
+		ENGINE_API ActionGroup(const ActionGroup&);
+		ENGINE_API ActionGroup& operator=(const ActionGroup&);
+
+		ENGINE_API ~ActionGroup();
 		/**
 		 * 
 		 * @param  HierarchyOut	A list of the category tiers that this action should be listed under.
@@ -311,21 +320,54 @@ public:
 		 */
 		ENGINE_API void PerformAction( class UEdGraph* ParentGraph, TArray<UEdGraphPin*>& FromPins, const FVector2D Location );
 		
+		/**
+		 * Returns a the string that should be used when searching for matching actions. Looks only at the first action.
+		 */
+		ENGINE_API const FString& GetSearchTextForFirstAction() const { return SearchText; }
+
+		/** Returns the SearchKeywordsArray */
+		ENGINE_API const TArray<FString>& GetSearchKeywordsArrayForFirstAction() const { return SearchKeywordsArray; }
+		/** Returns the MenuDescriptionArray */
+		ENGINE_API const TArray<FString>& GetMenuDescriptionArrayForFirstAction() const { return MenuDescriptionArray; }
+		/** Returns the SearchTitleArray */
+		ENGINE_API const TArray<FString>& GetSearchTitleArrayForFirstAction() const { return SearchTitleArray; }
+		/** Returns the SearchCategoryArray */
+		ENGINE_API const TArray<FString>& GetSearchCategoryArrayForFirstAction() const { return SearchCategoryArray; }
+
 		/** All of the actions this entry contains */
 		TArray< TSharedPtr<FEdGraphSchemaAction> > Actions;
 
 	private:
+		void Move(ActionGroup& Other);
+		void Copy(const ActionGroup& Other);
+
 		/**
-		 *Concatenates RootCategory with the first action's category (RootCategory
+		 * Concatenates RootCategory with the first action's category (RootCategory
 		 * coming first, as a prefix, and the splits the category hierarchy apart
 		 * into separate entries.
 		 */
 		void InitCategoryChain();
 
+		/**
+		 * Initializes the search text.
+		 */
+		void InitSearchText();
+
+		void InitScoringData();
+
 		/** The category to list this entry under (could be left empty, as it gets concatenated with the first sub-action's category) */
 		FString RootCategory;
 		/** The chain of categories */
 		TArray<FString> CategoryChain;
+
+		/** Scoring data */
+		TArray<FString> SearchKeywordsArray;
+		TArray<FString> MenuDescriptionArray;
+		TArray<FString> SearchTitleArray;
+		TArray<FString> SearchCategoryArray;
+
+		/** The text that can be used to find this entry */
+		FString SearchText;
 	};
 private:
 
diff --git a/Engine/Source/Runtime/Engine/Private/EdGraph/EdGraphSchema.cpp b/Engine/Source/Runtime/Engine/Private/EdGraph/EdGraphSchema.cpp
--- a/Engine/Source/Runtime/Engine/Private/EdGraph/EdGraphSchema.cpp
+++ b/Engine/Source/Runtime/Engine/Private/EdGraph/EdGraphSchema.cpp
@@ -30,7 +30,7 @@ void FGraphActionListBuilderBase::AddActionList( const TArray<TSharedPtr<FEdGrap
 
 void FGraphActionListBuilderBase::Append( FGraphActionListBuilderBase& Other )
 {
-	Entries.Append( Other.Entries );
+	Entries.Append( MoveTemp(Other.Entries) );
 }
 
 int32 FGraphActionListBuilderBase::GetNumActions() const
@@ -56,6 +56,7 @@ FGraphActionListBuilderBase::ActionGroup::ActionGroup( TSharedPtr<FEdGraphSchema
 {
 	Actions.Add( InAction );
 	InitCategoryChain();
+	InitScoringData();
 }
 
 FGraphActionListBuilderBase::ActionGroup::ActionGroup( const TArray< TSharedPtr<FEdGraphSchemaAction> >& InActions, FString const& CategoryPrefix/* = TEXT("") */ )
@@ -63,6 +64,39 @@ FGraphActionListBuilderBase::ActionGroup::ActionGroup( const TArray< TSharedPtr<
 {
 	Actions = InActions;
 	InitCategoryChain();
+	InitScoringData();
+}
+
+FGraphActionListBuilderBase::ActionGroup::ActionGroup(FGraphActionListBuilderBase::ActionGroup && Other)
+{
+	Move(Other);
+}
+
+FGraphActionListBuilderBase::ActionGroup& FGraphActionListBuilderBase::ActionGroup::operator=(FGraphActionListBuilderBase::ActionGroup && Other)
+{
+	if (&Other != this)
+	{
+		Move(Other);
+	}
+	return *this;
+}
+
+FGraphActionListBuilderBase::ActionGroup::ActionGroup(const FGraphActionListBuilderBase::ActionGroup& Other)
+{
+	Copy(Other);
+}
+
+FGraphActionListBuilderBase::ActionGroup& FGraphActionListBuilderBase::ActionGroup::operator=(const FGraphActionListBuilderBase::ActionGroup& Other)
+{
+	if (&Other != this)
+	{
+		Copy(Other);
+	}
+	return *this;
+}
+
+FGraphActionListBuilderBase::ActionGroup::~ActionGroup()
+{
 }
 
 void FGraphActionListBuilderBase::ActionGroup::GetCategoryChain(TArray<FString>& HierarchyOut) const
@@ -84,6 +118,28 @@ void FGraphActionListBuilderBase::ActionGroup::PerformAction( class UEdGraph* Pa
 	}
 }
 
+void FGraphActionListBuilderBase::ActionGroup::Move(FGraphActionListBuilderBase::ActionGroup& Other)
+{
+	Actions = MoveTemp(Other.Actions);
+	RootCategory = MoveTemp(Other.RootCategory);
+	CategoryChain = MoveTemp(Other.CategoryChain);
+	SearchKeywordsArray = MoveTemp(Other.SearchKeywordsArray);
+	MenuDescriptionArray = MoveTemp(Other.MenuDescriptionArray);
+	SearchCategoryArray = MoveTemp(Other.SearchCategoryArray);
+	SearchText = MoveTemp(Other.SearchText);
+}
+
+void FGraphActionListBuilderBase::ActionGroup::Copy(const ActionGroup& Other)
+{
+	Actions = Other.Actions;
+	RootCategory = Other.RootCategory;
+	CategoryChain = Other.CategoryChain;
+	SearchKeywordsArray = Other.SearchKeywordsArray;
+	MenuDescriptionArray = Other.MenuDescriptionArray;
+	SearchCategoryArray = Other.SearchCategoryArray;
+	SearchText = Other.SearchText;
+}
+
 void FGraphActionListBuilderBase::ActionGroup::InitCategoryChain()
 {
 #if WITH_EDITOR
@@ -107,6 +163,39 @@ void FGraphActionListBuilderBase::ActionGroup::InitCategoryChain()
 #endif
 }
 
+void FGraphActionListBuilderBase::ActionGroup::InitScoringData()
+{
+	if (Actions.Num() > 0)
+	{
+		FEdGraphSchemaAction& FirstAction = *Actions[0];
+		const FString& Keywords = FirstAction.GetSearchKeywords();
+		const FString& Title = FirstAction.GetSearchTitle();
+		const FString& Category = FirstAction.GetSearchCategory();
+
+		// We keep these individual arrays so that they can be weighted differently by the scoring algorithm in SGraphActionMenu::GetActionFilteredWeight
+		Keywords.ParseIntoArray(SearchKeywordsArray, TEXT(" "), true);
+		FirstAction.MenuDescription.ToString().ToLower().ParseIntoArray(MenuDescriptionArray, TEXT(" "), true);
+		Title.ParseIntoArray(SearchTitleArray, TEXT(" "), true);
+		Category.ParseIntoArray(SearchCategoryArray, TEXT(" "), true);
+
+		// Glob search text together, we use the SearchText string for basic filtering:
+		for (const auto& Entry : SearchTitleArray)
+		{
+			SearchText += Entry;
+		}
+		SearchText.Append(LINE_TERMINATOR);
+		for (const auto& Entry : SearchKeywordsArray)
+		{
+			SearchText += Entry;
+		}
+		SearchText.Append(LINE_TERMINATOR);
+		for (const auto& Entry : SearchCategoryArray)
+		{
+			SearchText += Entry;
+		}
+	}
+}
+
 /////////////////////////////////////////////////////
 // FCategorizedGraphActionListBuilder
 
