[NO NON-PII COMMIT MESSAGE]

diff --git a/Engine/Source/Editor/GraphEditor/Private/SGraphNode.cpp b/Engine/Source/Editor/GraphEditor/Private/SGraphNode.cpp
--- a/Engine/Source/Editor/GraphEditor/Private/SGraphNode.cpp
+++ b/Engine/Source/Editor/GraphEditor/Private/SGraphNode.cpp
@@ -58,7 +58,7 @@ FText SNodeTitle::GetNodeTitle() const
 
 FText SNodeTitle::GetHeadTitle() const
 {
-	return GraphNode->bCanRenameNode ? GraphNode->GetNodeTitle(ENodeTitleType::EditableTitle) : CachedHeadTitle;
+	return (GraphNode && GraphNode->bCanRenameNode) ? GraphNode->GetNodeTitle(ENodeTitleType::EditableTitle) : CachedHeadTitle;
 }
 
 FVector2D SNodeTitle::GetTitleSize() const
