Improved verbage around pin removal/additions, also added details to the tooltip

diff --git a/Engine/Source/Editor/GraphEditor/Private/GraphDiffControl.cpp b/Engine/Source/Editor/GraphEditor/Private/GraphDiffControl.cpp
--- a/Engine/Source/Editor/GraphEditor/Private/GraphDiffControl.cpp
+++ b/Engine/Source/Editor/GraphEditor/Private/GraphDiffControl.cpp
@@ -294,11 +294,91 @@ static void DiffR_NodePinCount(FDiffResults& Results, class UEdGraphNode* Node,
 		Args.Add(TEXT("NodeName"), NodeName);
 		Args.Add(TEXT("OriginalCount"), OriginalCount);
 		Args.Add(TEXT("NewCount"), NewCount);
-		Diff.ToolTip = FText::Format(LOCTEXT("DIF_NodePinCountChangedToolTip", "Node '{NodeName}' had {OriginalCount} Pins, now has {NewCount} Pins"), Args);
 		Diff.DisplayColor = FLinearColor(0.45f,0.4f,0.4f);
-		Diff.DisplayString = (OriginalCount < NewCount) 
-			? (FText::Format(LOCTEXT("DIF_NodePinCountIncreased", "Add Pin '{NodeName}'"), Args)) 
-			: (FText::Format(LOCTEXT("DIF_NodePinCountDecreased", "Removed Pin '{NodeName}'"), Args));
+
+		struct FMatchName
+		{
+			FMatchName(const FString& InPinName)
+				: PinName(InPinName)
+			{
+			}
+
+			const FString& PinName;
+
+			bool operator()(const UEdGraphPin* Entry )
+			{
+				return PinName == Entry->PinName;
+			}
+		};
+
+		FText ListOfPins;
+		TArray< FText > RemovedPins;
+		for (auto OldPin : Pins1)
+		{
+			const UEdGraphPin* const* FoundPin = Pins2.FindByPredicate(FMatchName(OldPin->PinName));
+			if (FoundPin == nullptr)
+			{
+				RemovedPins.Add(OldPin->GetDisplayName());
+			}
+		}
+
+		TArray< FText > AddedPins;
+		for (auto OldPin : Pins2)
+		{
+			const UEdGraphPin* const* FoundPin = Pins1.FindByPredicate(FMatchName(OldPin->PinName));
+			if (FoundPin == nullptr)
+			{
+				AddedPins.Add(OldPin->GetDisplayName());
+			}
+		}
+
+		if (RemovedPins.Num() > 0 && AddedPins.Num() > 0)
+		{
+			Diff.DisplayString = FText::Format(LOCTEXT("DIF_NodePinsAddedAndRemoved", "Added and removed Pins from '{NodeName}'"), Args);
+		}
+		else if (AddedPins.Num() > 0)
+		{
+			if (AddedPins.Num() == 1)
+			{
+				Diff.DisplayString = FText::Format(LOCTEXT("DIF_NodePinCountIncreased", "Added Pin to '{NodeName}'"), Args);
+			}
+			else
+			{
+				Diff.DisplayString = FText::Format(LOCTEXT("DIF_NodePinCountIncreasedSeveral", "Added Pins to '{NodeName}'"), Args);
+			}
+		}
+		else if (RemovedPins.Num() > 0)
+		{
+			if (RemovedPins.Num() == 1)
+			{
+				Diff.DisplayString = FText::Format(LOCTEXT("DIF_NodePinCountDecreased", "Removed Pin from '{NodeName}'"), Args);
+			}
+			else
+			{
+				Diff.DisplayString = FText::Format(LOCTEXT("DIF_NodePinCountDecreasedSeveral", "Removed Pins from '{NodeName}'"), Args);
+			}
+		}
+
+		FTextBuilder Builder;
+		Builder.AppendLine(FText::Format(LOCTEXT("DIF_NodePinCountChangedToolTip", "Node '{NodeName}' had {OriginalCount} Pins, now has {NewCount} Pins"), Args));
+		if (AddedPins.Num() > 0)
+		{
+			Builder.AppendLine(LOCTEXT("DIF_PinsAddedList", "Pins Added:"));
+			for (const auto& Added : AddedPins)
+			{
+				Builder.AppendLine(Added);
+			}
+		}
+
+		if (RemovedPins.Num() > 0)
+		{
+			Builder.AppendLine(LOCTEXT("DIF_PinsRemovedList", "Pins Removed:"));
+			for (const auto& Removed : RemovedPins)
+			{
+				Builder.AppendLine(Removed);
+			}
+		}
+		Diff.ToolTip = Builder.ToText();
 	}
 
 	Results.Add(Diff);
