Minor optimization to blueprint context menu construction time. Improves construction time by .1-.2s when context sensitivity is off (from .9s to .7s)

diff --git a/Engine/Source/Editor/GraphEditor/Private/GraphActionNode.cpp b/Engine/Source/Editor/GraphEditor/Private/GraphActionNode.cpp
--- a/Engine/Source/Editor/GraphEditor/Private/GraphActionNode.cpp
+++ b/Engine/Source/Editor/GraphEditor/Private/GraphActionNode.cpp
@@ -282,6 +282,7 @@ void FGraphActionNode::ExpandAllChildren(TSharedPtr< STreeView< TSharedPtr<FGrap
 void FGraphActionNode::ClearChildren()
 {
 	Children.Empty();
+	CategoryNodes.Empty();
 	ChildGroupings.Empty();
 	ChildSections.Empty();
 }
@@ -501,32 +502,45 @@ TSharedPtr<FGraphActionNode> FGraphActionNode::FindMatchingParent(FString const&
 	// actions (meaning that action node can have children).
 	bool const bCanNestUnderActionNodes = NodeToAdd->IsActionNode() && NodeToAdd->GetPrimaryAction()->IsParentable();
 
-	for (TSharedPtr<FGraphActionNode> const& ChildNode : Children)
+	if (bCanNestUnderActionNodes)
 	{
-		if (ChildNode->IsCategoryNode())
+		// slow path, not commonly used:
+		for (TSharedPtr<FGraphActionNode> const& ChildNode : Children)
 		{
-			if ((NodeToAdd->SectionID == ChildNode->SectionID) &&
-				(ParentName == ChildNode->DisplayText.ToString()))
+			if (ChildNode->IsCategoryNode())
 			{
-				FoundCategoryNode = ChildNode;
-				break;
+				if ((NodeToAdd->SectionID == ChildNode->SectionID) &&
+					(ParentName == ChildNode->DisplayText.ToString()))
+				{
+					FoundCategoryNode = ChildNode;
+					break;
+				}
 			}
-		}
-		else if (bCanNestUnderActionNodes && ChildNode->IsActionNode())
-		{
-			// make the action's name into a display name, all categories are 
-			// set as such (to ensure that the action name best matches the 
-			// category ParentName)
-			FString ChildNodeName = FName::NameToDisplayString(ChildNode->DisplayText.ToString(), /*bIsBool =*/false);
-
-			// @TODO: should we be matching section/grouping as well?
-			if (ChildNodeName == ParentName)
+			else if (bCanNestUnderActionNodes && ChildNode->IsActionNode())
 			{
-				FoundCategoryNode = ChildNode;
-				break;
+				// make the action's name into a display name, all categories are 
+				// set as such (to ensure that the action name best matches the 
+				// category ParentName)
+				FString ChildNodeName = FName::NameToDisplayString(ChildNode->DisplayText.ToString(), /*bIsBool =*/false);
+
+				// @TODO: should we be matching section/grouping as well?
+				if (ChildNodeName == ParentName)
+				{
+					FoundCategoryNode = ChildNode;
+					break;
+				}
 			}
 		}
 	}
+	else
+	{
+		// fast path, just look up in category map:
+		TSharedPtr<FGraphActionNode>* PotentialCategoryNode = CategoryNodes.Find(ParentName);
+		if (PotentialCategoryNode && (*PotentialCategoryNode)->SectionID == NodeToAdd->SectionID)
+		{
+			FoundCategoryNode = *PotentialCategoryNode;
+		}
+	}
 
 	return FoundCategoryNode;
 }
@@ -558,6 +572,10 @@ void FGraphActionNode::InsertChild(TSharedPtr<FGraphActionNode> NodeToAdd)
 
 			NodeToAdd->InsertOrder = NewSection->Children.Num();
 			NewSection->Children.Add(NodeToAdd);
+			if (NodeToAdd->IsCategoryNode())
+			{
+				CategoryNodes.Add(NodeToAdd->DisplayText.ToString(), NodeToAdd);
+			}
 			return;
 		}
 	}
@@ -590,6 +608,10 @@ void FGraphActionNode::InsertChild(TSharedPtr<FGraphActionNode> NodeToAdd)
 
 	NodeToAdd->InsertOrder = Children.Num();
 	Children.Add(NodeToAdd);
+	if (NodeToAdd->IsCategoryNode())
+	{
+		CategoryNodes.Add(NodeToAdd->DisplayText.ToString(), NodeToAdd);
+	}
 }
 
 
diff --git a/Engine/Source/Editor/GraphEditor/Private/GraphActionNode.h b/Engine/Source/Editor/GraphEditor/Private/GraphActionNode.h
--- a/Engine/Source/Editor/GraphEditor/Private/GraphActionNode.h
+++ b/Engine/Source/Editor/GraphEditor/Private/GraphActionNode.h
@@ -21,6 +21,9 @@ public:
 	/** */
 	TArray< TSharedPtr<FGraphActionNode> > Children;
 
+	/** Lookup table for category nodes, used to speed up menu construction */
+	TMap< FString, TSharedPtr<FGraphActionNode> > CategoryNodes;
+
 public:
 	/**
 	 * Static allocator for a new root node (so external users have a starting
