Stop graph panels handling num-pad and direction-keys when they are pressed with a modifier.

This enables key-bindings for actions in graph panels for chords such as CTRL+Left, CTRL+Up etc.

diff --git a/Engine/Source/Editor/GraphEditor/Private/SGraphPanel.cpp b/Engine/Source/Editor/GraphEditor/Private/SGraphPanel.cpp
--- a/Engine/Source/Editor/GraphEditor/Private/SGraphPanel.cpp
+++ b/Engine/Source/Editor/GraphEditor/Private/SGraphPanel.cpp
@@ -616,26 +616,31 @@ FReply SGraphPanel::OnKeyDown( const FGeometry& MyGeometry, const FKeyEvent& InK
 {
 	if( IsEditable.Get() )
 	{
-		if( InKeyEvent.GetKey() == EKeys::Up  ||InKeyEvent.GetKey() ==  EKeys::NumPadEight )
+		const bool bIsModifierActive = InKeyEvent.IsCommandDown() || InKeyEvent.IsAltDown() || InKeyEvent.IsShiftDown() || InKeyEvent.IsControlDown();
+		if (!bIsModifierActive)
 		{
-			UpdateSelectedNodesPositions(FVector2D(0.0f,-GetSnapGridSize()));
-			return FReply::Handled();
-		}
-		if( InKeyEvent.GetKey() ==  EKeys::Down || InKeyEvent.GetKey() ==  EKeys::NumPadTwo )
-		{
-			UpdateSelectedNodesPositions(FVector2D(0.0f,GetSnapGridSize()));
-			return FReply::Handled();
-		}
-		if( InKeyEvent.GetKey() ==  EKeys::Right || InKeyEvent.GetKey() ==  EKeys::NumPadSix )
-		{
-			UpdateSelectedNodesPositions(FVector2D(GetSnapGridSize(),0.0f));
-			return FReply::Handled();
-		}
-		if( InKeyEvent.GetKey() ==  EKeys::Left || InKeyEvent.GetKey() ==  EKeys::NumPadFour )
-		{
-			UpdateSelectedNodesPositions(FVector2D(-GetSnapGridSize(),0.0f));
-			return FReply::Handled();
+			if( InKeyEvent.GetKey() == EKeys::Up  ||InKeyEvent.GetKey() ==  EKeys::NumPadEight )
+			{
+				UpdateSelectedNodesPositions(FVector2D(0.0f,-GetSnapGridSize()));
+				return FReply::Handled();
+			}
+			if( InKeyEvent.GetKey() ==  EKeys::Down || InKeyEvent.GetKey() ==  EKeys::NumPadTwo )
+			{
+				UpdateSelectedNodesPositions(FVector2D(0.0f,GetSnapGridSize()));
+				return FReply::Handled();
+			}
+			if( InKeyEvent.GetKey() ==  EKeys::Right || InKeyEvent.GetKey() ==  EKeys::NumPadSix )
+			{
+				UpdateSelectedNodesPositions(FVector2D(GetSnapGridSize(),0.0f));
+				return FReply::Handled();
+			}
+			if( InKeyEvent.GetKey() ==  EKeys::Left || InKeyEvent.GetKey() ==  EKeys::NumPadFour )
+			{
+				UpdateSelectedNodesPositions(FVector2D(-GetSnapGridSize(),0.0f));
+				return FReply::Handled();
+			}
 		}
+
 		if(InKeyEvent.GetKey() == FGraphEditorCommands::Get().ZoomOut->GetActiveChord()->Key)
 		{
 			ChangeZoomLevel(-1, CachedAllottedGeometryScaledSize / 2.f, InKeyEvent.IsControlDown());
