This fixes the problem with hidden comment bubbles occluding graph node pins.
#UE-21225 - Hidden Comment Bubble is Partially Covering/Blocking Top of Node
#Branch UE4
#Proj GraphEditor

#info The comment bubble toggle button should have never been visible at all in this case, oversight on my part. I've made the visibility state return as collapsed rather than hidden.

diff --git a/Engine/Source/Editor/GraphEditor/Private/SCommentBubble.cpp b/Engine/Source/Editor/GraphEditor/Private/SCommentBubble.cpp
--- a/Engine/Source/Editor/GraphEditor/Private/SCommentBubble.cpp
+++ b/Engine/Source/Editor/GraphEditor/Private/SCommentBubble.cpp
@@ -403,7 +403,7 @@ void SCommentBubble::OnCommentTextCommitted( const FText& NewText, ETextCommit::
 
 EVisibility SCommentBubble::GetToggleButtonVisibility() const
 {
-	EVisibility Visibility = EVisibility::Hidden;
+	EVisibility Visibility = EVisibility::Collapsed;
 
 	if( OpacityValue > 0.f && !GraphNode->bCommentBubbleVisible )
 	{
