Moved error bar for variable nodes to the bottom instead of the top to prevent shifting of nodes when there are errors.

diff --git a/Engine/Source/Editor/GraphEditor/Private/KismetNodes/SGraphNodeK2Var.cpp b/Engine/Source/Editor/GraphEditor/Private/KismetNodes/SGraphNodeK2Var.cpp
--- a/Engine/Source/Editor/GraphEditor/Private/KismetNodes/SGraphNodeK2Var.cpp
+++ b/Engine/Source/Editor/GraphEditor/Private/KismetNodes/SGraphNodeK2Var.cpp
@@ -188,13 +188,6 @@ void SGraphNodeK2Var::UpdateGraphNode()
 	[
 		SNew(SVerticalBox)
 		+SVerticalBox::Slot()
-		.VAlign(VAlign_Top)
-		.AutoHeight() 
-		.Padding( FMargin(5.0f, 1.0f) )
-		[
-			ErrorReporting->AsWidget()
-		]
-		+SVerticalBox::Slot()
 		[
 			SNew(SOverlay)
 			.AddMetaData<FGraphNodeMetaData>(TagMeta)
@@ -245,6 +238,13 @@ void SGraphNodeK2Var::UpdateGraphNode()
 				]
 			]
 		]
+		+SVerticalBox::Slot()
+		.VAlign(VAlign_Top)
+		.AutoHeight() 
+		.Padding( FMargin(5.0f, 1.0f) )
+		[
+			ErrorReporting->AsWidget()
+		]
 	];
 
 	float VerticalPaddingAmount = 0.0f;
