Pressing enter when selecting an object from a Blueprint node now works

diff --git a/Engine/Source/Editor/GraphEditor/Private/KismetPins/SGraphPinObject.cpp b/Engine/Source/Editor/GraphEditor/Private/KismetPins/SGraphPinObject.cpp
--- a/Engine/Source/Editor/GraphEditor/Private/KismetPins/SGraphPinObject.cpp
+++ b/Engine/Source/Editor/GraphEditor/Private/KismetPins/SGraphPinObject.cpp
@@ -192,6 +192,7 @@ TSharedRef<SWidget> SGraphPinObject::GenerateAssetPicker()
 	AssetPickerConfig.bAllowNullSelection = true;
 	AssetPickerConfig.Filter.bRecursiveClasses = true;
 	AssetPickerConfig.OnAssetSelected = FOnAssetSelected::CreateSP(this, &SGraphPinObject::OnAssetSelectedFromPicker);
+	AssetPickerConfig.OnAssetEnterPressed = FOnAssetEnterPressed::CreateSP(this, &SGraphPinObject::OnAssetEnterPressedInPicker);
 	AssetPickerConfig.InitialAssetViewType = EAssetViewType::List;
 	AssetPickerConfig.bAllowDragging = false;
 
@@ -240,6 +241,13 @@ void SGraphPinObject::OnAssetSelectedFromPicker(const struct FAssetData& AssetDa
 	}
 }
 
+void SGraphPinObject::OnAssetEnterPressedInPicker(const TArray<FAssetData>& InSelectedAssets)
+{
+	if(InSelectedAssets.Num() > 0)
+	{
+		OnAssetSelectedFromPicker(InSelectedAssets[0]);
+	}
+}
 
 FText SGraphPinObject::GetValue() const
 {
diff --git a/Engine/Source/Editor/GraphEditor/Public/KismetPins/SGraphPinObject.h b/Engine/Source/Editor/GraphEditor/Public/KismetPins/SGraphPinObject.h
--- a/Engine/Source/Editor/GraphEditor/Public/KismetPins/SGraphPinObject.h
+++ b/Engine/Source/Editor/GraphEditor/Public/KismetPins/SGraphPinObject.h
@@ -50,6 +50,8 @@ protected:
 	virtual TSharedRef<SWidget> GenerateAssetPicker();
 	/** Called to validate selection from picker window */
 	virtual void OnAssetSelectedFromPicker(const struct FAssetData& AssetData);
+	/** Called when enter is pressed when items are selected in the picker window */
+	void OnAssetEnterPressedInPicker(const TArray<FAssetData>& InSelectedAssets);
 
 	/** Used to update the combo button text */
 	FText OnGetComboTextValue() const;
