[NO NON-PII COMMIT MESSAGE]

diff --git a/Engine/Source/Editor/GraphEditor/Private/DragConnection.cpp b/Engine/Source/Editor/GraphEditor/Private/DragConnection.cpp
--- a/Engine/Source/Editor/GraphEditor/Private/DragConnection.cpp
+++ b/Engine/Source/Editor/GraphEditor/Private/DragConnection.cpp
@@ -265,7 +265,8 @@ FReply FDragConnection::DroppedOnNode(FVector2D ScreenPosition, FVector2D GraphP
 
 					UEdGraphPin* EdGraphPin = NodeOver->GetSchema()->DropPinOnNode(GetHoveredNode(), PinName, SourcePin->PinType, SourcePin->Direction);
 
-					if(EdGraphPin)
+					// This can invalidate the source pin due to node reconstruction, abort in that case
+					if(SourcePin->GetOwningNodeUnchecked() && EdGraphPin)
 					{
 						SourcePin->Modify();
 						EdGraphPin->Modify();
