Back out 3911900
Font size changes will once again immediately update comment node title sizes

diff --git a/Engine/Source/Editor/GraphEditor/Private/GraphEditorSettings.cpp b/Engine/Source/Editor/GraphEditor/Private/GraphEditorSettings.cpp
--- a/Engine/Source/Editor/GraphEditor/Private/GraphEditorSettings.cpp
+++ b/Engine/Source/Editor/GraphEditor/Private/GraphEditorSettings.cpp
@@ -16,7 +16,6 @@ UGraphEditorSettings::UGraphEditorSettings( const FObjectInitializer& ObjectInit
 	, PaddingRightOfInput(10.0f)
 	, PaddingLeftOfOutput(10.0f)
 	, PaddingTowardsNodeEdge(10.0f)
-	, bShowCommentBubbles(false)
 	, bTreatSplinesLikePins(true)
 	, SplineHoverTolerance(2.0f)
 	, ForwardSplineHorizontalDeltaRange(1000.0f)
diff --git a/Engine/Source/Editor/GraphEditor/Private/SGraphNodeComment.cpp b/Engine/Source/Editor/GraphEditor/Private/SGraphNodeComment.cpp
--- a/Engine/Source/Editor/GraphEditor/Private/SGraphNodeComment.cpp
+++ b/Engine/Source/Editor/GraphEditor/Private/SGraphNodeComment.cpp
@@ -68,37 +68,22 @@ void SGraphNodeComment::Tick( const FGeometry& AllottedGeometry, const double In
 		CachedCommentTitle = CurrentCommentTitle;
 	}
 
-
 	const int32 CurrentWidth = static_cast<int32>(UserSize.X);
 	if (CurrentWidth != CachedWidth)
 	{
 		CachedWidth = CurrentWidth;
 	}
 
-	const bool bShowCommentBubbles = GetDefault<UGraphEditorSettings>()->bShowCommentBubbles;
-	if (bShowCommentBubbles)
+	UEdGraphNode_Comment* CommentNode = CastChecked<UEdGraphNode_Comment>(GraphNode);
+	if (bCachedBubbleVisibility != CommentNode->bCommentBubbleVisible_InDetailsPanel)
 	{
-		UEdGraphNode_Comment* CommentNode = CastChecked<UEdGraphNode_Comment>(GraphNode);
-		if (bCachedBubbleVisibility != CommentNode->bCommentBubbleVisible_InDetailsPanel)
-		{
-			GetCommentBubble()->UpdateBubble();
-			bCachedBubbleVisibility = CommentNode->bCommentBubbleVisible_InDetailsPanel;
-		}
-
-		if (CachedFontSize != CommentNode->FontSize)
-		{
-			UpdateGraphNode();
-		}
+		CommentBubble->UpdateBubble();
+		bCachedBubbleVisibility = CommentNode->bCommentBubbleVisible_InDetailsPanel;
 	}
-	else
-	{
-		if (bCachedBubbleVisibility && CommentBubble.IsValid())
-		{
-			CommentBubble->SetVisibility(EVisibility::Collapsed);
-		}
 
-		CommentBubble.Reset();
-		bCachedBubbleVisibility = false;
+	if (CachedFontSize != CommentNode->FontSize)
+	{
+		UpdateGraphNode();
 	}
 }
 
@@ -117,36 +102,6 @@ float SGraphNodeComment::GetWrapAt() const
 	return (float)(CachedWidth - 32.0f);
 }
 
-TSharedPtr<SCommentBubble> SGraphNodeComment::GetCommentBubble()
-{
-	if (!CommentBubble.IsValid())
-	{
-		// Create comment bubble
-		CommentBubble = SNew(SCommentBubble)
-			.GraphNode(GraphNode)
-			.Text(this, &SGraphNodeComment::GetNodeComment)
-			.OnTextCommitted(this, &SGraphNodeComment::OnNameTextCommited)
-			.ColorAndOpacity(this, &SGraphNodeComment::GetCommentBubbleColor)
-			.AllowPinning(true)
-			.EnableTitleBarBubble(false)
-			.EnableBubbleCtrls(false)
-			.GraphLOD(this, &SGraphNode::GetCurrentLOD)
-			.InvertLODCulling(true)
-			.IsGraphNodeHovered(this, &SGraphNode::IsHovered);
-
-		GetOrAddSlot(ENodeZone::TopCenter)
-			.SlotOffset(TAttribute<FVector2D>(CommentBubble.Get(), &SCommentBubble::GetOffset))
-			.SlotSize(TAttribute<FVector2D>(CommentBubble.Get(), &SCommentBubble::GetSize))
-			.AllowScaling(TAttribute<bool>(CommentBubble.Get(), &SCommentBubble::IsScalingAllowed))
-			.VAlign(VAlign_Top)
-			[
-				CommentBubble.ToSharedRef()
-			];
-	}
-
-	return CommentBubble;
-}
-
 bool SGraphNodeComment::IsNameReadOnly() const
 {
 	return !IsEditable.Get() || SGraphNode::IsNameReadOnly();
@@ -154,8 +109,6 @@ bool SGraphNodeComment::IsNameReadOnly() const
 
 void SGraphNodeComment::UpdateGraphNode()
 {
-	const bool bShowCommentBubbles = GetDefault<UGraphEditorSettings>()->bShowCommentBubbles;
-
 	// No pins in a comment box
 	InputPins.Empty();
 	OutputPins.Empty();
@@ -166,7 +119,7 @@ void SGraphNodeComment::UpdateGraphNode()
 
 	// Remember if we should be showing the bubble
 	UEdGraphNode_Comment* CommentNode = CastChecked<UEdGraphNode_Comment>(GraphNode);
-	bCachedBubbleVisibility = bShowCommentBubbles ? CommentNode->bCommentBubbleVisible_InDetailsPanel : false;
+	bCachedBubbleVisibility = CommentNode->bCommentBubbleVisible_InDetailsPanel;
 
 	// Setup a tag for this node
 	FString TagName;
@@ -248,11 +201,27 @@ void SGraphNodeComment::UpdateGraphNode()
 			]			
 		];
 
-	if (bShowCommentBubbles)
-	{
-		// Create comment bubble
-		GetCommentBubble();
-	}
+	// Create comment bubble
+	CommentBubble = SNew(SCommentBubble)
+	.GraphNode(GraphNode)
+	.Text(this, &SGraphNodeComment::GetNodeComment)
+	.OnTextCommitted(this, &SGraphNodeComment::OnNameTextCommited)
+	.ColorAndOpacity(this, &SGraphNodeComment::GetCommentBubbleColor )
+	.AllowPinning(true)
+	.EnableTitleBarBubble(false)
+	.EnableBubbleCtrls(false)
+	.GraphLOD(this, &SGraphNode::GetCurrentLOD)
+	.InvertLODCulling(true)
+	.IsGraphNodeHovered(this, &SGraphNode::IsHovered);
+
+	GetOrAddSlot(ENodeZone::TopCenter)
+	.SlotOffset(TAttribute<FVector2D>(CommentBubble.Get(), &SCommentBubble::GetOffset))
+	.SlotSize(TAttribute<FVector2D>(CommentBubble.Get(), &SCommentBubble::GetSize))
+	.AllowScaling(TAttribute<bool>(CommentBubble.Get(), &SCommentBubble::IsScalingAllowed))
+	.VAlign(VAlign_Top)
+	[
+		CommentBubble.ToSharedRef()
+	];
 }
 
 FVector2D SGraphNodeComment::ComputeDesiredSize( float ) const
diff --git a/Engine/Source/Editor/GraphEditor/Public/GraphEditorSettings.h b/Engine/Source/Editor/GraphEditor/Public/GraphEditorSettings.h
--- a/Engine/Source/Editor/GraphEditor/Public/GraphEditorSettings.h
+++ b/Engine/Source/Editor/GraphEditor/Public/GraphEditorSettings.h
@@ -65,10 +65,6 @@ public:
 	UPROPERTY(config, EditAnywhere, Category=GeneralStyle, AdvancedDisplay)
 	float PaddingTowardsNodeEdge;
 
-	/** Whether the comment bubbles for comment nodes are visible or hidden. */
-	UPROPERTY(config, EditAnywhere, Category = GeneralStyle, AdvancedDisplay, meta = (DisplayName = "Show Comment Bubbles"))
-	bool bShowCommentBubbles;
-
 	/** Returns the padding to use for an input pin on a graph node */
 	FMargin GetInputPinPadding() const;
 
diff --git a/Engine/Source/Editor/GraphEditor/Public/SGraphNodeComment.h b/Engine/Source/Editor/GraphEditor/Public/SGraphNodeComment.h
--- a/Engine/Source/Editor/GraphEditor/Public/SGraphNodeComment.h
+++ b/Engine/Source/Editor/GraphEditor/Public/SGraphNodeComment.h
@@ -92,10 +92,6 @@ private:
 	float GetWrapAt() const;
 
 private:
-
-	/** Gets the comment bubble and creates one if it doesn't exist. */
-	TSharedPtr<SCommentBubble> GetCommentBubble();
-
 	/** The comment bubble widget (used when zoomed out) */
 	TSharedPtr<SCommentBubble> CommentBubble;
 
