PR #5017: UE-63105: Modify SGraphActionMenu::OnKeyDown to use a different branc. (Contributed by [REDACTED_PERSON])

diff --git a/Engine/Source/Editor/GraphEditor/Private/SGraphActionMenu.cpp b/Engine/Source/Editor/GraphEditor/Private/SGraphActionMenu.cpp
--- a/Engine/Source/Editor/GraphEditor/Private/SGraphActionMenu.cpp
+++ b/Engine/Source/Editor/GraphEditor/Private/SGraphActionMenu.cpp
@@ -1276,8 +1276,14 @@ FReply SGraphActionMenu::OnKeyDown( const FGeometry& MyGeometry, const FKeyEvent
 	{
 		return TryToSpawnActiveSuggestion() ? FReply::Handled() : FReply::Unhandled();
 	}
-	else if (FilteredActionNodes.Num() > 0 && !FilterTextBox->GetText().IsEmpty())
+	else if (!FilterTextBox->GetText().IsEmpty())
 	{
+		// Needs to be done here in order not to eat up the text navigation key events when list isn't populated
+		if (FilteredActionNodes.Num() <= 0)
+		{
+			return FReply::Unhandled();
+		}
+
 		// Up and down move thru the filtered node list
 		if (KeyEvent.GetKey() == EKeys::Up)
 		{
