Fix auto expanding all the sections for the niagara parameters list view.

diff --git a/Engine/Source/Editor/GraphEditor/Private/SGraphActionMenu.cpp b/Engine/Source/Editor/GraphEditor/Private/SGraphActionMenu.cpp
--- a/Engine/Source/Editor/GraphEditor/Private/SGraphActionMenu.cpp
+++ b/Engine/Source/Editor/GraphEditor/Private/SGraphActionMenu.cpp
@@ -588,6 +588,11 @@ static bool CompareGraphActionNode(TSharedPtr<FGraphActionNode> A, TSharedPtr<FG
 		return false;
 	}
 
+	/*if (A->SectionID != B->SectionID)
+	{
+		return false;
+	}*/
+
 	if (A->HasValidAction() && B->HasValidAction())
 	{
 		return A->GetPrimaryAction()->GetMenuDescription().CompareTo(B->GetPrimaryAction()->GetMenuDescription()) == 0;
