Fix auto expanding all the sections for the niagara parameters list view. Fix incorrect comment changelist: 4102752

diff --git a/Engine/Source/Editor/GraphEditor/Private/SGraphActionMenu.cpp b/Engine/Source/Editor/GraphEditor/Private/SGraphActionMenu.cpp
--- a/Engine/Source/Editor/GraphEditor/Private/SGraphActionMenu.cpp
+++ b/Engine/Source/Editor/GraphEditor/Private/SGraphActionMenu.cpp
@@ -588,10 +588,10 @@ static bool CompareGraphActionNode(TSharedPtr<FGraphActionNode> A, TSharedPtr<FG
 		return false;
 	}
 
-	/*if (A->SectionID != B->SectionID)
+	if (A->SectionID != B->SectionID)
 	{
 		return false;
-	}*/
+	}
 
 	if (A->HasValidAction() && B->HasValidAction())
 	{
