Remove unnecessary null check causing static analysis warning in SGraphPinClass

diff --git a/Engine/Source/Editor/GraphEditor/Private/KismetPins/SGraphPinClass.cpp b/Engine/Source/Editor/GraphEditor/Private/KismetPins/SGraphPinClass.cpp
--- a/Engine/Source/Editor/GraphEditor/Private/KismetPins/SGraphPinClass.cpp
+++ b/Engine/Source/Editor/GraphEditor/Private/KismetPins/SGraphPinClass.cpp
@@ -67,7 +67,7 @@ public:
 		
 			// Don't allow classes from a loaded map (e.g. LSBPs) unless we're already working inside that package context. Otherwise, choosing the class would lead to a GLEO at save time.
 			Result &= !ClassPackage->ContainsMap() || ClassPackage == GraphPinOutermostPackage;
-			Result &= bAllowAbstractClasses || (InClass != nullptr && !InClass->HasAnyClassFlags(CLASS_Abstract));
+			Result &= bAllowAbstractClasses || !InClass->HasAnyClassFlags(CLASS_Abstract);
 		}
 
 		return Result;
