PR #5627: Fixed Microsoft mouse scroll wheel issue in blueprint editor. (Contributed by [REDACTED_PERSON])

Reimplemented according to comment from author.

diff --git a/Engine/Source/Editor/GraphEditor/Private/SNodePanel.cpp b/Engine/Source/Editor/GraphEditor/Private/SNodePanel.cpp
--- a/Engine/Source/Editor/GraphEditor/Private/SNodePanel.cpp
+++ b/Engine/Source/Editor/GraphEditor/Private/SNodePanel.cpp
@@ -978,8 +978,8 @@ FReply SNodePanel::OnMouseWheel(const FGeometry& MyGeometry, const FPointerEvent
 {
 	// We want to zoom into this point; i.e. keep it the same fraction offset into the panel
 	const FVector2D WidgetSpaceCursorPos = MyGeometry.AbsoluteToLocal( MouseEvent.GetScreenSpacePosition() );
-	const int32 ZoomLevelDelta = FMath::FloorToInt( MouseEvent.GetWheelDelta() );
-	ChangeZoomLevel(ZoomLevelDelta, WidgetSpaceCursorPos, MouseEvent.IsControlDown());
+	const int32 ZoomLevelDelta = FMath::TruncToInt( FMath::RoundFromZero( MouseEvent.GetWheelDelta() ) );
+	ChangeZoomLevel( ZoomLevelDelta, WidgetSpaceCursorPos, MouseEvent.IsControlDown() );
 
 	// Stop the zoom-to-fit in favor of user control
 	CancelZoomToFit();
