Fix flicker in graph preview in anim state machines

Stopped zooming every frame, and only update the graph when it is created.

diff --git a/Engine/Source/Editor/GraphEditor/Private/SGraphPreviewer.cpp b/Engine/Source/Editor/GraphEditor/Private/SGraphPreviewer.cpp
--- a/Engine/Source/Editor/GraphEditor/Private/SGraphPreviewer.cpp
+++ b/Engine/Source/Editor/GraphEditor/Private/SGraphPreviewer.cpp
@@ -7,23 +7,24 @@
 #include "EditorStyleSet.h"
 #include "SGraphPanel.h"
 
-void SGraphPreviewer::Tick( const FGeometry& AllottedGeometry, const double InCurrentTime, const float InDeltaTime )
+EActiveTimerReturnType SGraphPreviewer::RefreshGraphTimer(const double InCurrentTime, const float InDeltaTime)
 {
-	// Keep the graph constantly zoomed to fit
-	GraphPanel->ZoomToFit(false);
-
-	// Refresh the graph if needed
-	if (bNeedsRefresh)
+	if (NeedsRefreshCounter > 0)
+	{
+		GraphPanel->ZoomToFit(false);
+		NeedsRefreshCounter--;
+		return EActiveTimerReturnType::Continue;
+	}
+	else
 	{
-		bNeedsRefresh = false;
-		GraphPanel->Update();
-	}	
+		return EActiveTimerReturnType::Stop;
+	}
 }
 
 void SGraphPreviewer::Construct( const FArguments& InArgs, UEdGraph* InGraphObj )
 {
 	EdGraphObj = InGraphObj;
-	bNeedsRefresh = true;
+	NeedsRefreshCounter = 2;
 
 	TSharedPtr<SOverlay> DisplayStack;
 
@@ -54,6 +55,8 @@ void SGraphPreviewer::Construct( const FArguments& InArgs, UEdGraph* InGraphObj
 		]
 	];
 
+	GraphPanel->Update();
+
 	// Add the title bar if specified
 	if (InArgs._TitleBar.IsValid())
 	{
@@ -63,4 +66,6 @@ void SGraphPreviewer::Construct( const FArguments& InArgs, UEdGraph* InGraphObj
 				InArgs._TitleBar.ToSharedRef()
 			];
 	}
+
+	RegisterActiveTimer(0.0f, FWidgetActiveTimerDelegate::CreateSP(this, &SGraphPreviewer::RefreshGraphTimer));
 }
diff --git a/Engine/Source/Editor/GraphEditor/Public/SGraphPreviewer.h b/Engine/Source/Editor/GraphEditor/Public/SGraphPreviewer.h
--- a/Engine/Source/Editor/GraphEditor/Public/SGraphPreviewer.h
+++ b/Engine/Source/Editor/GraphEditor/Public/SGraphPreviewer.h
@@ -23,16 +23,16 @@ public:
 	SLATE_END_ARGS()
 
 	void Construct(const FArguments& InArgs, UEdGraph* InGraphObj);
-public:
-	// SWidget interface
-	virtual void Tick( const FGeometry& AllottedGeometry, const double InCurrentTime, const float InDeltaTime ) override;
-	// End of SWidget interface
+private:
+	/** Helper function used to refresh the graph */
+	EActiveTimerReturnType RefreshGraphTimer(const double InCurrentTime, const float InDeltaTime);
+
 protected:
 	/// The Graph we are currently viewing
 	UEdGraph* EdGraphObj;
 
-	// Do we need to refresh next tick?
-	bool bNeedsRefresh;
+	// As node's bounds dont get updated immediately, to truly zoom out to fit we need to tick a few times
+	int32 NeedsRefreshCounter;
 
 	// The underlying graph panel
 	TSharedPtr<SGraphPanel> GraphPanel;
