PR #5251: Added scrolling to CreateEvent node method list (Contributed by [REDACTED_PERSON])

diff --git a/Engine/Source/Editor/GraphEditor/Private/KismetNodes/SGraphNodeK2CreateDelegate.cpp b/Engine/Source/Editor/GraphEditor/Private/KismetNodes/SGraphNodeK2CreateDelegate.cpp
--- a/Engine/Source/Editor/GraphEditor/Private/KismetNodes/SGraphNodeK2CreateDelegate.cpp
+++ b/Engine/Source/Editor/GraphEditor/Private/KismetNodes/SGraphNodeK2CreateDelegate.cpp
@@ -5,6 +5,7 @@
 #include "Widgets/Views/STableViewBase.h"
 #include "Widgets/Views/STableRow.h"
 #include "Widgets/Views/SListView.h"
+#include "Widgets/SBoxPanel.h"
 #include "EdGraphSchema_K2.h"
 #include "K2Node_CreateDelegate.h"
 
@@ -225,7 +226,7 @@ void SGraphNodeK2CreateDelegate::CreateBelowPinControls(TSharedPtr<SVerticalBox>
 			}
 
 			TSharedRef<SComboButton> SelectFunctionWidgetRef = SNew(SComboButton)
-				.Method(EPopupMethod::UseCurrentWindow)
+				.Method(EPopupMethod::CreateNewWindow)
 				.ButtonContent()
 				[
 					SNew(STextBlock)
@@ -233,10 +234,16 @@ void SGraphNodeK2CreateDelegate::CreateBelowPinControls(TSharedPtr<SVerticalBox>
 				]
 			.MenuContent()
 				[
-					SNew(SListView<TSharedPtr<FFunctionItemData> >)
-					.ListItemsSource(&FunctionDataItems)
-				.OnGenerateRow(this, &SGraphNodeK2CreateDelegate::HandleGenerateRowFunction)
-				.OnSelectionChanged(this, &SGraphNodeK2CreateDelegate::OnFunctionSelected)
+					SNew(SVerticalBox)
+ 					 + SVerticalBox::Slot()
+ 					 .AutoHeight()
+ 					 .MaxHeight(500.f)
+ 					 [
+ 						SNew(SListView<TSharedPtr<FFunctionItemData> >)
+ 						.ListItemsSource(&FunctionDataItems)
+ 						.OnGenerateRow(this, &SGraphNodeK2CreateDelegate::HandleGenerateRowFunction)
+ 						.OnSelectionChanged(this, &SGraphNodeK2CreateDelegate::OnFunctionSelected)
+ 					 ]
 				];
 
 			MainBox->AddSlot()
