(4.21.1) Don't crash if the SGraphPinObject is pointing to a now gone pin or schema.
#rnx

diff --git a/Engine/Source/Editor/GraphEditor/Private/KismetPins/SGraphPinObject.cpp b/Engine/Source/Editor/GraphEditor/Private/KismetPins/SGraphPinObject.cpp
--- a/Engine/Source/Editor/GraphEditor/Private/KismetPins/SGraphPinObject.cpp
+++ b/Engine/Source/Editor/GraphEditor/Private/KismetPins/SGraphPinObject.cpp
@@ -40,9 +40,21 @@ void SGraphPinObject::Construct(const FArguments& InArgs, UEdGraphPin* InGraphPi
 
 TSharedRef<SWidget>	SGraphPinObject::GetDefaultValueWidget()
 {
+	if (GraphPinObj == nullptr)
+	{
+		return SNullWidget::NullWidget;
+	}
+
+	const UEdGraphSchema* Schema = GraphPinObj->GetSchema();
+
+	if (Schema == nullptr)
+	{
+		return SNullWidget::NullWidget;
+	}
+
 	if( AllowSelfPinWidget() )
 	{
-		if(GraphPinObj->GetSchema()->IsSelfPin(*GraphPinObj))
+		if(Schema->IsSelfPin(*GraphPinObj))
 		{
 			return SNew(SEditableTextBox)
 				.Style( FEditorStyle::Get(), "Graph.EditableTextBox" )
@@ -54,7 +66,7 @@ TSharedRef<SWidget>	SGraphPinObject::GetDefaultValueWidget()
 		}
 	}
 	// Don't show literal buttons for component type objects
-	if (GraphPinObj->GetSchema()->ShouldShowAssetPickerForPin(GraphPinObj))
+	if (Schema->ShouldShowAssetPickerForPin(GraphPinObj))
 	{
 		return
 			SNew(SHorizontalBox)
