Adopt de factor standard graph editor shortcuts as the actual default

diff --git a/Engine/Source/Editor/GraphEditor/Private/GraphEditorActions.cpp b/Engine/Source/Editor/GraphEditor/Private/GraphEditorActions.cpp
--- a/Engine/Source/Editor/GraphEditor/Private/GraphEditorActions.cpp
+++ b/Engine/Source/Editor/GraphEditor/Private/GraphEditorActions.cpp
@@ -43,14 +43,14 @@ void FGraphEditorCommandsImpl::RegisterCommands()
 	UI_COMMAND( CollapseSelectionToFunction, "Collapse to Function", "Collapses selected nodes into a single function node.", EUserInterfaceActionType::Button, FInputChord() )
 	UI_COMMAND( CollapseSelectionToMacro, "Collapse to Macro", "Collapses selected nodes into a single macro node.", EUserInterfaceActionType::Button, FInputChord() )
 
-	UI_COMMAND( AlignNodesTop, "Align Top", "Aligns the top edges of the selected nodes", EUserInterfaceActionType::Button, FInputChord() )
-	UI_COMMAND( AlignNodesMiddle, "Align Middle", "Aligns the vertical middles of the selected nodes", EUserInterfaceActionType::Button, FInputChord() )
-	UI_COMMAND( AlignNodesBottom, "Align Bottom", "Aligns the bottom edges of the selected nodes", EUserInterfaceActionType::Button, FInputChord() )
-	UI_COMMAND( AlignNodesLeft, "Align Left", "Aligns the left edges of the selected nodes", EUserInterfaceActionType::Button, FInputChord() )
-	UI_COMMAND( AlignNodesCenter, "Align Center", "Aligns the horizontal centers of the selected nodes", EUserInterfaceActionType::Button, FInputChord() )
-	UI_COMMAND( AlignNodesRight, "Align Right", "Aligns the right edges of the selected nodes", EUserInterfaceActionType::Button, FInputChord() )
-
-	UI_COMMAND( StraightenConnections, "Straighten Connection(s)", "Straightens connections between the selected nodes.", EUserInterfaceActionType::Button, FInputChord() )
+	UI_COMMAND( AlignNodesTop, "Align Top", "Aligns the top edges of the selected nodes", EUserInterfaceActionType::Button, FInputChord(EModifierKey::Shift, EKeys::W) )
+	UI_COMMAND( AlignNodesMiddle, "Align Middle", "Aligns the vertical middles of the selected nodes", EUserInterfaceActionType::Button, FInputChord(EModifierKey::Shift|EModifierKey::Alt, EKeys::W) )
+	UI_COMMAND( AlignNodesBottom, "Align Bottom", "Aligns the bottom edges of the selected nodes", EUserInterfaceActionType::Button, FInputChord(EModifierKey::Shift, EKeys::S) )
+	UI_COMMAND( AlignNodesLeft, "Align Left", "Aligns the left edges of the selected nodes", EUserInterfaceActionType::Button, FInputChord(EModifierKey::Shift, EKeys::A) )
+	UI_COMMAND( AlignNodesCenter, "Align Center", "Aligns the horizontal centers of the selected nodes", EUserInterfaceActionType::Button, FInputChord(EModifierKey::Shift | EModifierKey::Alt, EKeys::S) )
+	UI_COMMAND( AlignNodesRight, "Align Right", "Aligns the right edges of the selected nodes", EUserInterfaceActionType::Button, FInputChord(EModifierKey::Shift, EKeys::D) )
+
+	UI_COMMAND( StraightenConnections, "Straighten Connection(s)", "Straightens connections between the selected nodes.", EUserInterfaceActionType::Button, FInputChord(EKeys::Q) )
 
 	UI_COMMAND( DistributeNodesHorizontally, "Distribute Horizontally", "Evenly distributes the selected nodes horizontally", EUserInterfaceActionType::Button, FInputChord() )
 	UI_COMMAND( DistributeNodesVertically, "Distribute Vertically", "Evenly distributes the selected nodes vertically", EUserInterfaceActionType::Button, FInputChord() )
