PR #5316: UE-66912: Prevent multiple FScopedTransactions for comment boxes (Contributed by [REDACTED_PERSON])

diff --git a/Engine/Source/Editor/GraphEditor/Private/SCommentBubble.cpp b/Engine/Source/Editor/GraphEditor/Private/SCommentBubble.cpp
--- a/Engine/Source/Editor/GraphEditor/Private/SCommentBubble.cpp
+++ b/Engine/Source/Editor/GraphEditor/Private/SCommentBubble.cpp
@@ -108,13 +108,7 @@ void SCommentBubble::Tick( const FGeometry& AllottedGeometry, const double InCur
 	{
 		CachedComment = CommentAttribute.Get();
 		CachedCommentText = FText::FromString( CachedComment );
-		// Call text commit delegate
-		OnTextCommittedDelegate.ExecuteIfBound( CachedCommentText, ETextCommit::Default );
-		// Reflect changes to the Textblock because it doesn't update itself.
-		if( TextBlock.IsValid() )
-		{
-			TextBlock->SetText( CachedCommentText );
-		}
+
 		// Toggle the comment on/off, provided it the parent isn't a comment node
 		if( !bInvertLODCulling )
 		{
@@ -252,7 +246,7 @@ void SCommentBubble::UpdateBubble()
 						.AutoWidth()
 						[
 							SAssignNew(TextBlock, SMultiLineEditableTextBox)
-							.Text( CachedCommentText )
+							.Text(MakeAttributeLambda([this] { return CachedCommentText; }))
 							.HintText( NSLOCTEXT( "CommentBubble", "EditCommentHint", "Click to edit" ))
 							.IsReadOnly( this, &SCommentBubble::IsReadOnly )
 							.Font( FEditorStyle::GetFontStyle( TEXT("Graph.Node.CommentFont")))
@@ -412,9 +406,8 @@ FSlateColor SCommentBubble::GetTextForegroundColor() const
 
 void SCommentBubble::OnCommentTextCommitted( const FText& NewText, ETextCommit::Type CommitInfo )
 {
-	if (CommitInfo != ETextCommit::OnEnter)
+	if (CommitInfo == ETextCommit::OnEnter || CommitInfo == ETextCommit::OnUserMovedFocus)
 	{
-		// Don't respond to OnEnter, as it will be immediately followed by OnCleared anyway (due to loss of keyboard focus) and generate a second transaction
 		CachedComment = NewText.ToString();
 		CachedCommentText = NewText;
 		OnTextCommittedDelegate.ExecuteIfBound(CachedCommentText, CommitInfo);
