Can now use home, end, page up, and page down in a filtered blueprint action menu

diff --git a/Engine/Source/Editor/GraphEditor/Private/SGraphActionMenu.cpp b/Engine/Source/Editor/GraphEditor/Private/SGraphActionMenu.cpp
--- a/Engine/Source/Editor/GraphEditor/Private/SGraphActionMenu.cpp
+++ b/Engine/Source/Editor/GraphEditor/Private/SGraphActionMenu.cpp
@@ -1279,38 +1279,44 @@ FReply SGraphActionMenu::OnKeyDown( const FGeometry& MyGeometry, const FKeyEvent
 	else if (!FilterTextBox->GetText().IsEmpty())
 	{
 		// Needs to be done here in order not to eat up the text navigation key events when list isn't populated
-		if (FilteredActionNodes.Num() <= 0)
+		if (FilteredActionNodes.Num() == 0)
 		{
 			return FReply::Unhandled();
 		}
 
-		// Up and down move thru the filtered node list
 		if (KeyEvent.GetKey() == EKeys::Up)
 		{
-			SelectionDelta = -1;
+			SelectedSuggestion = FMath::Max(0, SelectedSuggestion - 1);
 		}
 		else if (KeyEvent.GetKey() == EKeys::Down)
 		{
-			SelectionDelta = +1;
+			SelectedSuggestion = FMath::Min(FilteredActionNodes.Num() - 1, SelectedSuggestion + 1);
 		}
-
-		if (SelectionDelta != 0)
+		if (KeyEvent.GetKey() == EKeys::PageUp)
 		{
-			// If we have no selected suggestion then we need to use the items in the root to set the selection and set the focus
-			if( SelectedSuggestion == INDEX_NONE )
-			{
-				SelectedSuggestion = (SelectedSuggestion + SelectionDelta + FilteredRootAction->Children.Num()) % FilteredRootAction->Children.Num();
-				MarkActiveSuggestion();
-				return FReply::Handled();
-			}
-
-			//Move up or down one, wrapping around
-			SelectedSuggestion = (SelectedSuggestion + SelectionDelta + FilteredActionNodes.Num()) % FilteredActionNodes.Num();
-
-			MarkActiveSuggestion();
-
-			return FReply::Handled();
+			const int32 NumItemsInAPage = 15; // arbitrary jump because we can't get at the visible item count from here
+			SelectedSuggestion = FMath::Max(0, SelectedSuggestion - NumItemsInAPage);
+		}
+		else if (KeyEvent.GetKey() == EKeys::PageDown)
+		{
+			const int32 NumItemsInAPage = 15; // arbitrary jump because we can't get at the visible item count from here
+			SelectedSuggestion = FMath::Min(FilteredActionNodes.Num() - 1, SelectedSuggestion + NumItemsInAPage);
+		}
+		else if (KeyEvent.GetKey() == EKeys::Home)
+		{
+			SelectedSuggestion = 0;
+		}
+		else if (KeyEvent.GetKey() == EKeys::End)
+		{
+			SelectedSuggestion = FilteredActionNodes.Num() - 1;
+		}
+		else
+		{
+			return FReply::Unhandled();
 		}
+
+		MarkActiveSuggestion();
+		return FReply::Handled();
 	}
 	else
 	{
