Fix missing else causing bad behavior when using up/down in the filtered graph action menu
#rnx

diff --git a/Engine/Source/Editor/GraphEditor/Private/SGraphActionMenu.cpp b/Engine/Source/Editor/GraphEditor/Private/SGraphActionMenu.cpp
--- a/Engine/Source/Editor/GraphEditor/Private/SGraphActionMenu.cpp
+++ b/Engine/Source/Editor/GraphEditor/Private/SGraphActionMenu.cpp
@@ -1292,7 +1292,7 @@ FReply SGraphActionMenu::OnKeyDown( const FGeometry& MyGeometry, const FKeyEvent
 		{
 			SelectedSuggestion = FMath::Min(FilteredActionNodes.Num() - 1, SelectedSuggestion + 1);
 		}
-		if (KeyEvent.GetKey() == EKeys::PageUp)
+		else if (KeyEvent.GetKey() == EKeys::PageUp)
 		{
 			const int32 NumItemsInAPage = 15; // arbitrary jump because we can't get at the visible item count from here
 			SelectedSuggestion = FMath::Max(0, SelectedSuggestion - NumItemsInAPage);
