PR #5107: UE-64219: Use correct target node to add pin node to (Contributed by [REDACTED_PERSON])

diff --git a/Engine/Source/Editor/BlueprintGraph/Private/K2Node_Tunnel.cpp b/Engine/Source/Editor/BlueprintGraph/Private/K2Node_Tunnel.cpp
--- a/Engine/Source/Editor/BlueprintGraph/Private/K2Node_Tunnel.cpp
+++ b/Engine/Source/Editor/BlueprintGraph/Private/K2Node_Tunnel.cpp
@@ -202,9 +202,9 @@ UEdGraphPin* UK2Node_Tunnel::CreatePinFromUserDefinition(const TSharedPtr<FUserP
 			TargetNode->CreatePin(UEdGraphPin::GetComplementaryDirection(Direction), NewPinInfo->PinType, NewPinInfo->PinName);
 		}
 	}
-	else if (UK2Node_Tunnel* TunnelNode = Cast<UK2Node_Tunnel>(TargetNode))
+	else if (UK2Node_Tunnel* TunnelNode = Cast<UK2Node_Tunnel>((Direction == EGPD_Input) ? InputSinkNode : OutputSourceNode))
 	{
-		TunnelNode->CreateUserDefinedPin(NewPinInfo->PinName, NewPinInfo->PinType, (Direction == EGPD_Input)? EGPD_Output : EGPD_Input);
+		TunnelNode->CreateUserDefinedPin(NewPinInfo->PinName, NewPinInfo->PinType, UEdGraphPin::GetComplementaryDirection(Direction));
 	}
 
 	//@TODO: Automatically update loaded macro instances when this node is changed too
diff --git a/Engine/Source/Editor/GraphEditor/Private/DragConnection.cpp b/Engine/Source/Editor/GraphEditor/Private/DragConnection.cpp
--- a/Engine/Source/Editor/GraphEditor/Private/DragConnection.cpp
+++ b/Engine/Source/Editor/GraphEditor/Private/DragConnection.cpp
@@ -247,7 +247,6 @@ FReply FDragConnection::DroppedOnNode(FVector2D ScreenPosition, FVector2D GraphP
 		TArray<UEdGraphPin*> ValidSourcePins;
 		ValidateGraphPinList(/*out*/ ValidSourcePins);
 
-
 		if (ValidSourcePins.Num())
 		{
 			for (UEdGraphPin* SourcePin : ValidSourcePins)
@@ -261,7 +260,7 @@ FReply FDragConnection::DroppedOnNode(FVector2D ScreenPosition, FVector2D GraphP
 					// Find which pin name to use and drop the pin on the node
 					const FName PinName = SourcePin->PinFriendlyName.IsEmpty()? SourcePin->PinName : *SourcePin->PinFriendlyName.ToString();
 
-					const FScopedTransaction Transaction( NSLOCTEXT("UnrealEd", "AddInParam", "Add In Parameter" ) );
+					const FScopedTransaction Transaction((SourcePin->Direction == EGPD_Output) ? NSLOCTEXT("UnrealEd", "AddInParam", "Add In Parameter" ) : NSLOCTEXT("UnrealEd", "AddOutParam", "Add Out Parameter"));
 
 					UEdGraphPin* EdGraphPin = NodeOver->GetSchema()->DropPinOnNode(GetHoveredNode(), PinName, SourcePin->PinType, SourcePin->Direction);
 
