Fixed FKey::GetDisplayName returning an improper value when using a blueprint setter node for a FKey variable without modifying the input pin

#UE-67116

diff --git a/Engine/Source/Editor/GraphEditor/Private/KismetPins/SGraphPinKey.cpp b/Engine/Source/Editor/GraphEditor/Private/KismetPins/SGraphPinKey.cpp
--- a/Engine/Source/Editor/GraphEditor/Private/KismetPins/SGraphPinKey.cpp
+++ b/Engine/Source/Editor/GraphEditor/Private/KismetPins/SGraphPinKey.cpp
@@ -18,6 +18,11 @@ void SGraphPinKey::Construct(const FArguments& InArgs, UEdGraphPin* InGraphPinOb
 		if (SelectedKey.GetFName() == NAME_None)
 		{
 			SelectedKey = EKeys::Invalid;
+
+			if (InGraphPinObj->GetSchema()->DoesDefaultValueMatchAutogenerated(*InGraphPinObj))
+			{
+				InGraphPinObj->GetSchema()->TrySetDefaultValue(*InGraphPinObj, SelectedKey.ToString());
+			}
 		}
 		else if (!SelectedKey.IsValid())
 		{
