Fixed blueprint being marked dirty by FKey set nodes

diff --git a/Engine/Source/Editor/GraphEditor/Private/KismetPins/SGraphPinKey.cpp b/Engine/Source/Editor/GraphEditor/Private/KismetPins/SGraphPinKey.cpp
--- a/Engine/Source/Editor/GraphEditor/Private/KismetPins/SGraphPinKey.cpp
+++ b/Engine/Source/Editor/GraphEditor/Private/KismetPins/SGraphPinKey.cpp
@@ -7,6 +7,10 @@
 
 void SGraphPinKey::Construct(const FArguments& InArgs, UEdGraphPin* InGraphPinObj)
 {
+	InGraphPinObj->AutogeneratedDefaultValue = "None";
+	InGraphPinObj->GetSchema()->ResetPinToAutogeneratedDefaultValue(InGraphPinObj, false);
+
+
 	TArray<FKey> KeyList;
 	EKeys::GetAllKeys(KeyList);
 
@@ -18,11 +22,6 @@ void SGraphPinKey::Construct(const FArguments& InArgs, UEdGraphPin* InGraphPinOb
 		if (SelectedKey.GetFName() == NAME_None)
 		{
 			SelectedKey = EKeys::Invalid;
-
-			if (InGraphPinObj->GetSchema()->DoesDefaultValueMatchAutogenerated(*InGraphPinObj))
-			{
-				InGraphPinObj->GetSchema()->TrySetDefaultValue(*InGraphPinObj, SelectedKey.ToString());
-			}
 		}
 		else if (!SelectedKey.IsValid())
 		{
