PR #5397: UE-67873: Prevent early out on mark pin connection when shift-dragging (Contributed by [REDACTED_PERSON])

diff --git a/Engine/Source/Editor/GraphEditor/Private/SGraphPanel.cpp b/Engine/Source/Editor/GraphEditor/Private/SGraphPanel.cpp
--- a/Engine/Source/Editor/GraphEditor/Private/SGraphPanel.cpp
+++ b/Engine/Source/Editor/GraphEditor/Private/SGraphPanel.cpp
@@ -624,7 +624,7 @@ FReply SGraphPanel::OnMouseButtonDown(const FGeometry& MyGeometry, const FPointe
 
 FReply SGraphPanel::OnMouseButtonUp(const FGeometry& MyGeometry, const FPointerEvent& MouseEvent)
 {
-	if ((MouseEvent.GetEffectingButton() == EKeys::LeftMouseButton) && (MouseEvent.IsShiftDown()))
+	if (!NodeUnderMousePtr.IsValid() && !Marquee.IsValid() && (MouseEvent.GetEffectingButton() == EKeys::LeftMouseButton) && (MouseEvent.IsShiftDown()))
 	{
 		if (SGraphPin* BestPinFromHoveredSpline = GetBestPinFromHoveredSpline())
 		{
