Home/End to jump to top and bottom of filter list now requires control to also be held to allow home/end to allow navigation within the text field
#rnx

diff --git a/Engine/Source/Editor/GraphEditor/Private/SGraphActionMenu.cpp b/Engine/Source/Editor/GraphEditor/Private/SGraphActionMenu.cpp
--- a/Engine/Source/Editor/GraphEditor/Private/SGraphActionMenu.cpp
+++ b/Engine/Source/Editor/GraphEditor/Private/SGraphActionMenu.cpp
@@ -1302,11 +1302,11 @@ FReply SGraphActionMenu::OnKeyDown( const FGeometry& MyGeometry, const FKeyEvent
 			const int32 NumItemsInAPage = 15; // arbitrary jump because we can't get at the visible item count from here
 			SelectedSuggestion = FMath::Min(FilteredActionNodes.Num() - 1, SelectedSuggestion + NumItemsInAPage);
 		}
-		else if (KeyEvent.GetKey() == EKeys::Home)
+		else if (KeyEvent.GetKey() == EKeys::Home && KeyEvent.IsControlDown())
 		{
 			SelectedSuggestion = 0;
 		}
-		else if (KeyEvent.GetKey() == EKeys::End)
+		else if (KeyEvent.GetKey() == EKeys::End && KeyEvent.IsControlDown())
 		{
 			SelectedSuggestion = FilteredActionNodes.Num() - 1;
 		}
