Fix assets being dragged over pins displaying incorrect message

Pins were only setting tooltip on 'enter', while nodes and graphs set it on 'hover'. Switched pins to hover as well.

diff --git a/Engine/Source/Editor/GraphEditor/Private/SGraphPin.cpp b/Engine/Source/Editor/GraphEditor/Private/SGraphPin.cpp
--- a/Engine/Source/Editor/GraphEditor/Private/SGraphPin.cpp
+++ b/Engine/Source/Editor/GraphEditor/Private/SGraphPin.cpp
@@ -657,26 +657,6 @@ void SGraphPin::OnDragEnter( const FGeometry& MyGeometry, const FDragDropEvent&
 		// Pins treat being dragged over the same as being hovered outside of drag and drop if they know how to respond to the drag action.
 		SBorder::OnMouseEnter( MyGeometry, DragDropEvent );
 	}
-	else if (Operation->IsOfType<FAssetDragDropOp>())
-	{
-		TSharedPtr<SGraphNode> NodeWidget = OwnerNodePtr.Pin();
-		if (NodeWidget.IsValid())
-		{
-			UEdGraphNode* Node = NodeWidget->GetNodeObj();
-			if(Node != NULL && Node->GetSchema() != NULL)
-			{
-				TSharedPtr<FAssetDragDropOp> AssetOp = StaticCastSharedPtr<FAssetDragDropOp>(Operation);
-				bool bOkIcon = false;
-				FString TooltipText;
-				if (AssetOp->HasAssets())
-				{
-					Node->GetSchema()->GetAssetsPinHoverMessage(AssetOp->GetAssets(), GraphPinObj, TooltipText, bOkIcon);
-				}
-				const FSlateBrush* TooltipIcon = bOkIcon ? FEditorStyle::GetBrush(TEXT("Graph.ConnectorFeedback.OK")) : FEditorStyle::GetBrush(TEXT("Graph.ConnectorFeedback.Error"));;
-				AssetOp->SetToolTip(FText::FromString(TooltipText), TooltipIcon);
-			}
-		}
-	}
 }
 
 void SGraphPin::OnDragLeave( const FDragDropEvent& DragDropEvent )
@@ -706,6 +686,35 @@ void SGraphPin::OnDragLeave( const FDragDropEvent& DragDropEvent )
 
 FReply SGraphPin::OnDragOver( const FGeometry& MyGeometry, const FDragDropEvent& DragDropEvent )
 {
+	TSharedPtr<FDragDropOperation> Operation = DragDropEvent.GetOperation();
+	if (!Operation.IsValid())
+	{
+		return FReply::Unhandled();
+	}
+
+	if (Operation->IsOfType<FAssetDragDropOp>())
+	{
+		TSharedPtr<SGraphNode> NodeWidget = OwnerNodePtr.Pin();
+		if (NodeWidget.IsValid())
+		{
+			UEdGraphNode* Node = NodeWidget->GetNodeObj();
+			if(Node != NULL && Node->GetSchema() != NULL)
+			{
+				TSharedPtr<FAssetDragDropOp> AssetOp = StaticCastSharedPtr<FAssetDragDropOp>(Operation);
+				bool bOkIcon = false;
+				FString TooltipText;
+				if (AssetOp->HasAssets())
+				{
+					Node->GetSchema()->GetAssetsPinHoverMessage(AssetOp->GetAssets(), GraphPinObj, TooltipText, bOkIcon);
+				}
+				const FSlateBrush* TooltipIcon = bOkIcon ? FEditorStyle::GetBrush(TEXT("Graph.ConnectorFeedback.OK")) : FEditorStyle::GetBrush(TEXT("Graph.ConnectorFeedback.Error"));;
+				AssetOp->SetToolTip(FText::FromString(TooltipText), TooltipIcon);
+					
+				return FReply::Handled();
+			}
+		}
+	}
+
 	return FReply::Unhandled();
 }
 
