PR #5456: UE-68484: Return Handled in SGraphPin::OnDragOver (Contributed by [REDACTED_PERSON])

diff --git a/Engine/Source/Editor/GraphEditor/Private/SGraphPin.cpp b/Engine/Source/Editor/GraphEditor/Private/SGraphPin.cpp
--- a/Engine/Source/Editor/GraphEditor/Private/SGraphPin.cpp
+++ b/Engine/Source/Editor/GraphEditor/Private/SGraphPin.cpp
@@ -706,7 +706,8 @@ void SGraphPin::OnDragLeave( const FDragDropEvent& DragDropEvent )
 
 FReply SGraphPin::OnDragOver( const FGeometry& MyGeometry, const FDragDropEvent& DragDropEvent )
 {
-	return FReply::Unhandled();
+	// This is considered handled as the functionality is done OnDragEnter/OnDragLeave
+	return FReply::Handled();
 }
 
 bool SGraphPin::TryHandlePinConnection(SGraphPin& OtherSPin)
