Rolling back a PR that affected how preview wires draw in the graph editor

diff --git a/Engine/Source/Editor/GraphEditor/Private/SGraphPin.cpp b/Engine/Source/Editor/GraphEditor/Private/SGraphPin.cpp
--- a/Engine/Source/Editor/GraphEditor/Private/SGraphPin.cpp
+++ b/Engine/Source/Editor/GraphEditor/Private/SGraphPin.cpp
@@ -706,8 +706,8 @@ void SGraphPin::OnDragLeave( const FDragDropEvent& DragDropEvent )
 
 FReply SGraphPin::OnDragOver( const FGeometry& MyGeometry, const FDragDropEvent& DragDropEvent )
 {
-	// This is considered handled as the functionality is done OnDragEnter/OnDragLeave
-	return FReply::Handled();
+	// This needs to remain unhandled or there is a frame of the preview wire dragging incorrectly
+	return FReply::Unhandled();
 }
 
 bool SGraphPin::TryHandlePinConnection(SGraphPin& OtherSPin)
