Fixed crash in SGraphNode related to tooltips. SGraphNodeK2Base had tooltip related code that kept the widget alive. Added some safety mechanisms to detect this state

diff --git a/Engine/Source/Editor/GraphEditor/Private/KismetNodes/SGraphNodeK2Base.cpp b/Engine/Source/Editor/GraphEditor/Private/KismetNodes/SGraphNodeK2Base.cpp
--- a/Engine/Source/Editor/GraphEditor/Private/KismetNodes/SGraphNodeK2Base.cpp
+++ b/Engine/Source/Editor/GraphEditor/Private/KismetNodes/SGraphNodeK2Base.cpp
@@ -223,14 +223,16 @@ TSharedPtr<SToolTip> SGraphNodeK2Base::GetComplexTooltip()
 
 	struct LocalUtils
 	{
-		static EVisibility IsToolTipVisible(TSharedRef<SGraphNodeK2Base> const NodeWidget)
+		static EVisibility IsToolTipVisible(TWeakPtr<SGraphNodeK2Base> const NodeWidget)
 		{
-			return NodeWidget->GetNodeTooltip().IsEmpty() ? EVisibility::Collapsed : EVisibility::Visible;
+			TSharedPtr<SGraphNodeK2Base> NodeWidgetPinned = NodeWidget.Pin();
+			return NodeWidgetPinned && NodeWidgetPinned->GetNodeTooltip().IsEmpty() ? EVisibility::Collapsed : EVisibility::Visible;
 		}
 
-		static EVisibility IsToolTipHeadingVisible(TSharedRef<SGraphNodeK2Base> const NodeWidget)
+		static EVisibility IsToolTipHeadingVisible(TWeakPtr<SGraphNodeK2Base> const NodeWidget)
 		{
-			return NodeWidget->GetToolTipHeading().IsEmpty() ? EVisibility::Collapsed : EVisibility::Visible;
+			TSharedPtr<SGraphNodeK2Base> NodeWidgetPinned = NodeWidget.Pin();
+			return NodeWidgetPinned && NodeWidgetPinned->GetToolTipHeading().IsEmpty() ? EVisibility::Collapsed : EVisibility::Visible;
 		}
 
 		static bool IsInteractive()
@@ -239,11 +241,11 @@ TSharedPtr<SToolTip> SGraphNodeK2Base::GetComplexTooltip()
 			return ( ModifierKeys.IsAltDown() && ModifierKeys.IsControlDown() );
 		}
 	};
-	TSharedRef<SGraphNodeK2Base> ThisRef = SharedThis(this);
+	TWeakPtr<SGraphNodeK2Base> ThisWeak = SharedThis(this);
 
 	TSharedPtr< SVerticalBox > VerticalBoxWidget;
 	SAssignNew(NodeToolTip, SToolTip)
-		.Visibility_Static(&LocalUtils::IsToolTipVisible, ThisRef)
+		.Visibility_Static(&LocalUtils::IsToolTipVisible, ThisWeak)
 		.IsInteractive_Static(&LocalUtils::IsInteractive)
 
 		// Emulate text-only tool-tip styling that SToolTip uses when no custom content is supplied.  We want node tool-tips to 
@@ -256,7 +258,7 @@ TSharedPtr<SToolTip> SGraphNodeK2Base::GetComplexTooltip()
 		+SVerticalBox::Slot()
 		[
 			SNew(SVerticalBox)
-				.Visibility_Static(&LocalUtils::IsToolTipHeadingVisible, ThisRef)
+				.Visibility_Static(&LocalUtils::IsToolTipHeadingVisible, ThisWeak)
 			+SVerticalBox::Slot()
 				.AutoHeight()
 			[
diff --git a/Engine/Source/Editor/GraphEditor/Private/SGraphNode.cpp b/Engine/Source/Editor/GraphEditor/Private/SGraphNode.cpp
--- a/Engine/Source/Editor/GraphEditor/Private/SGraphNode.cpp
+++ b/Engine/Source/Editor/GraphEditor/Private/SGraphNode.cpp
@@ -1555,3 +1555,15 @@ void SGraphNode::RefreshErrorInfo()
 {
 	SetupErrorReporting();
 }
+
+void SGraphNode::InvalidateGraphData()
+{
+	for (TSharedRef<SGraphPin>& Pin : InputPins)
+	{
+		Pin->InvalidateGraphData();
+	}
+	for (TSharedRef<SGraphPin>& Pin : OutputPins)
+	{
+		Pin->InvalidateGraphData();
+	}
+}
diff --git a/Engine/Source/Editor/GraphEditor/Private/SGraphPanel.cpp b/Engine/Source/Editor/GraphEditor/Private/SGraphPanel.cpp
--- a/Engine/Source/Editor/GraphEditor/Private/SGraphPanel.cpp
+++ b/Engine/Source/Editor/GraphEditor/Private/SGraphPanel.cpp
@@ -1015,6 +1015,12 @@ void SGraphPanel::RemoveAllNodes()
 {
 	NodeGuidMap.Empty();
 	CurrentHoveredPins.Empty();
+	
+	for (int32 Iter = 0; Iter != Children.Num(); ++Iter)
+	{
+		GetChild(Iter)->InvalidateGraphData();
+	}
+
 	SNodePanel::RemoveAllNodes();
 }
 
@@ -1397,6 +1403,7 @@ void SGraphPanel::RemoveNode(const UEdGraphNode* Node)
 		TSharedRef<SGraphNode> Child = GetChild(Iter);
 		if (Child->GetNodeObj() == Node)
 		{
+			Child->InvalidateGraphData();
 			Children.RemoveAt(Iter);
 			break;
 		}
diff --git a/Engine/Source/Editor/GraphEditor/Private/SGraphPin.cpp b/Engine/Source/Editor/GraphEditor/Private/SGraphPin.cpp
--- a/Engine/Source/Editor/GraphEditor/Private/SGraphPin.cpp
+++ b/Engine/Source/Editor/GraphEditor/Private/SGraphPin.cpp
@@ -67,11 +67,13 @@ TSharedPtr<SGraphPin> FGraphPinHandle::FindInGraphPanel(const SGraphPanel& InPan
 
 SGraphPin::SGraphPin()
 	: GraphPinObj(nullptr)
+	, PinColorModifier(FLinearColor::White)
+	, CachedNodeOffset(FVector2D::ZeroVector)
+	, bGraphDataInvalid(false)
 	, bShowLabel(true)
 	, bOnlyShowDefaultValue(false)
 	, bIsMovingLinks(false)
-	, PinColorModifier(FLinearColor::White)
-	, CachedNodeOffset(FVector2D::ZeroVector)
+	, bUsePinColorForText(false)
 {
 	IsEditable = true;
 
@@ -552,7 +554,7 @@ FReply SGraphPin::OnMouseButtonUp( const FGeometry& MyGeometry, const FPointerEv
 
 void SGraphPin::OnMouseEnter( const FGeometry& MyGeometry, const FPointerEvent& MouseEvent )
 {
-	if (!bIsHovered)
+	if (!bIsHovered && ensure(!bGraphDataInvalid))
 	{
 		UEdGraphPin* MyPin = GetPinObj();
 		if (MyPin && !MyPin->IsPendingKill() && MyPin->GetOuter() && MyPin->GetOuter()->IsA(UEdGraphNode::StaticClass()))
@@ -1110,6 +1112,11 @@ void SGraphPin::SetOnlyShowDefaultValue(bool bNewOnlyShowDefaultValue)
 
 FText SGraphPin::GetTooltipText() const
 {
+	if(!ensure(!bGraphDataInvalid))
+	{
+		return FText();
+	}
+
 	FText HoverText = FText::GetEmpty();
 
 	UEdGraphNode* GraphNode = GraphPinObj && !GraphPinObj->IsPendingKill() ? GraphPinObj->GetOwningNodeUnchecked() : nullptr;
diff --git a/Engine/Source/Editor/GraphEditor/Public/SGraphNode.h b/Engine/Source/Editor/GraphEditor/Public/SGraphNode.h
--- a/Engine/Source/Editor/GraphEditor/Public/SGraphNode.h
+++ b/Engine/Source/Editor/GraphEditor/Public/SGraphNode.h
@@ -234,7 +234,9 @@ public:
 
 	/** Called when GraphNode changes its error information, may be called when no change has actually occurred: */
 	void RefreshErrorInfo();
-
+	
+	/** Called when ed graph data is cleared, indicating this widget can no longer safely access GraphNode - forwards call to owned pins: */
+	void InvalidateGraphData();
 protected:
 	SGraphNode();
 
diff --git a/Engine/Source/Editor/GraphEditor/Public/SGraphPin.h b/Engine/Source/Editor/GraphEditor/Public/SGraphPin.h
--- a/Engine/Source/Editor/GraphEditor/Public/SGraphPin.h
+++ b/Engine/Source/Editor/GraphEditor/Public/SGraphPin.h
@@ -190,7 +190,15 @@ public:
 
 	/** True if pin can be edited */
 	bool IsEditingEnabled() const;
-
+	
+	/** 
+	  * Called when ed graph data is cleared, indicating this widget can no longer safely access GraphPinObj 
+	  * Typically an SGraphPin is destroyed before the ed graph data is actually deleted but because SWidget
+	  * is reference counted we have had problems in the past with the SGraphPin (or Node) living unexpectedly 
+	  * long. If such a situation occurs the code attempting to own the SWidget should be demoted to a WeakPtr,
+	  * but for added safety we track graph data invalidation here:
+	  */
+	void InvalidateGraphData() { bGraphDataInvalid = true; }
 protected:
 
 	/** If true the graph pin subclass is responsible for setting the IsEnabled delegates for the aspects it cares about. If false, the default value widget enabling is done by the base class */
@@ -269,26 +277,15 @@ protected:
 	/** Is this pin editable */
 	TAttribute<bool> IsEditable;
 
-	/** If we should draw the label on this pin */
-	bool bShowLabel;
-
-	/** If we should draw the label on this pin */
-	bool bOnlyShowDefaultValue;
-
-	/** true when we're moving links between pins. */
-	bool bIsMovingLinks;
-
 	/** Color modifier for use by the connection drawing policy */
 	FLinearColor PinColorModifier;
 
 	/** Cached offset from owning node to approximate position of culled pins */
 	FVector2D CachedNodeOffset;
 
+	/** Set of pins that are currently being hovered */
 	TSet< FEdGraphPinReference > HoverPinSet;
 
-	/** TRUE if the pin should use the Pin's color for the text */
-	bool bUsePinColorForText;
-
 	//@TODO: Want to cache these once for all SGraphPins, but still handle slate style updates
 	const FSlateBrush* CachedImg_ArrayPin_Connected;
 	const FSlateBrush* CachedImg_ArrayPin_Disconnected;
@@ -304,4 +301,19 @@ protected:
 
 	const FSlateBrush* CachedImg_Pin_Background;
 	const FSlateBrush* CachedImg_Pin_BackgroundHovered;
+	
+	/** flag indicating that graph data has been deleted by the user */
+	bool bGraphDataInvalid;
+
+	/** If we should draw the label on this pin */
+	bool bShowLabel;
+
+	/** If we should draw the label on this pin */
+	bool bOnlyShowDefaultValue;
+
+	/** true when we're moving links between pins. */
+	bool bIsMovingLinks;
+
+	/** TRUE if the pin should use the Pin's color for the text */
+	bool bUsePinColorForText;
 };
