Fixed toggle comment bubble for reroute nodes
Logic originally added in 2350119, superseded in 2425487 and 2528202
I found this while investigating a fix for UE-23733, but it also fixes comments on reroute nodes disapearing when compile is clicked.

diff --git a/Engine/Source/Editor/GraphEditor/Private/SGraphNodeKnot.cpp b/Engine/Source/Editor/GraphEditor/Private/SGraphNodeKnot.cpp
--- a/Engine/Source/Editor/GraphEditor/Private/SGraphNodeKnot.cpp
+++ b/Engine/Source/Editor/GraphEditor/Private/SGraphNodeKnot.cpp
@@ -464,23 +464,3 @@ void SGraphNodeKnot::OnCommentTextCommitted(const FText& NewComment, ETextCommit
 		CommentBubble->SetCommentBubbleVisibility(/*bVisible =*/false);
 	}
 }
-
-void SGraphNodeKnot::OnMouseEnter(const FGeometry& MyGeometry, const FPointerEvent& MouseEvent)
-{
-	SGraphNode::OnMouseEnter(MyGeometry, MouseEvent);
-	if (!GraphNode->bCommentBubbleVisible && !GraphNode->NodeComment.IsEmpty())
-	{
-		// Show the bubble widget while hovered
-		CommentBubble->SetCommentBubbleVisibility(/*bVisible =*/true);
-	}
-}
-
-void SGraphNodeKnot::OnMouseLeave(const FPointerEvent& MouseEvent)
-{
-	SGraphNode::OnMouseLeave(MouseEvent);
-	if (!bAlwaysShowCommentBubble && !CommentBubble->TextBlockHasKeyboardFocus())
-	{
-		// Hide the comment bubble if visibility hasn't changed;
-		CommentBubble->SetCommentBubbleVisibility(/*bVisible =*/false);
-	}
-}
diff --git a/Engine/Source/Editor/GraphEditor/Public/SGraphNodeKnot.h b/Engine/Source/Editor/GraphEditor/Public/SGraphNodeKnot.h
--- a/Engine/Source/Editor/GraphEditor/Public/SGraphNodeKnot.h
+++ b/Engine/Source/Editor/GraphEditor/Public/SGraphNodeKnot.h
@@ -28,11 +28,6 @@ public:
 	virtual void RequestRenameOnSpawn() override { }
 	// End of SGraphNode interface
 
-	// SWidget interface
-	virtual void OnMouseEnter( const FGeometry& MyGeometry, const FPointerEvent& MouseEvent ) override;
-	virtual void OnMouseLeave( const FPointerEvent& MouseEvent ) override;
-	// End of SWidget interface
-
 protected:
 	/** Returns Offset to center comment on the node's only pin */
 	FVector2D GetCommentOffset() const;
