[NO NON-PII COMMIT MESSAGE]

diff --git a/Engine/Source/Developer/Merge/Private/SMergeGraphView.cpp b/Engine/Source/Developer/Merge/Private/SMergeGraphView.cpp
--- a/Engine/Source/Developer/Merge/Private/SMergeGraphView.cpp
+++ b/Engine/Source/Developer/Merge/Private/SMergeGraphView.cpp
@@ -27,24 +27,27 @@ struct FBlueprintRevPair
 	const FRevisionInfo& RevData;
 };;
 
-static UEdGraph* FindGraphByName(UBlueprint const& FromBlueprint, const FName& GraphName)
+static UEdGraph* FindGraphByPath(UBlueprint const& FromBlueprint, const FString& GraphPath)
 {
 	TArray<UEdGraph*> Graphs;
 	FromBlueprint.GetAllGraphs(Graphs);
 
-	UEdGraph* Ret = nullptr;
-	if (UEdGraph** Result = Graphs.FindByPredicate(FMatchFName(GraphName)))
+	for (UEdGraph* Graph : Graphs)
 	{
-		Ret = *Result;
+		FString SearchGraphPath = FGraphDiffControl::GetGraphPath(Graph);
+		if (SearchGraphPath.Equals(GraphPath))
+		{
+			return Graph;
+		}
 	}
-	return Ret;
+	return nullptr;
 }
 
 struct FMergeGraphRowEntry
 {
 	FText Label;
 
-	FName GraphName;
+	FString GraphPath;
 
 	UEdGraphNode* LocalNode;
 	UEdGraphNode* BaseNode;
@@ -61,7 +64,7 @@ struct FMergeGraphRowEntry
 
 struct FMergeGraphEntry
 {
-	FName GraphName;
+	FString GraphPath;
 
 	TArray<FMergeGraphRowEntry> Changes;
 	bool bAnyConflics;
@@ -77,22 +80,24 @@ static TArray< FMergeGraphEntry > GenerateDiffListItems(const FBlueprintRevPair&
 {
 	// Index all the graphs by name, we use the name of the graph as the 
 	// basis of comparison between the various versions of the blueprint.
-	TMap< FName, UEdGraph* > RemoteGraphMap, BaseGraphMap, LocalGraphMap;
+	TMap< FString, UEdGraph* > RemoteGraphMap, BaseGraphMap, LocalGraphMap;
 	// We also want the set of all graph names in these blueprints, so that we 
 	// can iterate over every graph.
-	TSet< FName > AllGraphNames;
+	TSet< FString > AllGraphPaths;
 	{
 		TArray<UEdGraph*> GraphsRemote, GraphsBase, GraphsLocal;
 		RemoteBlueprint.Blueprint->GetAllGraphs(GraphsRemote);
 		BaseBlueprint.Blueprint->GetAllGraphs(GraphsBase);
 		LocalBlueprint.Blueprint->GetAllGraphs(GraphsLocal);
 
-		const auto ToMap = [&AllGraphNames](const TArray<UEdGraph*>& InList, TMap<FName, UEdGraph*>& OutMap)
+		const auto ToMap = [&AllGraphPaths](const TArray<UEdGraph*>& InList, TMap<FString, UEdGraph*>& OutMap)
 		{
-			for (auto Graph : InList)
+			for (UEdGraph* Graph : InList)
 			{
-				OutMap.Add(Graph->GetFName(), Graph);
-				AllGraphNames.Add(Graph->GetFName());
+				FString GraphPath = FGraphDiffControl::GetGraphPath(Graph);
+
+				OutMap.Add(GraphPath, Graph);
+				AllGraphPaths.Add(GraphPath);
 			}
 		};
 		ToMap(GraphsRemote, RemoteGraphMap);
@@ -118,18 +123,18 @@ static TArray< FMergeGraphEntry > GenerateDiffListItems(const FBlueprintRevPair&
 			return Results;
 		};
 
-		for (const auto& GraphName : AllGraphNames)
+		for (const FString& GraphPath : AllGraphPaths)
 		{
 			TArray< FDiffSingleResult > RemoteDifferences;
 			TArray< FDiffSingleResult > LocalDifferences;
 			bool bExistsInRemote, bExistsInBase, bExistsInLocal;
 
 			FMergeGraphEntry GraphEntry;
-			GraphEntry.GraphName = GraphName;
+			GraphEntry.GraphPath = GraphPath;
 			{
-				UEdGraph** RemoteGraph = RemoteGraphMap.Find(GraphName);
-				UEdGraph** BaseGraph = BaseGraphMap.Find(GraphName);
-				UEdGraph** LocalGraph = LocalGraphMap.Find(GraphName);
+				UEdGraph** RemoteGraph = RemoteGraphMap.Find(GraphPath);
+				UEdGraph** BaseGraph = BaseGraphMap.Find(GraphPath);
+				UEdGraph** LocalGraph = LocalGraphMap.Find(GraphPath);
 
 				GraphEntry.bAnyConflics = false;
 				GraphEntry.bExistsInRemote = RemoteGraph != nullptr;
@@ -209,7 +214,7 @@ static TArray< FMergeGraphEntry > GenerateDiffListItems(const FBlueprintRevPair&
 
 						FMergeGraphRowEntry NewEntry = {
 							Label
-							, Difference.OwningGraph
+							, Difference.OwningGraphPath
 							, ConflictingDifference ? (*ConflictingDifference)->Node2 : nullptr /*UEdGraphNode* LocalNode*/
 							, Difference.Node1 /*UEdGraphNode* BaseNode*/
 							, Difference.Node2 /*UEdGraphNode* RemoteNode*/
@@ -234,7 +239,7 @@ static TArray< FMergeGraphEntry > GenerateDiffListItems(const FBlueprintRevPair&
 						{
 							FMergeGraphRowEntry NewEntry = {
 								Difference.DisplayString
-								, Difference.OwningGraph
+								, Difference.OwningGraphPath
 								, Difference.Node2 /*UEdGraphNode* LocalNode*/
 								, Difference.Node1 /*UEdGraphNode* BaseNode*/
 								, nullptr
@@ -453,22 +458,29 @@ void SMergeGraphView::Construct(const FArguments InArgs
 			FLinearColor LocalColor = ComputeColor(InDifference->bAnyConflics, InDifference->bLocalDifferences);
 			FLinearColor TextColor = ComputeColor(InDifference->bAnyConflics, InDifference->bLocalDifferences || InDifference->bRemoteDifferences);
 
+			FString DisplayString = InDifference->GraphPath;
+			int32 PeriodIndex = INDEX_NONE;
+			if (DisplayString.FindLastChar('.', PeriodIndex))
+			{
+				DisplayString = DisplayString.Mid(PeriodIndex + 1);
+			}
+
 			return SNew(SHorizontalBox)
 				+ SHorizontalBox::Slot()
 				[
 					SNew(STextBlock)
 					.ColorAndOpacity( TextColor )
-					.Text(FText::FromString(InDifference->GraphName.GetPlainNameString()))
+					.Text(FText::FromString(DisplayString))
 				]
 				+ DiffViewUtils::Box(InDifference->bExistsInRemote, RemoteColor)
 				+ DiffViewUtils::Box(InDifference->bExistsInBase, BaseColor)
 				+ DiffViewUtils::Box(InDifference->bExistsInLocal, LocalColor);
 		};
 
-		const auto FocusGraph = [](FOnMergeNodeSelected InSelectionCallback, SMergeGraphView* Parent, FName GraphName)
+		const auto FocusGraph = [](FOnMergeNodeSelected InSelectionCallback, SMergeGraphView* Parent, FString GraphPath)
 		{
 			InSelectionCallback.ExecuteIfBound();
-			Parent->FocusGraph( GraphName );
+			Parent->FocusGraph( GraphPath );
 		};
 
 		if( Children.Num() == 0 )
@@ -478,7 +490,7 @@ void SMergeGraphView::Construct(const FArguments InArgs
 
 		OutTreeEntries.Push(
 			TSharedPtr<FBlueprintDifferenceTreeEntry>(new FBlueprintDifferenceTreeEntry(
-				FOnDiffEntryFocused::CreateStatic(FocusGraph, SelectionCallback, this, Difference.GraphName)
+				FOnDiffEntryFocused::CreateStatic(FocusGraph, SelectionCallback, this, Difference.GraphPath)
 				, FGenerateDiffEntryWidget::CreateStatic(Widget, &Difference)
 				, Children
 			))
@@ -508,11 +520,11 @@ void SMergeGraphView::Construct(const FArguments InArgs
 	];
 }
 
-void SMergeGraphView::FocusGraph(FName GraphName)
+void SMergeGraphView::FocusGraph(const FString& GraphPath)
 {
-	UEdGraph* GraphRemote = FindGraphByName(*GetRemotePanel().Blueprint, GraphName);
-	UEdGraph* GraphBase = FindGraphByName(*GetBasePanel().Blueprint, GraphName);
-	UEdGraph* GraphLocal = FindGraphByName(*GetLocalPanel().Blueprint, GraphName);
+	UEdGraph* GraphRemote = FindGraphByPath(*GetRemotePanel().Blueprint, GraphPath);
+	UEdGraph* GraphBase = FindGraphByPath(*GetBasePanel().Blueprint, GraphPath);
+	UEdGraph* GraphLocal = FindGraphByPath(*GetLocalPanel().Blueprint, GraphPath);
 
 	GetBasePanel().GeneratePanel(GraphBase, nullptr);
 	GetRemotePanel().GeneratePanel(GraphRemote, GraphBase);
@@ -523,7 +535,7 @@ void SMergeGraphView::FocusGraph(FName GraphName)
 
 void SMergeGraphView::HighlightEntry(const struct FMergeGraphRowEntry& Conflict)
 {
-	FocusGraph(Conflict.GraphName);
+	FocusGraph(Conflict.GraphPath);
 
 	const auto FocusPinOrNode = [this]( UEdGraphPin* Pin, UEdGraphNode* Node )
 	{
diff --git a/Engine/Source/Developer/Merge/Private/SMergeGraphView.h b/Engine/Source/Developer/Merge/Private/SMergeGraphView.h
--- a/Engine/Source/Developer/Merge/Private/SMergeGraphView.h
+++ b/Engine/Source/Developer/Merge/Private/SMergeGraphView.h
@@ -25,7 +25,7 @@ public:
 					, TArray< TSharedPtr<FBlueprintDifferenceTreeEntry> >& OutConflicts);
 private:
 	/** Helper functions and event handlers: */
-	void FocusGraph(FName GraphName);
+	void FocusGraph(const FString& GraphPath);
 	void HighlightEntry(const struct FMergeGraphRowEntry& Conflict);
 	bool HasNoDifferences() const;
 	
diff --git a/Engine/Source/Editor/GraphEditor/Private/GraphDiffControl.cpp b/Engine/Source/Editor/GraphEditor/Private/GraphDiffControl.cpp
--- a/Engine/Source/Editor/GraphEditor/Private/GraphDiffControl.cpp
+++ b/Engine/Source/Editor/GraphEditor/Private/GraphDiffControl.cpp
@@ -4,6 +4,7 @@
 #include "EdGraph/EdGraphNode.h"
 #include "EdGraph/EdGraphPin.h"
 #include "EdGraph/EdGraph.h"
+#include "Kismet2/BlueprintEditorUtils.h"
 
 #define LOCTEXT_NAMESPACE "GraphDiffControl"
 
@@ -59,19 +60,19 @@ static void DiffR_NodeRemoved( const FGraphDiffControl::FNodeDiffContext& DiffCo
 }
 
 /** Diff result when a node comment was changed */
-static void DiffR_NodeCommentChanged(const FGraphDiffControl::FNodeDiffContext& DiffContext, FDiffResults& Results, class UEdGraphNode* Node, class UEdGraphNode* Node2)
+static void DiffR_NodeCommentChanged(const FGraphDiffControl::FNodeDiffContext& DiffContext, FDiffResults& Results, UEdGraphNode* NewNode, UEdGraphNode* OldNode)
 {
 	FDiffSingleResult Diff;
 	Diff.Diff = EDiffType::NODE_COMMENT;
-	Diff.Node1 = Node;
-	Diff.Node2 = Node2;
+	Diff.Node1 = OldNode;
+	Diff.Node2 = NewNode;
 
 	// Only bother setting up the display data if we're storing the result
 	if(Results.CanStoreResults())
 	{
 		FFormatNamedArguments Args;
 		Args.Add(TEXT("NodeType"), DiffContext.NodeTypeDisplayName);
-		Args.Add(TEXT("NodeTitle"), Node->GetNodeTitle(ENodeTitleType::ListView));
+		Args.Add(TEXT("NodeTitle"), NewNode->GetNodeTitle(ENodeTitleType::ListView));
 		Diff.ToolTip =  FText::Format(LOCTEXT("DIF_CommentModified", "Comment Modified {NodeType} '{NodeTitle}'"), Args);
 		Diff.DisplayString = Diff.ToolTip;
 		Diff.DisplayColor = FLinearColor(0.25f,0.4f,0.5f);
@@ -81,19 +82,19 @@ static void DiffR_NodeCommentChanged(const FGraphDiffControl::FNodeDiffContext&
 }
 
 /** Diff result when a node was moved on the graph */
-static void DiffR_NodeMoved(const FGraphDiffControl::FNodeDiffContext& DiffContext, FDiffResults& Results, UEdGraphNode* Node, class UEdGraphNode* OtherNode)
+static void DiffR_NodeMoved(const FGraphDiffControl::FNodeDiffContext& DiffContext, FDiffResults& Results, UEdGraphNode* NewNode, UEdGraphNode* OldNode)
 {
 	FDiffSingleResult Diff;
 	Diff.Diff = EDiffType::NODE_MOVED;
-	Diff.Node1 = Node;
-	Diff.Node2 = OtherNode;
+	Diff.Node1 = OldNode;
+	Diff.Node2 = NewNode;
 
 	// Only bother setting up the display data if we're storing the result
 	if(Results.CanStoreResults())
 	{
 		FFormatNamedArguments Args;
 		Args.Add(TEXT("NodeType"), DiffContext.NodeTypeDisplayName);
-		Args.Add(TEXT("NodeTitle"), Node->GetNodeTitle(ENodeTitleType::ListView));
+		Args.Add(TEXT("NodeTitle"), NewNode->GetNodeTitle(ENodeTitleType::ListView));
 		Diff.ToolTip = FText::Format(LOCTEXT("DIF_MoveNode", "Moved {NodeType} '{NodeTitle}'"), Args);
 		Diff.DisplayString = Diff.ToolTip;
 		Diff.DisplayColor = FLinearColor(0.9f, 0.84f, 0.43f);
@@ -103,10 +104,10 @@ static void DiffR_NodeMoved(const FGraphDiffControl::FNodeDiffContext& DiffConte
 }
 
 /** Diff result when a pin type was changed */
-static void DiffR_PinTypeChanged(FDiffResults& Results, UEdGraphPin* Pin2, UEdGraphPin* Pin1)
+static void DiffR_PinTypeChanged(FDiffResults& Results, UEdGraphPin* NewPin, UEdGraphPin* OldPin)
 {
-	FEdGraphPinType Type1 = Pin1->PinType;
-	FEdGraphPinType Type2 = Pin2->PinType;
+	FEdGraphPinType Type1 = OldPin->PinType;
+	FEdGraphPinType Type2 = NewPin->PinType;
 
 	FDiffSingleResult Diff;
 
@@ -120,9 +121,9 @@ static void DiffR_PinTypeChanged(FDiffResults& Results, UEdGraphPin* Pin2, UEdGr
 		// Only bother setting up the display data if we're storing the result
 		if(Results.CanStoreResults())
 		{
-			Diff.ToolTip = FText::Format(LOCTEXT("DIF_PinCategoryToolTipFmt", "Pin '{0}' Category was '{1}', but is now '{2}"), FText::FromName(Pin2->PinName), FText::FromName(Pin1->PinType.PinCategory), FText::FromName(Pin2->PinType.PinCategory));
+			Diff.ToolTip = FText::Format(LOCTEXT("DIF_PinCategoryToolTipFmt", "Pin '{0}' Category was '{1}', but is now '{2}"), FText::FromName(NewPin->PinName), FText::FromName(OldPin->PinType.PinCategory), FText::FromName(NewPin->PinType.PinCategory));
 			Diff.DisplayColor = FLinearColor(0.15f,0.53f,0.15f);
-			Diff.DisplayString = FText::Format(LOCTEXT("DIF_PinCategoryFmt", "Pin Category '{0}' ['{1}' -> '{2}']"), FText::FromName(Pin2->PinName), FText::FromName(Pin1->PinType.PinCategory), FText::FromName(Pin2->PinType.PinCategory));
+			Diff.DisplayString = FText::Format(LOCTEXT("DIF_PinCategoryFmt", "Pin Category '{0}' ['{1}' -> '{2}']"), FText::FromName(NewPin->PinName), FText::FromName(OldPin->PinType.PinCategory), FText::FromName(NewPin->PinType.PinCategory));
 		}
 	}
 	else if(Type1.PinSubCategory != Type2.PinSubCategory)
@@ -132,9 +133,9 @@ static void DiffR_PinTypeChanged(FDiffResults& Results, UEdGraphPin* Pin2, UEdGr
 		// Only bother setting up the display data if we're storing the result
 		if(Results.CanStoreResults())
 		{
-			Diff.ToolTip = FText::Format(LOCTEXT("DIF_PinSubCategoryToolTipFmt", "Pin '{0}' SubCategory was '{1}', but is now '{2}"), FText::FromName(Pin2->PinName), FText::FromName(Pin1->PinType.PinSubCategory), FText::FromName(Pin2->PinType.PinSubCategory));
+			Diff.ToolTip = FText::Format(LOCTEXT("DIF_PinSubCategoryToolTipFmt", "Pin '{0}' SubCategory was '{1}', but is now '{2}"), FText::FromName(NewPin->PinName), FText::FromName(OldPin->PinType.PinSubCategory), FText::FromName(NewPin->PinType.PinSubCategory));
 			Diff.DisplayColor = FLinearColor(0.45f,0.53f,0.65f);
-			Diff.DisplayString = FText::Format(LOCTEXT("DIF_PinSubCategoryFmt", "Pin SubCategory '{0}'  ['{1}' -> '{2}']"), FText::FromName(Pin2->PinName), FText::FromName(Pin1->PinType.PinSubCategory), FText::FromName(Pin2->PinType.PinSubCategory));
+			Diff.DisplayString = FText::Format(LOCTEXT("DIF_PinSubCategoryFmt", "Pin SubCategory '{0}'  ['{1}' -> '{2}']"), FText::FromName(NewPin->PinName), FText::FromName(OldPin->PinType.PinSubCategory), FText::FromName(NewPin->PinType.PinSubCategory));
 		}
 	}
 	else if(T1Obj != T2Obj && (T1Obj && T2Obj ) &&
@@ -148,9 +149,9 @@ static void DiffR_PinTypeChanged(FDiffResults& Results, UEdGraphPin* Pin2, UEdGr
 			const FName Obj1 = T1Obj->GetFName();
 			const FName Obj2 = T2Obj->GetFName();
 
-			Diff.ToolTip = FText::Format(LOCTEXT("DIF_PinSubCategorObjToolTipFmt", "Pin '{0}' was SubCategoryObject '{1}', but is now '{2}"), FText::FromName(Pin2->PinName), FText::FromName(Obj1), FText::FromName(Obj2));
+			Diff.ToolTip = FText::Format(LOCTEXT("DIF_PinSubCategorObjToolTipFmt", "Pin '{0}' was SubCategoryObject '{1}', but is now '{2}"), FText::FromName(NewPin->PinName), FText::FromName(Obj1), FText::FromName(Obj2));
 			Diff.DisplayColor = FLinearColor(0.45f,0.13f,0.25f);
-			Diff.DisplayString = FText::Format(LOCTEXT("DIF_PinSubCategoryObjFmt", "Pin SubCategoryObject '{0}' ['{1}' -> '{2}']"), FText::FromName(Pin2->PinName), FText::FromName(Obj1), FText::FromName(Obj2));
+			Diff.DisplayString = FText::Format(LOCTEXT("DIF_PinSubCategoryObjFmt", "Pin SubCategoryObject '{0}' ['{1}' -> '{2}']"), FText::FromName(NewPin->PinName), FText::FromName(Obj1), FText::FromName(Obj2));
 		}
 	}
 	else if(Type1.ContainerType != Type2.ContainerType)
@@ -161,11 +162,11 @@ static void DiffR_PinTypeChanged(FDiffResults& Results, UEdGraphPin* Pin2, UEdGr
 		// Only bother setting up the display data if we're storing the result
 		if(Results.CanStoreResults())
 		{
-			FText IsArray1 = Pin1->PinType.IsArray() ? LOCTEXT("true", "true") : LOCTEXT("false", "false");
-			FText IsArray2 = Pin2->PinType.IsArray() ? LOCTEXT("true", "true") : LOCTEXT("false", "false");
+			FText IsArray1 = OldPin->PinType.IsArray() ? LOCTEXT("true", "true") : LOCTEXT("false", "false");
+			FText IsArray2 = NewPin->PinType.IsArray() ? LOCTEXT("true", "true") : LOCTEXT("false", "false");
 
-			Diff.ToolTip = FText::Format(LOCTEXT("DIF_PinIsArrayToolTipFmt", "PinType IsArray for '{0}' modified. Was '{1}', but is now '{2}"), FText::FromName(Pin2->PinName), IsArray1, IsArray2);
-			Diff.DisplayString = FText::Format(LOCTEXT("DIF_PinIsArrayFmt", "Pin IsArray '{0}' ['{1}' -> '{2}']"), FText::FromName(Pin2->PinName), IsArray1, IsArray2);
+			Diff.ToolTip = FText::Format(LOCTEXT("DIF_PinIsArrayToolTipFmt", "PinType IsArray for '{0}' modified. Was '{1}', but is now '{2}"), FText::FromName(NewPin->PinName), IsArray1, IsArray2);
+			Diff.DisplayString = FText::Format(LOCTEXT("DIF_PinIsArrayFmt", "Pin IsArray '{0}' ['{1}' -> '{2}']"), FText::FromName(NewPin->PinName), IsArray1, IsArray2);
 			Diff.DisplayColor = FLinearColor(0.45f,0.33f,0.35f);
 		}
 	}
@@ -176,42 +177,42 @@ static void DiffR_PinTypeChanged(FDiffResults& Results, UEdGraphPin* Pin2, UEdGr
 		// Only bother setting up the display data if we're storing the result
 		if(Results.CanStoreResults())
 		{
-			FText IsRef1 = Pin1->PinType.bIsReference ? LOCTEXT("true", "true") : LOCTEXT("false", "false");
-			FText IsRef2 = Pin2->PinType.bIsReference ? LOCTEXT("true", "true") : LOCTEXT("false", "false");
+			FText IsRef1 = OldPin->PinType.bIsReference ? LOCTEXT("true", "true") : LOCTEXT("false", "false");
+			FText IsRef2 = NewPin->PinType.bIsReference ? LOCTEXT("true", "true") : LOCTEXT("false", "false");
 
-			Diff.ToolTip = FText::Format(LOCTEXT("DIF_PinIsRefToolTipFmt", "PinType IsReference for '{0}' modified. Was '{1}', but is now '{2}"), FText::FromName(Pin2->PinName), IsRef1, IsRef2);
+			Diff.ToolTip = FText::Format(LOCTEXT("DIF_PinIsRefToolTipFmt", "PinType IsReference for '{0}' modified. Was '{1}', but is now '{2}"), FText::FromName(NewPin->PinName), IsRef1, IsRef2);
 			Diff.DisplayColor = FLinearColor(0.25f,0.43f,0.35f);
-			Diff.DisplayString = FText::Format(LOCTEXT("DIF_PinIsRefFmt", "Pin IsReference '{0}' ['{1}' -> '{2}']"), FText::FromName(Pin2->PinName), IsRef1, IsRef2);
+			Diff.DisplayString = FText::Format(LOCTEXT("DIF_PinIsRefFmt", "Pin IsReference '{0}' ['{1}' -> '{2}']"), FText::FromName(NewPin->PinName), IsRef1, IsRef2);
 		}
 	}
 
-	Diff.Pin1 = Pin1;
-	Diff.Pin2 = Pin2;
+	Diff.Pin1 = OldPin;
+	Diff.Pin2 = NewPin;
 	Results.Add(Diff);
 }
 
 /** Diff result when the # of links to a pin was changed */
-static void DiffR_PinLinkCountChanged(FDiffResults& Results, UEdGraphPin* Pin2, UEdGraphPin* Pin1)
+static void DiffR_PinLinkCountChanged(FDiffResults& Results, UEdGraphPin* NewPin, UEdGraphPin* OldPin)
 {
 	FDiffSingleResult Diff;
-	Diff.Diff = Pin2->LinkedTo.Num() > Pin1->LinkedTo.Num()  ?  EDiffType::PIN_LINKEDTO_NUM_INC : EDiffType::PIN_LINKEDTO_NUM_DEC;
-	Diff.Pin2 = Pin2;
-	Diff.Pin1 = Pin1;
+	Diff.Diff = NewPin->LinkedTo.Num() > OldPin->LinkedTo.Num()  ?  EDiffType::PIN_LINKEDTO_NUM_INC : EDiffType::PIN_LINKEDTO_NUM_DEC;
+	Diff.Pin2 = NewPin;
+	Diff.Pin1 = OldPin;
 
 	// Only bother setting up the display data if we're storing the result
 	if(Results.CanStoreResults())
 	{
 		if(Diff.Diff == EDiffType::PIN_LINKEDTO_NUM_INC)
 		{
-			Diff.ToolTip = FText::Format(LOCTEXT("DIF_PinLinkCountIncToolTipFmt", "Pin '{0}' has more links (was {1} now {2})"), FText::FromName(Pin1->PinName), FText::AsNumber(Pin1->LinkedTo.Num()), FText::AsNumber(Pin2->LinkedTo.Num()));
+			Diff.ToolTip = FText::Format(LOCTEXT("DIF_PinLinkCountIncToolTipFmt", "Pin '{0}' has more links (was {1} now {2})"), FText::FromName(OldPin->PinName), FText::AsNumber(OldPin->LinkedTo.Num()), FText::AsNumber(NewPin->LinkedTo.Num()));
 			Diff.DisplayColor = FLinearColor(0.5f,0.3f,0.85f);
-			Diff.DisplayString = FText::Format(LOCTEXT("DIF_PinLinkCountIncFmt", "Added Link to '{0}'"), FText::FromName(Pin1->PinName));
+			Diff.DisplayString = FText::Format(LOCTEXT("DIF_PinLinkCountIncFmt", "Added Link to '{0}'"), FText::FromName(OldPin->PinName));
 		}
 		else
 		{
-			Diff.ToolTip = FText::Format(LOCTEXT("DIF_PinLinkCountDecToolTipFmt", "Pin '{0}' has fewer links (was {1} now {2})"), FText::FromName(Pin1->PinName), FText::AsNumber(Pin1->LinkedTo.Num()), FText::AsNumber(Pin2->LinkedTo.Num()));
+			Diff.ToolTip = FText::Format(LOCTEXT("DIF_PinLinkCountDecToolTipFmt", "Pin '{0}' has fewer links (was {1} now {2})"), FText::FromName(OldPin->PinName), FText::AsNumber(OldPin->LinkedTo.Num()), FText::AsNumber(NewPin->LinkedTo.Num()));
 			Diff.DisplayColor = FLinearColor(0.5f,0.3f,0.85f);
-			Diff.DisplayString = FText::Format(LOCTEXT("DIF_PinLinkCountDecFmt", "Removed Link to '{0}'"), FText::FromName(Pin1->PinName));
+			Diff.DisplayString = FText::Format(LOCTEXT("DIF_PinLinkCountDecFmt", "Removed Link to '{0}'"), FText::FromName(OldPin->PinName));
 		}
 	}
 
@@ -219,23 +220,23 @@ static void DiffR_PinLinkCountChanged(FDiffResults& Results, UEdGraphPin* Pin2,
 }
 
 /** Diff result when a pin to relinked to a different node */
-static void DiffR_LinkedToNode(FDiffResults& Results, UEdGraphPin* Pin1, UEdGraphPin* Pin2,  UEdGraphNode* Node1, UEdGraphNode* Node2)
+static void DiffR_LinkedToNode(FDiffResults& Results, UEdGraphPin* OldPin, UEdGraphPin* NewPin, UEdGraphNode* OldNode, UEdGraphNode* NewNode)
 {
 	FDiffSingleResult Diff;
 	Diff.Diff = EDiffType::PIN_LINKEDTO_NODE;
-	Diff.Pin2 = Pin2;
-	Diff.Pin1 = Pin1;
-	Diff.Node1 = Node1;
-	Diff.Node2 = Node2;
+	Diff.Pin1 = OldPin;
+	Diff.Pin2 = NewPin;
+	Diff.Node1 = OldNode;
+	Diff.Node2 = NewNode;
 
 	// Only bother setting up the display data if we're storing the result
 	if(Results.CanStoreResults())
 	{
-		FText Node1Name = Node1->GetNodeTitle(ENodeTitleType::ListView);
-		FText Node2Name = Node2->GetNodeTitle(ENodeTitleType::ListView);
+		FText Node1Name = OldNode->GetNodeTitle(ENodeTitleType::ListView);
+		FText Node2Name = NewNode->GetNodeTitle(ENodeTitleType::ListView);
 
 		FFormatNamedArguments Args;
-		Args.Add(TEXT("PinNameForNode1"), FText::FromName(Pin1->PinName));
+		Args.Add(TEXT("PinNameForNode1"), FText::FromName(OldPin->PinName));
 		Args.Add(TEXT("NodeName1"), Node1Name);
 		Args.Add(TEXT("NodeName2"), Node2Name);
 		Diff.ToolTip = FText::Format(LOCTEXT("DIF_PinLinkMovedToolTip", "Pin '{PinNameForNode1}' was linked to Node '{NodeName1}', but is now linked to Node '{NodeName2}'"), Args);
@@ -248,20 +249,20 @@ static void DiffR_LinkedToNode(FDiffResults& Results, UEdGraphPin* Pin1, UEdGrap
 }
 
 /** Diff result when a pin default value was changed, and is in use*/
-static void DiffR_PinDefaultValueChanged(FDiffResults& Results, UEdGraphPin* Pin2, UEdGraphPin* Pin1)
+static void DiffR_PinDefaultValueChanged(FDiffResults& Results, UEdGraphPin* NewPin, UEdGraphPin* OldPin)
 {
 	FDiffSingleResult Diff;
 	Diff.Diff = EDiffType::PIN_DEFAULT_VALUE;
-	Diff.Pin1 = Pin1;
-	Diff.Pin2 = Pin2;
+	Diff.Pin1 = OldPin;
+	Diff.Pin2 = NewPin;
 
 	// Only bother setting up the display data if we're storing the result
 	if(Results.CanStoreResults())
 	{
 		FFormatNamedArguments Args;
-		Args.Add(TEXT("PinNameForValue1"), FText::FromName(Pin2->PinName));
-		Args.Add(TEXT("PinValue1"), FText::FromString(Pin1->GetDefaultAsString()));
-		Args.Add(TEXT("PinValue2"), FText::FromString(Pin2->GetDefaultAsString()));
+		Args.Add(TEXT("PinNameForValue1"), FText::FromName(NewPin->PinName));
+		Args.Add(TEXT("PinValue1"), FText::FromString(OldPin->GetDefaultAsString()));
+		Args.Add(TEXT("PinValue2"), FText::FromString(NewPin->GetDefaultAsString()));
 		Diff.ToolTip = FText::Format(LOCTEXT("DIF_PinDefaultValueToolTip", "Pin '{PinNameForValue1}' Default Value was '{PinValue1}', but is now '{PinValue2}"), Args);
 		Diff.DisplayColor = FLinearColor(0.665f,0.13f,0.455f);
 		Diff.DisplayString = FText::Format(LOCTEXT("DIF_PinDefaultValue", "Pin Default '{PinNameForValue1}' '{PinValue1}' -> '{PinValue2}']"), Args);
@@ -270,24 +271,16 @@ static void DiffR_PinDefaultValueChanged(FDiffResults& Results, UEdGraphPin* Pin
 	Results.Add(Diff);
 }
 
-/** 
- * Diff result when pin count is not the same.
- *
- * @param	Result	Struct to receive result.
- * @param	Node	Node 1 to which the pins relate.
- * @param	Node2	Node 2 to which the pins relate.
- * @param	Pins1	Array of relevant pins on Node 1.
- * @param	Pins2	Array of relevant pins on Node 2.
- */
-static void DiffR_NodePinCount(FDiffResults& Results, UEdGraphNode* Node, UEdGraphNode* Node2, const TArray<UEdGraphPin*> &Pins1, const TArray<UEdGraphPin*> &Pins2)
+/** Diff result when pin count is not the same */
+static void DiffR_NodePinCount(FDiffResults& Results, UEdGraphNode* NewNode, UEdGraphNode* OldNode, const TArray<UEdGraphPin*>& NewPins, const TArray<UEdGraphPin*>& OldPins)
 {
-	FText NodeName = Node->GetNodeTitle(ENodeTitleType::ListView);
-	int32 OriginalCount = Pins2.Num();
-	int32 NewCount = Pins1.Num();
+	FText NodeName = NewNode->GetNodeTitle(ENodeTitleType::ListView);
+	int32 OriginalCount = OldPins.Num();
+	int32 NewCount = NewPins.Num();
 	FDiffSingleResult Diff;
 	Diff.Diff = EDiffType::NODE_PIN_COUNT;
-	Diff.Node1 = Node;
-	Diff.Node2 = Node2;
+	Diff.Node1 = OldNode;
+	Diff.Node2 = NewNode;
 
 	// Only bother setting up the display data if we're storing the result
 	if(Results.CanStoreResults())
@@ -315,22 +308,23 @@ static void DiffR_NodePinCount(FDiffResults& Results, UEdGraphNode* Node, UEdGra
 
 		FText ListOfPins;
 		TArray< FText > RemovedPins;
-		for (UEdGraphPin* OldPin : Pins1)
+		TArray< FText > AddedPins;
+		
+		for (UEdGraphPin* SearchPin : OldPins)
 		{
-			const UEdGraphPin* const* FoundPin = Pins2.FindByPredicate(FMatchName(OldPin->PinName));
+			const UEdGraphPin* const* FoundPin = NewPins.FindByPredicate(FMatchName(SearchPin->PinName));
 			if (FoundPin == nullptr)
 			{
-				RemovedPins.Add(OldPin->GetDisplayName());
+				RemovedPins.Add(SearchPin->GetDisplayName());
 			}
 		}
 
-		TArray< FText > AddedPins;
-		for (UEdGraphPin* OldPin : Pins2)
+		for (UEdGraphPin* SearchPin : NewPins)
 		{
-			const UEdGraphPin* const* FoundPin = Pins1.FindByPredicate(FMatchName(OldPin->PinName));
+			const UEdGraphPin* const* FoundPin = OldPins.FindByPredicate(FMatchName(SearchPin->PinName));
 			if (FoundPin == nullptr)
 			{
-				AddedPins.Add(OldPin->GetDisplayName());
+				AddedPins.Add(SearchPin->GetDisplayName());
 			}
 		}
 
@@ -397,7 +391,7 @@ static void BuildArrayOfRelevantPins(const TArray<UEdGraphPin*>& InPins, TArray<
 	for(int32 i = 0;i<InPins.Num();++i)
 	{
 		UEdGraphPin* EachPin = InPins[i];
-		if( EachPin != NULL )
+		if( EachPin != nullptr )
 		{
 			if( EachPin->bHidden == false )
 			{
@@ -441,20 +435,20 @@ static UEdGraphPin* FindOtherLink(TArray<UEdGraphPin*>& Links2, int32 OriginalIn
 }
 
 /** Determine if the LinkedTo pins are the same */
-static bool LinkedToDifferent(UEdGraphPin* OriginalPin1, UEdGraphPin* OriginalPin2, const TArray<UEdGraphPin*>& Links1, TArray<UEdGraphPin*>& Links2, FDiffResults& Results)
+static bool LinkedToDifferent(UEdGraphPin* OldPin, UEdGraphPin* NewPin, const TArray<UEdGraphPin*>& OldLinks, TArray<UEdGraphPin*>& NewLinks, FDiffResults& Results)
 {
-	const int32 Size = Links1.Num();
+	const int32 Size = OldLinks.Num();
 	bool bHasResult = false;
 	for(int32 i = 0;i<Size;++i)
 	{
-		UEdGraphPin* Pin1 = Links1[i];
-		UEdGraphPin* Pin2 = FindOtherLink(Links2, i, Pin1);
+		UEdGraphPin* OldLinkedPin = OldLinks[i];
+		UEdGraphPin* NewLinkedPin = FindOtherLink(NewLinks, i, OldLinkedPin);
 
-		UEdGraphNode* Node1 = Pin1->GetOwningNode();
-		UEdGraphNode* Node2 = Pin2->GetOwningNode();
-		if(!FGraphDiffControl::IsNodeMatch(Node1, Node2))
+		UEdGraphNode* OldNode = OldLinkedPin->GetOwningNode();
+		UEdGraphNode* NewNode = NewLinkedPin->GetOwningNode();
+		if(!FGraphDiffControl::IsNodeMatch(OldNode, NewNode))
 		{
-			DiffR_LinkedToNode(Results, OriginalPin1, OriginalPin2, Node1, Node2);
+			DiffR_LinkedToNode(Results, OldPin, NewPin, OldNode, NewNode);
 			KEEP_GOING_IF_RESULTS()
 		}
 	}
@@ -464,39 +458,39 @@ static bool LinkedToDifferent(UEdGraphPin* OriginalPin1, UEdGraphPin* OriginalPi
 /** 
  * Determine of two Arrays of Pins are different.
  *
- * @param	Pins1	First set of pins to compare.
- * @param	Pins2	Second set of pins to compare.
+ * @param	OldPins	First set of pins to compare.
+ * @param	NewPins	Second set of pins to compare.
  * @param	Results Difference results.
  * 
  * returns true if any pins are different and populates the Results array
  */
-static bool ArePinsDifferent(const TArray<UEdGraphPin*>& Pins1, TArray<UEdGraphPin*>& Pins2, FDiffResults& Results)
+static bool ArePinsDifferent(const TArray<UEdGraphPin*>& OldPins, TArray<UEdGraphPin*>& NewPins, FDiffResults& Results)
 {
-	const int32 Size = Pins1.Num();
+	const int32 Size = OldPins.Num();
 	bool bHasResult = false;
 	for(int32 i = 0;i<Size;++i)
 	{
-		UEdGraphPin* Pin1 = Pins1[i];
-		UEdGraphPin* Pin2 = Pins2[i];
+		UEdGraphPin* OldPin = OldPins[i];
+		UEdGraphPin* NewPin = NewPins[i];
 
-		if(IsPinTypeDifferent(Pin1->PinType, Pin2->PinType))
+		if(IsPinTypeDifferent(OldPin->PinType, NewPin->PinType))
 		{
-			DiffR_PinTypeChanged(Results, Pin2, Pin1);
+			DiffR_PinTypeChanged(Results, NewPin, OldPin);
 			KEEP_GOING_IF_RESULTS()
 		}
-		if(Pin1->LinkedTo.Num() != Pin2->LinkedTo.Num())
+		if(OldPin->LinkedTo.Num() != NewPin->LinkedTo.Num())
 		{
-			DiffR_PinLinkCountChanged(Results, Pin2, Pin1);
+			DiffR_PinLinkCountChanged(Results, NewPin, OldPin);
 			KEEP_GOING_IF_RESULTS()
 		}
-		else if(LinkedToDifferent(Pin1, Pin2,Pin1->LinkedTo, Pin2->LinkedTo, Results))
+		else if(LinkedToDifferent(OldPin, NewPin, OldPin->LinkedTo, NewPin->LinkedTo, Results))
 		{
 			KEEP_GOING_IF_RESULTS()
 		}
 
-		if(Pin2->LinkedTo.Num() == 0 && (Pin2->GetDefaultAsString() != Pin1->GetDefaultAsString()))
+		if(NewPin->LinkedTo.Num() == 0 && (NewPin->GetDefaultAsString() != OldPin->GetDefaultAsString()))
 		{
-			DiffR_PinDefaultValueChanged(Results, Pin2, Pin1); //note: some issues with how floating point is stored as string format(0.0 vs 0.00) can cause false diffs
+			DiffR_PinDefaultValueChanged(Results, NewPin, OldPin); //note: some issues with how floating point is stored as string format(0.0 vs 0.00) can cause false diffs
 			KEEP_GOING_IF_RESULTS()
 		}
 	}
@@ -509,10 +503,10 @@ static bool ArePinsDifferent(const TArray<UEdGraphPin*>& Pins1, TArray<UEdGraphP
 
 bool FGraphDiffControl::FNodeMatch::IsValid() const
 {
-	return ((NewNode != NULL) && (OldNode != NULL));
+	return ((NewNode != nullptr) && (OldNode != nullptr));
 }
 
-bool FGraphDiffControl::FNodeMatch::Diff(const FNodeDiffContext& DiffContext, TArray<FDiffSingleResult>* OptionalDiffsArray /* = NULL*/) const
+bool FGraphDiffControl::FNodeMatch::Diff(const FNodeDiffContext& DiffContext, TArray<FDiffSingleResult>* OptionalDiffsArray /* = nullptr*/) const
 {
 	FDiffResults DiffsOut(OptionalDiffsArray);
 	return Diff(DiffContext, DiffsOut);
@@ -567,7 +561,7 @@ bool FGraphDiffControl::FNodeMatch::Diff(const FNodeDiffContext& DiffContext, FD
 		}
 	}
 	else if(DiffContext.DiffFlags & EDiffFlags::NodeExistance)
-	// one of the nodes is NULL
+	// one of the nodes is nullptr
 	{
 		bIsDifferent = true;
 		switch (DiffContext.DiffMode)
@@ -592,7 +586,7 @@ bool FGraphDiffControl::FNodeMatch::Diff(const FNodeDiffContext& DiffContext, FD
 * FGraphDiffControl
 *******************************************************************************/
 
-FGraphDiffControl::FNodeMatch FGraphDiffControl::FindNodeMatch(class UEdGraph* Graph, class UEdGraphNode* Node, TArray<FNodeMatch> const& PriorMatches)
+FGraphDiffControl::FNodeMatch FGraphDiffControl::FindNodeMatch(UEdGraph* Graph, UEdGraphNode* Node, TArray<FNodeMatch> const& PriorMatches)
 {
 	FNodeMatch Match;
 	Match.NewNode = Node;
@@ -613,7 +607,7 @@ FGraphDiffControl::FNodeMatch FGraphDiffControl::FindNodeMatch(class UEdGraph* G
 	return Match;
 }
 
-bool FGraphDiffControl::IsNodeMatch(class UEdGraphNode* Node1, class UEdGraphNode* Node2, TArray<FGraphDiffControl::FNodeMatch> const* Exclusions)
+bool FGraphDiffControl::IsNodeMatch(UEdGraphNode* Node1, UEdGraphNode* Node2, TArray<FGraphDiffControl::FNodeMatch> const* Exclusions)
 {
 	if(Node2->GetClass() != Node1->GetClass()) 
 	{
@@ -713,13 +707,31 @@ bool FGraphDiffControl::DiffGraphs(UEdGraph* const LhsGraph, UEdGraph* const Rhs
 	}
 
 	// storing the graph name for all diff entries:
-	const FName GraphName = LhsGraph ? LhsGraph->GetFName() : RhsGraph->GetFName();
+	const FString GraphPath = LhsGraph ? GetGraphPath(LhsGraph) : GetGraphPath(RhsGraph);
 	for( FDiffSingleResult& Entry : DiffsOut )
 	{
-		Entry.OwningGraph = GraphName;
+		Entry.OwningGraphPath = GraphPath;
 	}
 
 	return bFoundDifferences;
 }
 
+FString FGraphDiffControl::GetGraphPath(UEdGraph* Graph)
+{
+	FString GraphPath;
+	if (UBlueprint* Blueprint = FBlueprintEditorUtils::FindBlueprintForGraph(Graph))
+	{
+		return Graph->GetPathName(Blueprint);
+	}
+	else if (UPackage* Package = Graph->GetOutermost())
+	{
+		return Graph->GetPathName(Package);
+	}
+	else if (Graph)
+	{
+		return Graph->GetName();
+	}
+	return FString();
+}
+
 #undef LOCTEXT_NAMESPACE
diff --git a/Engine/Source/Editor/GraphEditor/Public/DiffResults.h b/Engine/Source/Editor/GraphEditor/Public/DiffResults.h
--- a/Engine/Source/Editor/GraphEditor/Public/DiffResults.h
+++ b/Engine/Source/Editor/GraphEditor/Public/DiffResults.h
@@ -72,8 +72,8 @@ struct FDiffSingleResult
 	/** User can override color to use for display string */
 	FLinearColor DisplayColor;
 
-	/** Name of the graph this difference was created on*/
-	FName OwningGraph;
+	/** Path string of graph, relative to blueprint/asset root */
+	FString OwningGraphPath;
 };
 
 FORCEINLINE bool operator==( const FDiffSingleResult& LHS, const FDiffSingleResult& RHS )
diff --git a/Engine/Source/Editor/GraphEditor/Public/GraphDiffControl.h b/Engine/Source/Editor/GraphEditor/Public/GraphDiffControl.h
--- a/Engine/Source/Editor/GraphEditor/Public/GraphDiffControl.h
+++ b/Engine/Source/Editor/GraphEditor/Public/GraphDiffControl.h
@@ -116,4 +116,7 @@ public:
 	 * @return True if any differences were found, false if both graphs are identical.
 	 */
 	static bool DiffGraphs(class UEdGraph* const OldGraph, class UEdGraph* const NewGraph, TArray<FDiffSingleResult>& DiffsOut);
+
+	/** Computes an object path for a specific graph, relative to the blueprint/asset root */
+	static FString GetGraphPath(UEdGraph* Graph);
 };
diff --git a/Engine/Source/Editor/Kismet/Private/SBlueprintDiff.cpp b/Engine/Source/Editor/Kismet/Private/SBlueprintDiff.cpp
--- a/Engine/Source/Editor/Kismet/Private/SBlueprintDiff.cpp
+++ b/Engine/Source/Editor/Kismet/Private/SBlueprintDiff.cpp
@@ -603,8 +603,8 @@ void FListItemGraphToDiff::OnGraphChanged( const FEdGraphEditAction& Action )
 
 FDiffPanel::FDiffPanel()
 {
-	Blueprint = NULL;
-	LastFocusedPin = NULL;
+	Blueprint = nullptr;
+	LastFocusedPin = nullptr;
 }
 
 void FDiffPanel::InitializeDiffPanel()
@@ -996,12 +996,12 @@ TSharedRef<SDockTab> SBlueprintDiff::CreateMyBlueprintsViews(const FSpawnTabArgs
 	];
 }
 
-FListItemGraphToDiff* SBlueprintDiff::FindGraphToDiffEntry(FName ByName)
+FListItemGraphToDiff* SBlueprintDiff::FindGraphToDiffEntry(const FString& GraphPath)
 {
 	for( const auto& Graph : Graphs )
 	{
-		FName GraphName = Graph->GetGraphOld() ? Graph->GetGraphOld()->GetFName() : Graph->GetGraphNew()->GetFName();
-		if (GraphName == ByName )
+		FString SearchGraphPath = Graph->GetGraphOld() ? FGraphDiffControl::GetGraphPath(Graph->GetGraphOld()) : FGraphDiffControl::GetGraphPath(Graph->GetGraphNew());
+		if (SearchGraphPath.Equals(GraphPath, ESearchCase::CaseSensitive))
 		{
 			return Graph.Get();
 		}
@@ -1013,15 +1013,8 @@ void SBlueprintDiff::FocusOnGraphRevisions( FListItemGraphToDiff* Diff )
 {
 	UEdGraph* Graph = Diff->GetGraphOld() ? Diff->GetGraphOld() : Diff->GetGraphNew();
 
-	FString GraphPath;
-	if (UBlueprint* Blueprint = FBlueprintEditorUtils::FindBlueprintForGraph(Graph))
-	{
-		GraphPath = Graph->GetPathName(Blueprint);
-	}
-	else
-	{
-		GraphPath = Graph->GetName();
-	}
+	FString GraphPath = FGraphDiffControl::GetGraphPath(Graph);
+
 	HandleGraphChanged(GraphPath);
 
 	ResetGraphEditors();
@@ -1029,8 +1022,8 @@ void SBlueprintDiff::FocusOnGraphRevisions( FListItemGraphToDiff* Diff )
 
 void SBlueprintDiff::OnDiffListSelectionChanged(TSharedPtr<FDiffResultItem> TheDiff )
 {
-	check( TheDiff->Result.OwningGraph != FName() );
-	FocusOnGraphRevisions( FindGraphToDiffEntry( TheDiff->Result.OwningGraph ) );
+	check( !TheDiff->Result.OwningGraphPath.IsEmpty() );
+	FocusOnGraphRevisions( FindGraphToDiffEntry( TheDiff->Result.OwningGraphPath) );
 	FDiffSingleResult Result = TheDiff->Result;
 
 	const auto SafeClearSelection = []( TWeakPtr<SGraphEditor> GraphEditor )
@@ -1102,7 +1095,7 @@ void FDiffPanel::GeneratePanel(UEdGraph* Graph, UEdGraph* GraphToDiff )
 	{
 		LastFocusedPin->bIsDiffing = false;
 	}
-	LastFocusedPin = NULL;
+	LastFocusedPin = nullptr;
 
 	TSharedPtr<SWidget> Widget = SNew(SBorder)
 								.HAlign(HAlign_Center)
@@ -1182,7 +1175,7 @@ bool FDiffPanel::CanCopyNodes() const
 	for (FGraphPanelSelectionSet::TConstIterator SelectedIter(SelectedNodes); SelectedIter; ++SelectedIter)
 	{
 		UEdGraphNode* Node = Cast<UEdGraphNode>(*SelectedIter);
-		if ((Node != NULL) && Node->CanDuplicateNode())
+		if ((Node != nullptr) && Node->CanDuplicateNode())
 		{
 			return true;
 		}
@@ -1208,7 +1201,7 @@ void FDiffPanel::FocusDiff(UEdGraphNode& Node)
 	{
 		LastFocusedPin->bIsDiffing = false;
 	}
-	LastFocusedPin = NULL;
+	LastFocusedPin = nullptr;
 
 	if (GraphEditor.IsValid())
 	{
@@ -1244,7 +1237,7 @@ void SBlueprintDiff::HandleGraphChanged( const FString& GraphPath )
 	UEdGraph* GraphOld = nullptr;
 	for (UEdGraph* OldGraph : GraphsOld)
 	{
-		if (GraphPath.Equals(OldGraph->GetPathName(PanelOld.Blueprint)))
+		if (GraphPath.Equals(FGraphDiffControl::GetGraphPath(OldGraph)))
 		{
 			GraphOld = OldGraph;
 			break;
@@ -1254,7 +1247,7 @@ void SBlueprintDiff::HandleGraphChanged( const FString& GraphPath )
 	UEdGraph* GraphNew = nullptr;
 	for (UEdGraph* NewGraph : GraphsNew)
 	{
-		if (GraphPath.Equals(NewGraph->GetPathName(PanelNew.Blueprint)))
+		if (GraphPath.Equals(FGraphDiffControl::GetGraphPath(NewGraph)))
 		{
 			GraphNew = NewGraph;
 			break;
@@ -1278,14 +1271,14 @@ void SBlueprintDiff::GenerateDifferencesList()
 	for (auto It(GraphsOld.CreateConstIterator()); It; It++)
 	{
 		UEdGraph* GraphOld = *It;
-		UEdGraph* GraphNew = NULL;
+		UEdGraph* GraphNew = nullptr;
 		for (auto It2(GraphsNew.CreateIterator()); It2; It2++)
 		{
 			UEdGraph* TestGraph = *It2;
 			if (TestGraph && GraphOld->GetName() == TestGraph->GetName())
 			{
 				GraphNew = TestGraph;
-				*It2 = NULL;
+				*It2 = nullptr;
 				break;
 			}
 		}
@@ -1300,9 +1293,9 @@ void SBlueprintDiff::GenerateDifferencesList()
 	for (auto It2(GraphsNew.CreateIterator()); It2; It2++)
 	{
 		UEdGraph* GraphNew = *It2;
-		if (GraphNew != NULL && IsGraphDiffNeeded(GraphNew))
+		if (GraphNew != nullptr && IsGraphDiffNeeded(GraphNew))
 		{
-			CreateGraphEntry(NULL, GraphNew);
+			CreateGraphEntry(nullptr, GraphNew);
 		}
 	}
 
diff --git a/Engine/Source/Editor/Kismet/Public/SBlueprintDiff.h b/Engine/Source/Editor/Kismet/Public/SBlueprintDiff.h
--- a/Engine/Source/Editor/Kismet/Public/SBlueprintDiff.h
+++ b/Engine/Source/Editor/Kismet/Public/SBlueprintDiff.h
@@ -23,21 +23,6 @@ class UEdGraph;
 struct FListItemGraphToDiff;
 enum class EAssetEditorCloseReason : uint8;
 
-struct FMatchFName
-{
-	FMatchFName(FName InName)
-	: Name(InName)
-	{
-	}
-
-	bool operator() (const UObject* Object)
-	{
-		return Object->GetFName() == Name;
-	}
-
-	FName const Name;
-};
-
 /** Individual Diff item shown in the list of diffs */
 struct FDiffResultItem : public TSharedFromThis<FDiffResultItem>
 {
@@ -161,8 +146,8 @@ protected:
 	typedef TSharedPtr<struct FListItemGraphToDiff>	FGraphToDiff;
 	typedef SListView<FGraphToDiff >	SListViewType;
 
-	/** Find the FListItemGraphToDiff that displays the graph with GraphName: */
-	FListItemGraphToDiff* FindGraphToDiffEntry(FName ByName);
+	/** Find the FListItemGraphToDiff that displays the graph with GraphPath relative path */
+	FListItemGraphToDiff* FindGraphToDiffEntry(const FString& GraphPath);
 
 	/** Bring these revisions of graph into focus on main display*/
 	void FocusOnGraphRevisions( struct FListItemGraphToDiff* Diff);
