[NO NON-PII COMMIT MESSAGE]

diff --git a/Engine/Source/Editor/GraphEditor/Private/SGraphEditorImpl.cpp b/Engine/Source/Editor/GraphEditor/Private/SGraphEditorImpl.cpp
--- a/Engine/Source/Editor/GraphEditor/Private/SGraphEditorImpl.cpp
+++ b/Engine/Source/Editor/GraphEditor/Private/SGraphEditorImpl.cpp
@@ -727,7 +727,7 @@ void SGraphEditorImpl::UnlockFromGraphEditor( TWeakPtr<SGraphEditor> Other )
 {
 	check(Other.IsValid());
 	int idx = LockedGraphs.Find(Other);
-	if( ensureMsgf(idx != INDEX_NONE, TEXT("Attempted to unlock graphs that were not locked together: %s %s"), *GetReadableLocation(), *(Other.Pin()->GetReadableLocation()) ) )
+	if (idx != INDEX_NONE)
 	{
 		LockedGraphs.RemoveAtSwap(idx);
 	}
