[NO NON-PII COMMIT MESSAGE]

diff --git a/Engine/Source/Editor/GraphEditor/Private/SGraphEditorImpl.cpp b/Engine/Source/Editor/GraphEditor/Private/SGraphEditorImpl.cpp
--- a/Engine/Source/Editor/GraphEditor/Private/SGraphEditorImpl.cpp
+++ b/Engine/Source/Editor/GraphEditor/Private/SGraphEditorImpl.cpp
@@ -287,6 +287,7 @@ void SGraphEditorImpl::Construct( const FArguments& InArgs )
 
 	OnFocused = InArgs._GraphEvents.OnFocused;
 	OnCreateActionMenu = InArgs._GraphEvents.OnCreateActionMenu;
+	OnCreateNodeOrPinMenu = InArgs._GraphEvents.OnCreateNodeOrPinMenu;
 	
 	struct Local
 	{
@@ -531,9 +532,16 @@ FActionMenuContent SGraphEditorImpl::GraphEd_OnGetContextMenuFor(const FGraphCon
 			// Show the menu for the pin or node under the cursor
 			const bool bShouldCloseAfterAction = true;
 			FMenuBuilder MenuBuilder( bShouldCloseAfterAction, this->Commands, MenuExtender );
-			Schema->GetContextMenuActions(EdGraphObj, SpawnInfo.GraphNode, SpawnInfo.GraphPin, &MenuBuilder, !IsEditable.Get());
 
-			Result = FActionMenuContent(MenuBuilder.MakeWidget());
+			if (OnCreateNodeOrPinMenu.IsBound())
+			{
+				Result = OnCreateNodeOrPinMenu.Execute(EdGraphObj, SpawnInfo.GraphNode, SpawnInfo.GraphPin, &MenuBuilder, !IsEditable.Get());
+			}
+			else
+			{
+				Schema->GetContextMenuActions(EdGraphObj, SpawnInfo.GraphNode, SpawnInfo.GraphPin, &MenuBuilder, !IsEditable.Get());
+				Result = FActionMenuContent(MenuBuilder.MakeWidget());
+			}
 		}
 		else if (IsEditable.Get())
 		{
diff --git a/Engine/Source/Editor/GraphEditor/Private/SGraphEditorImpl.h b/Engine/Source/Editor/GraphEditor/Private/SGraphEditorImpl.h
--- a/Engine/Source/Editor/GraphEditor/Private/SGraphEditorImpl.h
+++ b/Engine/Source/Editor/GraphEditor/Private/SGraphEditorImpl.h
@@ -195,6 +195,7 @@ private:
 
 	SGraphEditor::FOnFocused OnFocused;
 	SGraphEditor::FOnCreateActionMenu OnCreateActionMenu;
+	SGraphEditor::FOnCreateNodeOrPinMenu OnCreateNodeOrPinMenu;
 
 	TAttribute<bool> IsEditable;
 
diff --git a/Engine/Source/Editor/Kismet/Private/SBlueprintDiff.cpp b/Engine/Source/Editor/Kismet/Private/SBlueprintDiff.cpp
--- a/Engine/Source/Editor/Kismet/Private/SBlueprintDiff.cpp
+++ b/Engine/Source/Editor/Kismet/Private/SBlueprintDiff.cpp
@@ -1286,7 +1286,15 @@ void FDiffPanel::GeneratePanel(UEdGraph* Graph, UEdGraph* GraphToDiff )
 			{
 				Container->ShowDetailsForObjects(SelectionSet.Array());
 			};
+
+			const auto ContextMenuHandler = [](UEdGraph* CurrentGraph, const UEdGraphNode* InGraphNode, const UEdGraphPin* InGraphPin, FMenuBuilder* MenuBuilder, bool bIsDebugging)
+			{
+				MenuBuilder->AddMenuEntry(FGenericCommands::Get().Copy);
+				return FActionMenuContent(MenuBuilder->MakeWidget());
+			};
+
 			InEvents.OnSelectionChanged = SGraphEditor::FOnSelectionChanged::CreateStatic(SelectionChangedHandler, DetailsView);
+			InEvents.OnCreateNodeOrPinMenu = SGraphEditor::FOnCreateNodeOrPinMenu::CreateStatic(ContextMenuHandler);
 		}
 
 		if ( !GraphEditorCommands.IsValid() )
diff --git a/Engine/Source/Editor/UnrealEd/Public/GraphEditor.h b/Engine/Source/Editor/UnrealEd/Public/GraphEditor.h
--- a/Engine/Source/Editor/UnrealEd/Public/GraphEditor.h
+++ b/Engine/Source/Editor/UnrealEd/Public/GraphEditor.h
@@ -86,6 +86,8 @@ public:
 
 	DECLARE_DELEGATE_RetVal_FiveParams( FActionMenuContent, FOnCreateActionMenu, UEdGraph*, const FVector2D&, const TArray<UEdGraphPin*>&, bool, FActionMenuClosed );
 
+	DECLARE_DELEGATE_RetVal_FiveParams( FActionMenuContent, FOnCreateNodeOrPinMenu, UEdGraph*, const UEdGraphNode*, const UEdGraphPin*, FMenuBuilder*, bool);
+
 	DECLARE_DELEGATE_RetVal_TwoParams( FReply, FOnSpawnNodeByShortcut, FInputChord, const FVector2D& );
 
 	DECLARE_DELEGATE( FOnNodeSpawnedByKeymap );
@@ -109,8 +111,10 @@ public:
 		FOnNodeVerifyTextCommit OnVerifyTextCommit;
 		/** Called when text is committed on the graph */
 		FOnNodeTextCommitted OnTextCommitted;
-		/** Called to create context menu */
+		/** Called to create context menu for right clicking in empty area */
 		FOnCreateActionMenu OnCreateActionMenu;
+		/** Called to create context menu for right clicking a node or pin, same parameters as GetContextMenuActions on schema */
+		FOnCreateNodeOrPinMenu OnCreateNodeOrPinMenu;
 		/** Called to spawn a node in the graph using a shortcut */
 		FOnSpawnNodeByShortcut OnSpawnNodeByShortcut;
 		/** Called when a keymap spawns a node */
