[NO NON-PII COMMIT MESSAGE]

diff --git a/Engine/Source/Editor/BehaviorTreeEditor/Private/BehaviorTreeGraphNode.cpp b/Engine/Source/Editor/BehaviorTreeEditor/Private/BehaviorTreeGraphNode.cpp
--- a/Engine/Source/Editor/BehaviorTreeEditor/Private/BehaviorTreeGraphNode.cpp
+++ b/Engine/Source/Editor/BehaviorTreeEditor/Private/BehaviorTreeGraphNode.cpp
@@ -136,14 +136,27 @@ void UBehaviorTreeGraphNode::FindDiffs(UEdGraphNode* OtherNode, FDiffResults& Re
 				FGraphDiffControl::FNodeMatch NodeMatch;
 				NodeMatch.NewNode = RhsSubNode;
 
+				// Do two passes, exact and soft
 				for (UEdGraphNode* LhsSubNode : LhsSubNodes)
 				{
-					if (FGraphDiffControl::IsNodeMatch(LhsSubNode, RhsSubNode, &NodeMatches))
+					if (FGraphDiffControl::IsNodeMatch(LhsSubNode, RhsSubNode, true, &NodeMatches))
 					{
 						NodeMatch.OldNode = LhsSubNode;
 						break;
 					}
 				}
+				
+				if (NodeMatch.NewNode == nullptr)
+				{
+					for (UEdGraphNode* LhsSubNode : LhsSubNodes)
+					{
+						if (FGraphDiffControl::IsNodeMatch(LhsSubNode, RhsSubNode, false, &NodeMatches))
+						{
+							NodeMatch.OldNode = LhsSubNode;
+							break;
+						}
+					}
+				}
 
 				// if we found a corresponding node in the lhs graph, track it (so we can prevent future matches with the same nodes)
 				if (NodeMatch.IsValid())
@@ -168,18 +181,10 @@ void UBehaviorTreeGraphNode::FindDiffs(UEdGraphNode* OtherNode, FDiffResults& Re
 					continue;
 				}
 
+				// There can't be a matching node in RhsGraph because it would have been found above
 				FGraphDiffControl::FNodeMatch NodeMatch;
 				NodeMatch.NewNode = LhsSubNode;
 
-				for (UEdGraphNode* RhsSubNode : RhsSubNodes)
-				{
-					if (FGraphDiffControl::IsNodeMatch(LhsSubNode, RhsSubNode, &NodeMatches))
-					{
-						NodeMatch.OldNode = RhsSubNode;
-						break;
-					}
-				}
-
 				NodeMatch.Diff(SubtractiveDiffContext, Results);
 			}
 		};
diff --git a/Engine/Source/Editor/GraphEditor/Private/GraphDiffControl.cpp b/Engine/Source/Editor/GraphEditor/Private/GraphDiffControl.cpp
--- a/Engine/Source/Editor/GraphEditor/Private/GraphDiffControl.cpp
+++ b/Engine/Source/Editor/GraphEditor/Private/GraphDiffControl.cpp
@@ -418,10 +418,10 @@ static bool IsPinTypeDifferent(const FEdGraphPinType& T1, const FEdGraphPinType&
 	return bIsDifferent;
 }
 
-/** Find Pin in array that matches pin */
+/** Find linked pin in array that matches pin */
 static UEdGraphPin* FindOtherLink(TArray<UEdGraphPin*>& Links2, int32 OriginalIndex, UEdGraphPin* PinToFind)
 {
-	//sometimes the order of the pins is different between revisions, although the pins themselves are unchanged, so we have to look at all of them
+	// Sometimes the order of the pins is different between revisions, although the pins themselves are unchanged, so we have to look at all of them
 	UEdGraphNode* Node1 = PinToFind->GetOwningNode();
 	for (UEdGraphPin* Other : Links2)
 	{
@@ -588,20 +588,29 @@ bool FGraphDiffControl::FNodeMatch::Diff(const FNodeDiffContext& DiffContext, FD
 * FGraphDiffControl
 *******************************************************************************/
 
-FGraphDiffControl::FNodeMatch FGraphDiffControl::FindNodeMatch(UEdGraph* Graph, UEdGraphNode* Node, TArray<FNodeMatch> const& PriorMatches)
+FGraphDiffControl::FNodeMatch FGraphDiffControl::FindNodeMatch(UEdGraph* OldGraph, UEdGraphNode* NewNode, TArray<FNodeMatch> const& PriorMatches)
 {
 	FNodeMatch Match;
-	Match.NewNode = Node;
+	Match.NewNode = NewNode;
 
-	if (Graph)
+	if (OldGraph)
 	{
-		// attempt to find a node matching 'Node'
-		for (UEdGraphNode* GraphNode : Graph->Nodes) 
+		// Attempt to find a node matching 'NewNode', first try exact then try soft match
+		for (UEdGraphNode* GraphNode : OldGraph->Nodes)
 		{
-			if (GraphNode && IsNodeMatch(Node, GraphNode, &PriorMatches))
+			if (GraphNode && IsNodeMatch(NewNode, GraphNode, true, &PriorMatches))
 			{
 				Match.OldNode = GraphNode;
-				break;
+				return Match;
+			}
+		}
+
+		for (UEdGraphNode* GraphNode : OldGraph->Nodes)
+		{
+			if (GraphNode && IsNodeMatch(NewNode, GraphNode, false, &PriorMatches))
+			{
+				Match.OldNode = GraphNode;
+				return Match;
 			}
 		}
 	}	
@@ -609,7 +618,7 @@ FGraphDiffControl::FNodeMatch FGraphDiffControl::FindNodeMatch(UEdGraph* Graph,
 	return Match;
 }
 
-bool FGraphDiffControl::IsNodeMatch(UEdGraphNode* Node1, UEdGraphNode* Node2, TArray<FGraphDiffControl::FNodeMatch> const* Exclusions)
+bool FGraphDiffControl::IsNodeMatch(UEdGraphNode* Node1, UEdGraphNode* Node2, bool bExactOnly, TArray<FGraphDiffControl::FNodeMatch> const* Exclusions)
 {
 	if(Node2->GetClass() != Node1->GetClass()) 
 	{
@@ -621,19 +630,14 @@ bool FGraphDiffControl::IsNodeMatch(UEdGraphNode* Node1, UEdGraphNode* Node2, TA
 		return true;
 	}
 
-	// we could be diffing two completely separate assets, this makes sure both 
-	// nodes historically belong to the same graph
-	bool bIsIntraAssetDiff = (Node1->GetGraph()->GraphGuid == Node2->GetGraph()->GraphGuid);
-
-	// if both nodes are from the same graph
-	if (bIsIntraAssetDiff)
+	if (bExactOnly)
 	{
-		return (Node1->GetFName() == Node2->GetFName());
+		return false;
 	}
 
 	if (Exclusions)
 	{
-		// have to see if this node has already been matched with another
+		// Have to see if this node has already been matched with another, if so don't allow a soft match
 		for (const FGraphDiffControl::FNodeMatch& PriorMatch : *Exclusions)
 		{
 			if (!PriorMatch.IsValid())
@@ -652,8 +656,7 @@ bool FGraphDiffControl::IsNodeMatch(UEdGraphNode* Node1, UEdGraphNode* Node2, TA
 		}
 	}
 
-	// the name hashes won't match for nodes from separate graph assets, so we 
-	// need to look for some kind of semblance between the two... title? (what's displayed to the user)
+	// For a soft match use the node title, which includes the function name and target usually
 	FText Title1 = Node1->GetNodeTitle(ENodeTitleType::FullTitle);
 	FText Title2 = Node2->GetNodeTitle(ENodeTitleType::FullTitle);
 
@@ -703,7 +706,10 @@ bool FGraphDiffControl::DiffGraphs(UEdGraph* const LhsGraph, UEdGraph* const Rhs
 				continue;
 			}
 
-			FGraphDiffControl::FNodeMatch NodeMatch = FGraphDiffControl::FindNodeMatch(RhsGraph, LhsNode, NodeMatches);
+			// There can't be a matching node in RhsGraph because it would have been found above
+			FGraphDiffControl::FNodeMatch NodeMatch;
+			NodeMatch.NewNode = LhsNode;
+
 			bFoundDifferences |= NodeMatch.Diff(SubtractiveDiffContext, &DiffsOut);
 		}
 	}
diff --git a/Engine/Source/Editor/GraphEditor/Public/GraphDiffControl.h b/Engine/Source/Editor/GraphEditor/Public/GraphDiffControl.h
--- a/Engine/Source/Editor/GraphEditor/Public/GraphDiffControl.h
+++ b/Engine/Source/Editor/GraphEditor/Public/GraphDiffControl.h
@@ -4,6 +4,9 @@
 #include "CoreMinimal.h"
 #include "DiffResults.h"
 
+class UEdGraph;
+class UEdGraphNode;
+
 /** Used to find differences between revisions of a graph. */
 class GRAPHEDITOR_API FGraphDiffControl
 {
@@ -65,27 +68,27 @@ public:
 	struct GRAPHEDITOR_API FNodeMatch
 	{
 		FNodeMatch()
-			: NewNode(NULL)
-			, OldNode(NULL)
+			: NewNode(nullptr)
+			, OldNode(nullptr)
 		{}
 
-		class UEdGraphNode* NewNode;
-		class UEdGraphNode* OldNode;
+		UEdGraphNode* NewNode;
+		UEdGraphNode* OldNode;
 
 		/**
 		 * Looks at the two node members and compares them to see if there are any
-		 * differences. If one of the nodes is NULL, then this will return true with
+		 * differences. If one of the nodes is nullptr, then this will return true with
 		 * a EDiffType::NODE_ADDED result.
 		 * 
 		 * @param  DiffsArrayOut	If supplied, this will be filled out with all the differences that were found.
 		 * @return True if there were differences found, false if the two nodes are identical.
 		 */
-		bool Diff(const FNodeDiffContext& DiffContext, TArray<FDiffSingleResult>* DiffsResultsOut = NULL) const;
+		bool Diff(const FNodeDiffContext& DiffContext, TArray<FDiffSingleResult>* DiffsResultsOut = nullptr) const;
 		bool Diff(const FNodeDiffContext& DiffContext, FDiffResults& DiffsOut) const;
 
 		/**
 		 * Checks to see if this is a valid match.
-		 * @return False if NewNode or OldNode is NULL, true if both are valid.
+		 * @return False if NewNode or OldNode is nullptr, true if both are valid.
 		 */
 		bool IsValid() const;
 	};
@@ -96,15 +99,21 @@ public:
 	 * match, so providing a list of already matched nodes helps us narrow it down (and 
 	 * prevents us from matching one node with multiple others).
 	 * 
-	 * @param  Graph		The graph you want to search.
-	 * @param  Node			The node you want to match.
+	 * @param  OldGraph		The graph you want to search.
+	 * @param  NewNode		The new node you want to match.
 	 * @param  PriorMatches	Previous made matches to exclude from our search.
-	 * @return A pair of nodes (including the supplied one) that best match each other (one may be NULL if no match was found).
+	 * @return A pair of nodes (including the supplied one) that best match each other (one may be nullptr if no match was found).
 	 */
-	static FNodeMatch FindNodeMatch(class UEdGraph* OldGraph, class UEdGraphNode* Node, TArray<FNodeMatch> const& PriorMatches);
+	static FNodeMatch FindNodeMatch(UEdGraph* OldGraph, UEdGraphNode* NewNode, TArray<FNodeMatch> const& PriorMatches);
 
-	/** Determine if the two Nodes are the same */
-	static bool IsNodeMatch(class UEdGraphNode* Node1, class UEdGraphNode* Node2, TArray<FGraphDiffControl::FNodeMatch> const* Exclusions = nullptr);
+	/** 
+	 * Returns true if the two Nodes are the same 
+	 * 
+	 * @param  Node1		First node to check
+	 * @param  Node2		Second node
+	 * @param  Exclusions	Previous made matches to exclude from our search.
+	 */
+	static bool IsNodeMatch(UEdGraphNode* Node1, UEdGraphNode* Node2, bool bExactOnly = false, TArray<FGraphDiffControl::FNodeMatch> const* Exclusions = nullptr);
 
 	/**
 	 * Looks for node differences between the two supplied graphs. Diffs will be
@@ -115,7 +124,7 @@ public:
 	 * @param  DiffsOut		All the differences that were found between the two.
 	 * @return True if any differences were found, false if both graphs are identical.
 	 */
-	static bool DiffGraphs(class UEdGraph* const OldGraph, class UEdGraph* const NewGraph, TArray<FDiffSingleResult>& DiffsOut);
+	static bool DiffGraphs(UEdGraph* const OldGraph, UEdGraph* const NewGraph, TArray<FDiffSingleResult>& DiffsOut);
 
 	/** Computes an object path for a specific graph, relative to the blueprint/asset root */
 	static FString GetGraphPath(UEdGraph* Graph);
