PR #5837: Display enabled state for impure compact nodes (Contributed by [REDACTED_PERSON])

diff --git a/Engine/Source/Editor/GraphEditor/Private/KismetNodes/SGraphNodeK2Base.cpp b/Engine/Source/Editor/GraphEditor/Private/KismetNodes/SGraphNodeK2Base.cpp
--- a/Engine/Source/Editor/GraphEditor/Private/KismetNodes/SGraphNodeK2Base.cpp
+++ b/Engine/Source/Editor/GraphEditor/Private/KismetNodes/SGraphNodeK2Base.cpp
@@ -136,10 +136,8 @@ void SGraphNodeK2Base::UpdateCompactNode()
 	//            |_______|______|_______|
 	//
 	this->ContentScale.Bind( this, &SGraphNode::GetContentScale );
-	this->GetOrAddSlot( ENodeZone::Center )
-	.HAlign(HAlign_Center)
-	.VAlign(VAlign_Center)
-	[
+	
+	TSharedRef<SVerticalBox> InnerVerticalBox =
 		SNew(SVerticalBox)
 		+SVerticalBox::Slot()
 		[
@@ -160,13 +158,33 @@ void SGraphNodeK2Base::UpdateCompactNode()
 			[
 				NodeOverlay
 			]
-		]
-		+SVerticalBox::Slot()
+		];
+	
+	TSharedPtr<SWidget> EnabledStateWidget = GetEnabledStateWidget();
+	if (EnabledStateWidget.IsValid())
+	{
+		InnerVerticalBox->AddSlot()
+			.AutoHeight()
+			.HAlign(HAlign_Fill)
+			.VAlign(VAlign_Top)
+			.Padding(FMargin(2, 0))
+			[
+				EnabledStateWidget.ToSharedRef()
+			];
+	}
+
+	InnerVerticalBox->AddSlot()
 		.AutoHeight()
 		.Padding( FMargin(5.0f, 1.0f) )
 		[
 			ErrorReporting->AsWidget()
-		]
+		];
+
+	this->GetOrAddSlot( ENodeZone::Center )
+	.HAlign(HAlign_Center)
+	.VAlign(VAlign_Center)
+	[
+		InnerVerticalBox
 	];
 
 	CreatePinWidgets();
diff --git a/Engine/Source/Editor/GraphEditor/Private/SGraphNode.cpp b/Engine/Source/Editor/GraphEditor/Private/SGraphNode.cpp
--- a/Engine/Source/Editor/GraphEditor/Private/SGraphNode.cpp
+++ b/Engine/Source/Editor/GraphEditor/Private/SGraphNode.cpp
@@ -879,33 +879,16 @@ void SGraphNode::UpdateGraphNode()
 			CreateNodeContentArea()
 		];
 
-	if ((GraphNode->GetDesiredEnabledState() != ENodeEnabledState::Enabled) && !GraphNode->IsAutomaticallyPlacedGhostNode())
+	TSharedPtr<SWidget> EnabledStateWidget = GetEnabledStateWidget();
+	if (EnabledStateWidget.IsValid())
 	{
-		const bool bDevelopmentOnly = GraphNode->GetDesiredEnabledState() == ENodeEnabledState::DevelopmentOnly;
-		const FText StatusMessage = bDevelopmentOnly ? NSLOCTEXT("SGraphNode", "DevelopmentOnly", "Development Only") : NSLOCTEXT("SGraphNode", "DisabledNode", "Disabled");
-		const FText StatusMessageTooltip = bDevelopmentOnly ?
-			NSLOCTEXT("SGraphNode", "DevelopmentOnlyTooltip", "This node will only be executed in the editor and in Development builds in a packaged game (it will be treated as disabled in Shipping or Test builds cooked from a commandlet)") :
-			NSLOCTEXT("SGraphNode", "DisabledNodeTooltip", "This node is currently disabled and will not be executed");
-
 		InnerVerticalBox->AddSlot()
 			.AutoHeight()
 			.HAlign(HAlign_Fill)
 			.VAlign(VAlign_Top)
 			.Padding(FMargin(2, 0))
 			[
-				SNew(SBorder)
-				.BorderImage(FEditorStyle::GetBrush(bDevelopmentOnly ? "Graph.Node.DevelopmentBanner" : "Graph.Node.DisabledBanner"))
-				.HAlign(HAlign_Fill)
-				.VAlign(VAlign_Fill)
-				[
-					SNew(STextBlock)
-					.Text(StatusMessage)
-					.ToolTipText(StatusMessageTooltip)
-					.Justification(ETextJustify::Center)
-					.ColorAndOpacity(FLinearColor::White)
-					.ShadowOffset(FVector2D::UnitVector)
-					.Visibility(EVisibility::Visible)
-				]
+				EnabledStateWidget.ToSharedRef()
 			];
 	}
 
@@ -977,6 +960,34 @@ void SGraphNode::UpdateGraphNode()
 END_SLATE_FUNCTION_BUILD_OPTIMIZATION
 
 
+TSharedPtr<SWidget> SGraphNode::GetEnabledStateWidget()
+{
+	if ((GraphNode->GetDesiredEnabledState() != ENodeEnabledState::Enabled) && !GraphNode->IsAutomaticallyPlacedGhostNode())
+	{
+		const bool bDevelopmentOnly = GraphNode->GetDesiredEnabledState() == ENodeEnabledState::DevelopmentOnly;
+		const FText StatusMessage = bDevelopmentOnly ? NSLOCTEXT("SGraphNode", "DevelopmentOnly", "Development Only") : NSLOCTEXT("SGraphNode", "DisabledNode", "Disabled");
+		const FText StatusMessageTooltip = bDevelopmentOnly ?
+			NSLOCTEXT("SGraphNode", "DevelopmentOnlyTooltip", "This node will only be executed in the editor and in Development builds in a packaged game (it will be treated as disabled in Shipping or Test builds cooked from a commandlet)") :
+			NSLOCTEXT("SGraphNode", "DisabledNodeTooltip", "This node is currently disabled and will not be executed");
+
+		return SNew(SBorder)
+			.BorderImage(FEditorStyle::GetBrush(bDevelopmentOnly ? "Graph.Node.DevelopmentBanner" : "Graph.Node.DisabledBanner"))
+			.HAlign(HAlign_Fill)
+			.VAlign(VAlign_Fill)
+			[
+				SNew(STextBlock)
+				.Text(StatusMessage)
+				.ToolTipText(StatusMessageTooltip)
+				.Justification(ETextJustify::Center)
+				.ColorAndOpacity(FLinearColor::White)
+				.ShadowOffset(FVector2D::UnitVector)
+				.Visibility(EVisibility::Visible)
+			];
+	}
+
+	return TSharedPtr<SWidget>();
+}
+
 TSharedRef<SWidget> SGraphNode::CreateNodeContentArea()
 {
 	// NODE CONTENT AREA
diff --git a/Engine/Source/Editor/GraphEditor/Public/SGraphNode.h b/Engine/Source/Editor/GraphEditor/Public/SGraphNode.h
--- a/Engine/Source/Editor/GraphEditor/Public/SGraphNode.h
+++ b/Engine/Source/Editor/GraphEditor/Public/SGraphNode.h
@@ -356,6 +356,9 @@ protected:
 	/** Returns TRUE if the input pin should be hidden from view */
 	bool ShouldPinBeHidden(const UEdGraphPin* InPin) const;
 
+	/** Returns the widget to use for the enabled state of the node */
+	TSharedPtr<SWidget> GetEnabledStateWidget();
+
 protected:
 	/** Input pin widgets on this node */
 	TArray< TSharedRef<SGraphPin> > InputPins;
