PR #5614: Fix alignment of impure compact node exec pins (Contributed by [REDACTED_PERSON])

diff --git a/Engine/Source/Editor/GraphEditor/Private/KismetNodes/SGraphNodeK2Base.cpp b/Engine/Source/Editor/GraphEditor/Private/KismetNodes/SGraphNodeK2Base.cpp
--- a/Engine/Source/Editor/GraphEditor/Private/KismetNodes/SGraphNodeK2Base.cpp
+++ b/Engine/Source/Editor/GraphEditor/Private/KismetNodes/SGraphNodeK2Base.cpp
@@ -108,10 +108,26 @@ void SGraphNodeK2Base::UpdateCompactNode()
 			]
 		];
 	
+	// Default to "pure" styling, where we can just center the pins vertically
+	// since don't need to worry about alignment with other nodes
+	float PinPaddingTop = 0.f;
+	EVerticalAlignment PinVerticalAlignment = VAlign_Center;
+
+	// But if this is an impure node, we'll align the pins to the top, 
+	// and add some padding so that the exec pins line up with the exec pins of other nodes
+	if (UK2Node* K2Node = Cast<UK2Node>(GraphNode))
+	{
+		if (!K2Node->IsNodePure())
+		{
+			PinPaddingTop += 8.0f;
+			PinVerticalAlignment = VAlign_Top;
+		}
+	}
+	
 	NodeOverlay->AddSlot()
 		.HAlign(HAlign_Left)
-		.VAlign(VAlign_Center)
-		.Padding(0.f, 0.f, 55.f, 0.f)
+		.VAlign(PinVerticalAlignment)
+		.Padding(0.f, PinPaddingTop, 55.f, 0.f)
 		[
 			// LEFT
 			SAssignNew(LeftNodeBox, SVerticalBox)
@@ -119,8 +135,8 @@ void SGraphNodeK2Base::UpdateCompactNode()
 
 	NodeOverlay->AddSlot()
 		.HAlign(HAlign_Right)
-		.VAlign(VAlign_Center)
-		.Padding(55.f, 0.f, 0.f, 0.f)
+		.VAlign(PinVerticalAlignment)
+		.Padding(55.f, PinPaddingTop, 0.f, 0.f)
 		[
 			// RIGHT
 			SAssignNew(RightNodeBox, SVerticalBox)
