Expanding FKey fix to properly work with blueprint function parameters

diff --git a/Engine/Source/Editor/GraphEditor/Private/KismetPins/SGraphPinKey.cpp b/Engine/Source/Editor/GraphEditor/Private/KismetPins/SGraphPinKey.cpp
--- a/Engine/Source/Editor/GraphEditor/Private/KismetPins/SGraphPinKey.cpp
+++ b/Engine/Source/Editor/GraphEditor/Private/KismetPins/SGraphPinKey.cpp
@@ -7,25 +7,27 @@
 
 void SGraphPinKey::Construct(const FArguments& InArgs, UEdGraphPin* InGraphPinObj)
 {
-
-
-	TArray<FKey> KeyList;
-	EKeys::GetAllKeys(KeyList);
-
 	SelectedKey = FKey(*InGraphPinObj->GetDefaultAsString());
 
+	InGraphPinObj->AutogeneratedDefaultValue = "None";
+	if (SelectedKey.GetFName() == NAME_None)
+	{
+		InGraphPinObj->GetSchema()->ResetPinToAutogeneratedDefaultValue(InGraphPinObj, false);
+		SelectedKey = FKey(*InGraphPinObj->GetDefaultAsString());
+	}
+
 	if (InGraphPinObj->Direction == EEdGraphPinDirection::EGPD_Input)
 	{
 		// None is a valid key
 		if (SelectedKey.GetFName() == NAME_None)
 		{
 			SelectedKey = EKeys::Invalid;
-			InGraphPinObj->AutogeneratedDefaultValue = "None";
-			InGraphPinObj->GetSchema()->ResetPinToAutogeneratedDefaultValue(InGraphPinObj, false);
 		}
 		else if (!SelectedKey.IsValid())
 		{
 			// Ensure first valid key is always set by default
+			TArray<FKey> KeyList;
+			EKeys::GetAllKeys(KeyList);
 			SelectedKey = KeyList[0];
 			InGraphPinObj->GetSchema()->TrySetDefaultValue(*InGraphPinObj, SelectedKey.ToString());
 		}
