More fixes for anim montages

Object graph pins now support DisallowedClasses.
Exposed node property pin metadata correctly in anim nodes.
Fixed ensures in sequence player.

diff --git a/Engine/Source/Editor/AnimGraph/Classes/AnimGraphNode_Base.h b/Engine/Source/Editor/AnimGraph/Classes/AnimGraphNode_Base.h
--- a/Engine/Source/Editor/AnimGraph/Classes/AnimGraphNode_Base.h
+++ b/Engine/Source/Editor/AnimGraph/Classes/AnimGraphNode_Base.h
@@ -143,6 +143,7 @@ class ANIMGRAPH_API UAnimGraphNode_Base : public UK2Node
 	virtual void GetPinHoverText(const UEdGraphPin& Pin, FString& HoverTextOut) const override;
 	virtual bool ShowPaletteIconOnNode() const override{ return false; }
 	virtual void PinDefaultValueChanged(UEdGraphPin* Pin) override;
+	virtual FString GetPinMetaData(FName InPinName, FName InKey) override;
 	// End of UEdGraphNode interface
 
 	// UK2Node interface
diff --git a/Engine/Source/Editor/AnimGraph/Private/AnimGraphNode_Base.cpp b/Engine/Source/Editor/AnimGraph/Private/AnimGraphNode_Base.cpp
--- a/Engine/Source/Editor/AnimGraph/Private/AnimGraphNode_Base.cpp
+++ b/Engine/Source/Editor/AnimGraph/Private/AnimGraphNode_Base.cpp
@@ -410,6 +410,27 @@ void UAnimGraphNode_Base::PinDefaultValueChanged(UEdGraphPin* Pin)
 	}
 }
 
+FString UAnimGraphNode_Base::GetPinMetaData(FName InPinName, FName InKey)
+{
+	FString MetaData = Super::GetPinMetaData(InPinName, InKey);
+	if(MetaData.IsEmpty())
+	{
+		// Check properties of our anim node
+		if(UStructProperty* NodeStructProperty = GetFNodeProperty())
+		{
+			for (TFieldIterator<UProperty> It(NodeStructProperty->Struct); It; ++It)
+			{
+				const UProperty* Property = *It;
+				if (Property && Property->GetFName() == InPinName)
+				{
+					return Property->GetMetaData(InKey);
+				}
+			}
+		}
+	}
+	return MetaData;
+}
+
 bool UAnimGraphNode_Base::IsPinExposedAndLinked(const FString& InPinName, const EEdGraphPinDirection InDirection) const
 {
 	UEdGraphPin* Pin = FindPin(InPinName, InDirection);
diff --git a/Engine/Source/Editor/GraphEditor/Private/KismetPins/SGraphPinObject.cpp b/Engine/Source/Editor/GraphEditor/Private/KismetPins/SGraphPinObject.cpp
--- a/Engine/Source/Editor/GraphEditor/Private/KismetPins/SGraphPinObject.cpp
+++ b/Engine/Source/Editor/GraphEditor/Private/KismetPins/SGraphPinObject.cpp
@@ -210,19 +210,30 @@ TSharedRef<SWidget> SGraphPinObject::GenerateAssetPicker()
 	AssetPickerConfig.InitialAssetViewType = EAssetViewType::List;
 	AssetPickerConfig.bAllowDragging = false;
 
-	// Check with the node to see if there is any "AllowClasses" metadata for the pin
-	FString ClassFilterString = GraphPinObj->GetOwningNode()->GetPinMetaData(GraphPinObj->PinName, FName(TEXT("AllowedClasses")));
-	if( !ClassFilterString.IsEmpty() )
+	// Check with the node to see if there is any "AllowClasses" or "DisallowedClasses" metadata for the pin
+	FString AllowedClassesFilterString = GraphPinObj->GetOwningNode()->GetPinMetaData(GraphPinObj->PinName, FName(TEXT("AllowedClasses")));
+	if( !AllowedClassesFilterString.IsEmpty() )
 	{
 		// Clear out the allowed class names and have the pin's metadata override.
 		AssetPickerConfig.Filter.ClassNames.Empty();
 
 		// Parse and add the classes from the metadata
-		TArray<FString> CustomClassFilterNames;
-		ClassFilterString.ParseIntoArray(CustomClassFilterNames, TEXT(","), true);
-		for(auto It = CustomClassFilterNames.CreateConstIterator(); It; ++It)
+		TArray<FString> AllowedClassesFilterNames;
+		AllowedClassesFilterString.ParseIntoArray(AllowedClassesFilterNames, TEXT(","), true);
+		for(const FString& AllowedClassesFilterName : AllowedClassesFilterNames)
 		{
-			AssetPickerConfig.Filter.ClassNames.Add(FName(**It));
+			AssetPickerConfig.Filter.ClassNames.Add(FName(*AllowedClassesFilterName));
+		}
+	}
+
+	FString DisallowedClassesFilterString = GraphPinObj->GetOwningNode()->GetPinMetaData(GraphPinObj->PinName, FName(TEXT("DisallowedClasses")));
+	if(!DisallowedClassesFilterString.IsEmpty())
+	{
+		TArray<FString> DisallowedClassesFilterNames;
+		DisallowedClassesFilterString.ParseIntoArray(DisallowedClassesFilterNames, TEXT(","), true);
+		for(const FString& DisallowedClassesFilterName : DisallowedClassesFilterNames)
+		{
+			AssetPickerConfig.Filter.RecursiveClassesExclusionSet.Add(FName(*DisallowedClassesFilterName));
 		}
 	}
 
diff --git a/Engine/Source/Runtime/Engine/Classes/Animation/AnimNode_SequencePlayer.h b/Engine/Source/Runtime/Engine/Classes/Animation/AnimNode_SequencePlayer.h
--- a/Engine/Source/Runtime/Engine/Classes/Animation/AnimNode_SequencePlayer.h
+++ b/Engine/Source/Runtime/Engine/Classes/Animation/AnimNode_SequencePlayer.h
@@ -18,7 +18,7 @@ struct ENGINE_API FAnimNode_SequencePlayer : public FAnimNode_AssetPlayerBase
 	GENERATED_USTRUCT_BODY()
 public:
 	// The animation sequence asset to play
-	UPROPERTY(EditAnywhere, BlueprintReadWrite, Category = Settings, meta = (PinHiddenByDefault))
+	UPROPERTY(EditAnywhere, BlueprintReadWrite, Category = Settings, meta = (PinHiddenByDefault, DisallowedClasses="AnimMontage"))
 	UAnimSequenceBase* Sequence;
 
 	// The Basis in which the PlayRate is expressed in. This is used to rescale PlayRate inputs.
diff --git a/Engine/Source/Runtime/Engine/Private/Animation/AnimNode_SequencePlayer.cpp b/Engine/Source/Runtime/Engine/Private/Animation/AnimNode_SequencePlayer.cpp
--- a/Engine/Source/Runtime/Engine/Private/Animation/AnimNode_SequencePlayer.cpp
+++ b/Engine/Source/Runtime/Engine/Private/Animation/AnimNode_SequencePlayer.cpp
@@ -33,7 +33,7 @@ void FAnimNode_SequencePlayer::Initialize_AnyThread(const FAnimationInitializeCo
 
 	GetEvaluateGraphExposedInputs().Execute(Context);
 
-	if(ensureMsgf(Sequence->IsA<UAnimMontage>(), TEXT("Sequence players do not support anim montages.")))
+	if(ensureMsgf(!Sequence->IsA<UAnimMontage>(), TEXT("Sequence players do not support anim montages.")))
 	{
 		Sequence = nullptr;
 	}
@@ -61,7 +61,7 @@ void FAnimNode_SequencePlayer::UpdateAssetPlayer(const FAnimationUpdateContext&
 {
 	GetEvaluateGraphExposedInputs().Execute(Context);
 
-	if(ensureMsgf(Sequence->IsA<UAnimMontage>(), TEXT("Sequence players do not support anim montages.")))
+	if(ensureMsgf(!Sequence->IsA<UAnimMontage>(), TEXT("Sequence players do not support anim montages.")))
 	{
 		Sequence = nullptr;
 	}
