The dropdown marked "Select Function..." should have an additional option that creates a function/event that matches the specified signature.

diff --git a/Engine/Source/Editor/BlueprintGraph/Classes/EdGraphSchema_K2.h b/Engine/Source/Editor/BlueprintGraph/Classes/EdGraphSchema_K2.h
--- a/Engine/Source/Editor/BlueprintGraph/Classes/EdGraphSchema_K2.h
+++ b/Engine/Source/Editor/BlueprintGraph/Classes/EdGraphSchema_K2.h
@@ -790,7 +790,7 @@ public:
 	/** Can this function be called by kismet code */
 	static bool CanUserKismetCallFunction(const UFunction* Function);
 
-	/** Returns if hunction has output parameter(s) */
+	/** Returns if function has output parameter(s) */
 	static bool HasFunctionAnyOutputParameter(const UFunction* Function);
 
 	enum EDelegateFilterMode
diff --git a/Engine/Source/Editor/BlueprintGraph/Classes/K2Node_CustomEvent.h b/Engine/Source/Editor/BlueprintGraph/Classes/K2Node_CustomEvent.h
--- a/Engine/Source/Editor/BlueprintGraph/Classes/K2Node_CustomEvent.h
+++ b/Engine/Source/Editor/BlueprintGraph/Classes/K2Node_CustomEvent.h
@@ -92,6 +92,14 @@ class UK2Node_CustomEvent : public UK2Node_Event
 	 */
 	BLUEPRINTGRAPH_API uint32 GetNetFlags() const;
 
+	/**
+	* Updates the Signature of this event. Empties current pins and adds the new given parameters.
+	* Does nothing if DelegateSignature is nullptr
+	* 
+	* @param DelegateSignature	The new signature for this function to have
+	*/
+	BLUEPRINTGRAPH_API void SetDelegateSignature(const UFunction* DelegateSignature);
+
 private:
 	/** Constructing FText strings can be costly, so we cache the node's title */
 	FNodeTextCache CachedNodeTitle;
diff --git a/Engine/Source/Editor/BlueprintGraph/Private/K2Node_CustomEvent.cpp b/Engine/Source/Editor/BlueprintGraph/Private/K2Node_CustomEvent.cpp
--- a/Engine/Source/Editor/BlueprintGraph/Private/K2Node_CustomEvent.cpp
+++ b/Engine/Source/Editor/BlueprintGraph/Private/K2Node_CustomEvent.cpp
@@ -395,32 +395,44 @@ void UK2Node_CustomEvent::ReconstructNode()
 
 	if (bUseDelegateSignature)
 	{
-		UserDefinedPins.Empty();
-		for (TFieldIterator<UProperty> PropIt(DelegateSignature); PropIt && (PropIt->PropertyFlags & CPF_Parm); ++PropIt)
+		SetDelegateSignature(DelegateSignature);
+	}
+
+	Super::ReconstructNode();
+}
+
+void UK2Node_CustomEvent::SetDelegateSignature(const UFunction* DelegateSignature)
+{
+	if (DelegateSignature == nullptr)
+	{
+		return;
+	}
+
+	const UEdGraphSchema_K2* K2Schema = GetDefault<UEdGraphSchema_K2>();
+	
+	UserDefinedPins.Empty();
+	for (TFieldIterator<UProperty> PropIt(DelegateSignature); PropIt && (PropIt->PropertyFlags & CPF_Parm); ++PropIt)
+	{
+		const UProperty* Param = *PropIt;
+		if (!Param->HasAnyPropertyFlags(CPF_OutParm) || Param->HasAnyPropertyFlags(CPF_ReferenceParm))
 		{
-			const UProperty* Param = *PropIt;
-			if (!Param->HasAnyPropertyFlags(CPF_OutParm) || Param->HasAnyPropertyFlags(CPF_ReferenceParm))
-			{
-				FEdGraphPinType PinType;
-				K2Schema->ConvertPropertyToPinType(Param, /*out*/ PinType);
+			FEdGraphPinType PinType;
+			K2Schema->ConvertPropertyToPinType(Param, /*out*/ PinType);
 
-				FName NewPinName = Param->GetFName();
-				int32 Index = 1;
-				while ((DelegateOutputName == NewPinName) || (UEdGraphSchema_K2::PN_Then == NewPinName))
-				{
-					++Index;
-					NewPinName = *FString::Printf(TEXT("%s%d"), *NewPinName.ToString(), Index);
-				}
-				TSharedPtr<FUserPinInfo> NewPinInfo = MakeShareable( new FUserPinInfo() );
-				NewPinInfo->PinName = NewPinName;
-				NewPinInfo->PinType = PinType;
-				NewPinInfo->DesiredPinDirection = EGPD_Output;
-				UserDefinedPins.Add(NewPinInfo);
+			FName NewPinName = Param->GetFName();
+			int32 Index = 1;
+			while ((DelegateOutputName == NewPinName) || (UEdGraphSchema_K2::PN_Then == NewPinName))
+			{
+				++Index;
+				NewPinName = *FString::Printf(TEXT("%s%d"), *NewPinName.ToString(), Index);
 			}
+			TSharedPtr<FUserPinInfo> NewPinInfo = MakeShareable(new FUserPinInfo());
+			NewPinInfo->PinName = NewPinName;
+			NewPinInfo->PinType = PinType;
+			NewPinInfo->DesiredPinDirection = EGPD_Output;
+			UserDefinedPins.Add(NewPinInfo);
 		}
 	}
-
-	Super::ReconstructNode();
 }
 
 UK2Node_CustomEvent* UK2Node_CustomEvent::CreateFromFunction(FVector2D GraphPosition, UEdGraph* ParentGraph, const FString& Name, const UFunction* Function, bool bSelectNewNode/* = true*/)
diff --git a/Engine/Source/Editor/GraphEditor/Private/KismetNodes/SGraphNodeK2CreateDelegate.cpp b/Engine/Source/Editor/GraphEditor/Private/KismetNodes/SGraphNodeK2CreateDelegate.cpp
--- a/Engine/Source/Editor/GraphEditor/Private/KismetNodes/SGraphNodeK2CreateDelegate.cpp
+++ b/Engine/Source/Editor/GraphEditor/Private/KismetNodes/SGraphNodeK2CreateDelegate.cpp
@@ -9,24 +9,28 @@
 #include "EdGraphSchema_K2.h"
 #include "K2Node_CreateDelegate.h"
 
-FString SGraphNodeK2CreateDelegate::FunctionDescription(const UFunction* Function, const bool bOnlyDescribeSignature, const int32 CharacterLimit)
+#include "Editor/UnrealEd/Public/Kismet2/BlueprintEditorUtils.h"
+#include "Editor/UnrealEd/Public/Kismet2/KismetEditorUtilities.h"
+#include "Editor/BlueprintGraph/Classes/BlueprintNodeBinder.h"
+#include "Editor/BlueprintGraph/Classes/BlueprintEventNodeSpawner.h"
+#include "Editor/BlueprintGraph/Classes/K2Node_CustomEvent.h"
+#include "Editor/GraphEditor/Public/SGraphNode.h"
+
+FText SGraphNodeK2CreateDelegate::FunctionDescription(const UFunction* Function, const bool bOnlyDescribeSignature /*= false*/, const int32 CharacterLimit /*= 32*/)
 {
 	if (!Function || !Function->GetOuter())
 	{
-		return NSLOCTEXT("GraphNodeK2Create", "Error", "Error").ToString();
+		return NSLOCTEXT("GraphNodeK2Create", "Error", "Error");
 	}
 
 	const UEdGraphSchema_K2* K2Schema = GetDefault<UEdGraphSchema_K2>();
 
 	FString Result;
-
-	//Result += Function->GetOuter()->GetName() + TEXT("::") + Function->GetName();
-	//Result += TEXT("(");
-
+	
 	// Show function name.
 	if (!bOnlyDescribeSignature)
 	{
-		Result += Function->GetName();
+		Result = Function->GetName();
 	}
 
 	Result += TEXT("(");
@@ -109,7 +113,7 @@ FString SGraphNodeK2CreateDelegate::FunctionDescription(const UFunction* Functio
 		}
 	}
 
-	return Result;
+	return FText::FromString(Result);
 }
 
 void SGraphNodeK2CreateDelegate::Construct(const FArguments& InArgs, UK2Node* InNode)
@@ -121,17 +125,17 @@ void SGraphNodeK2CreateDelegate::Construct(const FArguments& InArgs, UK2Node* In
 FText SGraphNodeK2CreateDelegate::GetCurrentFunctionDescription() const
 {
 	UK2Node_CreateDelegate* Node = Cast<UK2Node_CreateDelegate>(GraphNode);
-	UFunction* FunctionSignature = Node ? Node->GetDelegateSignature() : NULL;
-	UClass* ScopeClass = Node ? Node->GetScopeClass() : NULL;
+	UFunction* FunctionSignature = Node ? Node->GetDelegateSignature() : nullptr;
+	UClass* ScopeClass = Node ? Node->GetScopeClass() : nullptr;
 
 	if (!FunctionSignature || !ScopeClass)
 	{
 		return FText::GetEmpty();
 	}
 
-	if (const auto Func = FindField<UFunction>(ScopeClass, Node->GetFunctionName()))
+	if (const UFunction* Func = FindField<UFunction>(ScopeClass, Node->GetFunctionName()))
 	{
-		return FText::FromString(FunctionDescription(Func));
+		return FunctionDescription(Func);
 	}
 
 	if (Node->GetFunctionName() != NAME_None)
@@ -147,7 +151,7 @@ TSharedRef<ITableRow> SGraphNodeK2CreateDelegate::HandleGenerateRowFunction(TSha
 	check(FunctionItemData.IsValid());
 	return SNew(STableRow< TSharedPtr<FFunctionItemData> >, OwnerTable).Content()
 		[
-			SNew(STextBlock).Text(FText::FromString(FunctionItemData->Description))
+			SNew(STextBlock).Text(FunctionItemData->Description)
 		];
 }
 
@@ -157,7 +161,52 @@ void SGraphNodeK2CreateDelegate::OnFunctionSelected(TSharedPtr<FFunctionItemData
 	{
 		if (UK2Node_CreateDelegate* Node = Cast<UK2Node_CreateDelegate>(GraphNode))
 		{
-			Node->SetFunction(FunctionItemData->Name);
+			UBlueprint* NodeBP = Node->GetBlueprint();
+			
+			check(NodeBP);
+
+			if (FunctionItemData == CreateMatchingFunctionData)
+			{
+				// Get a valid name for the function graph
+				FString ProposedFuncName = NodeBP->GetName() + "_AutoGenFunc";
+				FName NewFuncName = FBlueprintEditorUtils::GenerateUniqueGraphName(NodeBP, ProposedFuncName);
+				
+				UEdGraph* NewGraph = nullptr;
+				NewGraph = FBlueprintEditorUtils::CreateNewGraph(NodeBP, NewFuncName, UEdGraph::StaticClass(), UEdGraphSchema_K2::StaticClass());
+
+				if (NewGraph != nullptr)
+				{
+					FBlueprintEditorUtils::AddFunctionGraph<UFunction>(NodeBP, NewGraph, true, Node->GetDelegateSignature());
+					FKismetEditorUtilities::BringKismetToFocusAttentionOnObject(NewGraph);
+				}
+				
+				Node->SetFunction(NewFuncName);
+			}
+			else if (FunctionItemData == CreateMatchingEventData)
+			{
+				// Get a valid name for the function graph
+				FName NewEventName = FBlueprintEditorUtils::FindUniqueCustomEventName(NodeBP);
+
+				UBlueprintEventNodeSpawner* Spawner = UBlueprintEventNodeSpawner::Create(UK2Node_CustomEvent::StaticClass(), NewEventName);
+				UEdGraphNode* NewNode = Spawner->Invoke(Node->GetGraph(), IBlueprintNodeBinder::FBindingSet(), FVector2D(Node->NodePosX, Node->NodePosY + 200));
+
+				if (UK2Node_CustomEvent* NewEventNode = Cast<UK2Node_CustomEvent>(NewNode))
+				{
+					NewEventNode->SetDelegateSignature(Node->GetDelegateSignature());
+					// Reconstruct to get the new parameters to show in the editor
+					NewEventNode->ReconstructNode();
+					NewEventNode->bIsEditable = true;
+					FKismetEditorUtilities::BringKismetToFocusAttentionOnObject(NewEventNode);
+				}
+				
+				Node->SetFunction(NewEventName);
+				FBlueprintEditorUtils::MarkBlueprintAsStructurallyModified(NodeBP);				
+			}
+			else
+			{
+				Node->SetFunction(FunctionItemData->Name);
+			}
+			
 			Node->HandleAnyChange(true);
 
 			auto SelectFunctionWidgetPtr = SelectFunctionWidget.Pin();
@@ -169,6 +218,14 @@ void SGraphNodeK2CreateDelegate::OnFunctionSelected(TSharedPtr<FFunctionItemData
 	}
 }
 
+TSharedPtr<SGraphNodeK2CreateDelegate::FFunctionItemData> SGraphNodeK2CreateDelegate::AddDefaultFunctionDataOption(const FText& DisplayName)
+{
+	TSharedPtr<FFunctionItemData> NewEntry = MakeShareable(new FFunctionItemData());
+	NewEntry->Description = DisplayName;
+	FunctionDataItems.Add(NewEntry);
+	return NewEntry;
+}
+
 void SGraphNodeK2CreateDelegate::CreateBelowPinControls(TSharedPtr<SVerticalBox> MainBox)
 {
 	if (UK2Node_CreateDelegate* Node = Cast<UK2Node_CreateDelegate>(GraphNode))
@@ -181,14 +238,14 @@ void SGraphNodeK2CreateDelegate::CreateBelowPinControls(TSharedPtr<SVerticalBox>
 			FText FunctionSignaturePrompt;
 			{
 				FFormatNamedArguments FormatArguments;
-				FormatArguments.Add(TEXT("FunctionSignature"), FText::FromString(FunctionDescription(FunctionSignature, true)));
+				FormatArguments.Add(TEXT("FunctionSignature"), FunctionDescription(FunctionSignature, true));
 				FunctionSignaturePrompt = FText::Format(NSLOCTEXT("GraphNodeK2Create", "FunctionSignaturePrompt", "Signature: {FunctionSignature}"), FormatArguments);
 			}
 
 			FText FunctionSignatureToolTipText;
 			{
 				FFormatNamedArguments FormatArguments;
-				FormatArguments.Add(TEXT("FullFunctionSignature"), FText::FromString(FunctionDescription(FunctionSignature, true, INDEX_NONE)));
+				FormatArguments.Add(TEXT("FullFunctionSignature"), FunctionDescription(FunctionSignature, true, INDEX_NONE));
 				FunctionSignatureToolTipText = FText::Format(NSLOCTEXT("GraphNodeK2Create", "FunctionSignatureToolTip", "Signature Syntax: (Inputs) -> [Outputs]\nFull Signature:{FullFunctionSignature}"), FormatArguments);
 			}
 
@@ -204,6 +261,18 @@ void SGraphNodeK2CreateDelegate::CreateBelowPinControls(TSharedPtr<SVerticalBox>
 
 			FunctionDataItems.Empty();
 
+			// add an empty row, so the user can clear the selection if they want
+			AddDefaultFunctionDataOption(NSLOCTEXT("GraphNodeK2Create", "EmptyFunctionOption", "[None]"));
+
+			// Option to create a function based on the event parameters
+			CreateMatchingFunctionData = AddDefaultFunctionDataOption(NSLOCTEXT("GraphNodeK2Create", "CreateMatchingFunctionOption", "[Create a matching function]"));
+
+			// Only signatures with no output parameters can be events
+			if (!UEdGraphSchema_K2::HasFunctionAnyOutputParameter(FunctionSignature))
+			{
+				CreateMatchingEventData = AddDefaultFunctionDataOption(NSLOCTEXT("GraphNodeK2Create", "CreateMatchingEventOption", "[Create a matching event]"));
+			}
+
 			for (TFieldIterator<UFunction> It(ScopeClass); It; ++It)
 			{
 				UFunction* Func = *It;
@@ -217,14 +286,6 @@ void SGraphNodeK2CreateDelegate::CreateBelowPinControls(TSharedPtr<SVerticalBox>
 				}
 			}
 
-			if (FunctionDataItems.Num() == 0)
-			{
-				// add an empty row, so the user can clear the selection if they want
-				TSharedPtr<FFunctionItemData> EmptyEntry = MakeShareable(new FFunctionItemData());
-				EmptyEntry->Description = NSLOCTEXT("GraphNodeK2Create", "EmptyFunctionLabel", "[NONE]").ToString();
-				FunctionDataItems.Add(EmptyEntry);
-			}
-
 			TSharedRef<SComboButton> SelectFunctionWidgetRef = SNew(SComboButton)
 				.Method(EPopupMethod::CreateNewWindow)
 				.ButtonContent()
diff --git a/Engine/Source/Editor/GraphEditor/Private/KismetNodes/SGraphNodeK2CreateDelegate.h b/Engine/Source/Editor/GraphEditor/Private/KismetNodes/SGraphNodeK2CreateDelegate.h
--- a/Engine/Source/Editor/GraphEditor/Private/KismetNodes/SGraphNodeK2CreateDelegate.h
+++ b/Engine/Source/Editor/GraphEditor/Private/KismetNodes/SGraphNodeK2CreateDelegate.h
@@ -18,23 +18,43 @@ public:
 	SLATE_BEGIN_ARGS(SGraphNodeK2CreateDelegate) {}
 	SLATE_END_ARGS()
 
-		struct FFunctionItemData
+	/** Data that defines a delegate function */
+	struct FFunctionItemData
 	{
 		FName Name;
-		FString Description;
+		FText Description;
 	};
+
+	/** Collection of function items that have matching function delegates to this node */
 	TArray<TSharedPtr<FFunctionItemData>> FunctionDataItems;
 	TWeakPtr<SComboButton> SelectFunctionWidget;
+
+	/** Data that can be used to create a matching function based on the parameters of a create event node */
+	TSharedPtr<FFunctionItemData> CreateMatchingFunctionData;
+	
+	/** Data that can be used to create a matching event based on based on the parameters of a create event node */
+	TSharedPtr<FFunctionItemData> CreateMatchingEventData;
+
 public:
 	virtual ~SGraphNodeK2CreateDelegate();
 	void Construct(const FArguments& InArgs, UK2Node* InNode);
 	virtual void CreateBelowPinControls(TSharedPtr<SVerticalBox> MainBox) override;
 
 protected:
-	static FString FunctionDescription(const UFunction* Function, const bool bOnlyDescribeSignature = false, const int32 CharacterLimit = 32);
+	static FText FunctionDescription(const UFunction* Function, const bool bOnlyDescribeSignature = false, const int32 CharacterLimit = 32);
 
 	FText GetCurrentFunctionDescription() const;
 	TSharedRef<ITableRow> HandleGenerateRowFunction(TSharedPtr<FFunctionItemData> FunctionItemData, const TSharedRef<STableViewBase>& OwnerTable);
 	void OnFunctionSelected(TSharedPtr<FFunctionItemData> FunctionItemData, ESelectInfo::Type SelectInfo);
+
+private:
+
+	/**
+	* Adds a FunctionItemData with a given description to the array of FunctionDataItems. 
+	* 
+	* @param	DescriptionName		Description of the option to give the user
+	* @return	Shared pointer to the FunctionItemData
+	*/
+	TSharedPtr<SGraphNodeK2CreateDelegate::FFunctionItemData> AddDefaultFunctionDataOption(const FText& DescriptionName);
 };
 
