Improve the blueprint context menu searching heuristics

Add context menu CVars for tweaking the behavior of finding matches in the context menu
Add some comparisons to the type of pin that was dragged off so that there is a preference towards that type
Prefer shorter matches over long ones
Update how percentage matches are calculated

diff --git a/Engine/Source/Editor/GraphEditor/Private/SGraphActionMenu.cpp b/Engine/Source/Editor/GraphEditor/Private/SGraphActionMenu.cpp
--- a/Engine/Source/Editor/GraphEditor/Private/SGraphActionMenu.cpp
+++ b/Engine/Source/Editor/GraphEditor/Private/SGraphActionMenu.cpp
@@ -21,6 +21,88 @@
 
 #define LOCTEXT_NAMESPACE "GraphActionMenu"
 
+//////////////////////////////////////////////////////////////////////////
+/** CVars for tweaking how the blueprint context menu search picks the best match */
+namespace ContextMenuConsoleVariables
+{
+	/** Increasing this weight will make shorter words preferred */
+	static float ShorterWeight = 2.0f;
+	static FAutoConsoleVariableRef CVarShorterWeight(
+		TEXT("ContextMenu.ShorterWeight"), ShorterWeight,
+		TEXT("Increasing this weight will make shorter words preferred"),
+		ECVF_Default);
+
+	static int32 KeywordLengthDifferenceModifier = 50;
+	static FAutoConsoleVariableRef CVarKeywordLengthDifferenceModifier(
+		TEXT("ContextMenu.KeywordLengthDifferenceModifier"), KeywordLengthDifferenceModifier,
+		TEXT("Used to calculate how much weight should be given to the current keyword (KeywordLengthDifferenceModifier - KeywordLength) * ShorterWeight"),
+		ECVF_Default);
+
+	/** Increasing this will prefer whole percentage matches when comparing the keyword to what the user has typed in */
+	static float PercentageMatchWeightMultiplier = 1.0f;
+	static FAutoConsoleVariableRef CVarPercentageMatchWeightMultiplier(
+		TEXT("ContextMenu.PercentageMatchWeightMultiplier"), PercentageMatchWeightMultiplier,
+		TEXT("A multiplier for how much weight to give something based on the percentage match it is"),
+		ECVF_Default);
+
+	/** How much weight the description of actions have */
+	static float DescriptionWeight = 10.0f;
+	static FAutoConsoleVariableRef CVarDescriptionWeight(
+		TEXT("ContextMenu.DescriptionWeight"), DescriptionWeight,
+		TEXT("The amount of weight placed on search items description"),
+		ECVF_Default);
+
+	/** Weight used to prefer categories that are the same as the node that was dragged off of */
+	static float MatchingFromPinCategory = 500.0f;
+	static FAutoConsoleVariableRef CVarMatchingFromPinCategory(
+		TEXT("ContextMenu.MatchingFromPinCategory"), MatchingFromPinCategory,
+		TEXT("The amount of weight placed on actions with the same category as the node being dragged off of"),
+		ECVF_Default);
+
+	/** Weight that a match to a category search has */
+	static float CategoryWeight = 5.0f;
+	static FAutoConsoleVariableRef CVarCategoryWeight(
+		TEXT("ContextMenu.CategoryWeight"), CategoryWeight,
+		TEXT("The amount of weight placed on categories that match what the user has typed in"),
+		ECVF_Default);
+
+	/** How much weight the node's title has */
+	static float NodeTitleWeight = 50.0f;
+	static FAutoConsoleVariableRef CVarNodeTitleWeight(
+		TEXT("ContextMenu.NodeTitleWeight"), NodeTitleWeight,
+		TEXT("The amount of weight placed on the search items title"),
+		ECVF_Default);
+
+	/** Weight used to prefer keywords of actions  */
+	static float KeywordWeight = 100.0f;
+	static FAutoConsoleVariableRef CVarKeywordWeight(
+		TEXT("ContextMenu.KeywordWeight"), KeywordWeight,
+		TEXT("The amount of weight placed on search items keyword"),
+		ECVF_Default);
+
+	/** The multiplier given if the keyword starts with a letter the user typed in */
+	static float StartsWithBonusWeightMultiplier = 5.0f;
+	static FAutoConsoleVariableRef CVarStartsWithBonusWeightMultiplier(
+		TEXT("ContextMenu.StartsWithBonusWeightMultiplier"), StartsWithBonusWeightMultiplier,
+		TEXT("The multiplier given if the keyword starts with a letter the user typed in"),
+		ECVF_Default);
+
+	/** The multiplier given if the keyword starts with a letter the user typed in */
+	static float WordContainsLetterWeightMultiplier = 0.5f;
+	static FAutoConsoleVariableRef CVarWordContainsLetterWeightMultiplier(
+		TEXT("ContextMenu.WordContainsLetterWeightMultiplier"), WordContainsLetterWeightMultiplier,
+		TEXT("The multiplier given if the keyword only contains a letter the user typed in"),
+		ECVF_Default);
+
+	/** Enabling the debug printing of context menu selections */
+	static bool bPrintDebugContextSelection = false;
+	static FAutoConsoleVariableRef CVarPrintDebugContextSelection(
+		TEXT("ContextMenu.bPrintDebugContextSelection"), bPrintDebugContextSelection,
+		TEXT("Flag for printing the debug info about the context menu selection"),
+		ECVF_Default);
+
+};	// namespace ContextMenuConsoleVariables
+
 //////////////////////////////////////////////////////////////////////////
 
 template<typename ItemType>
@@ -281,8 +363,9 @@ void SGraphActionMenu::Construct( const FArguments& InArgs, bool bIsReadOnly/* =
 	this->OnGetSectionToolTip = InArgs._OnGetSectionToolTip;
 	this->OnGetSectionWidget = InArgs._OnGetSectionWidget;
 	this->FilteredRootAction = FGraphActionNode::NewRootNode();
-	this->OnActionMatchesName = InArgs._OnActionMatchesName;
-	
+	this->OnActionMatchesName = InArgs._OnActionMatchesName;	
+	this->DraggedFromPins = InArgs._DraggedFromPins;
+
 	// If a delegate for filtering text is passed in, assign it so that it will be used instead of the built-in filter box
 	if(InArgs._OnGetFilterText.IsBound())
 	{
@@ -667,7 +750,7 @@ void SGraphActionMenu::GenerateFilteredItems(bool bPreserveExpansion)
 	// Tokenize the search box text into a set of terms; all of them must be present to pass the filter
 	TArray<FString> FilterTerms;
 	TrimmedFilterString.ParseIntoArray(FilterTerms, TEXT(" "), true);
-	for (auto& String : FilterTerms)
+	for (FString& String : FilterTerms)
 	{
 		String = String.ToLower();
 	}
@@ -683,15 +766,20 @@ void SGraphActionMenu::GenerateFilteredItems(bool bPreserveExpansion)
 	ensure( SanitizedFilterTerms.Num() == FilterTerms.Num() );// Both of these should match !
 
 	const bool bRequiresFiltering = FilterTerms.Num() > 0;
-	int32 BestMatchCount = 0;
-	int32 BestMatchIndex = INDEX_NONE;		
+	float BestMatchCount = 0.0f;
+	int32 BestMatchIndex = INDEX_NONE;
+
+	FContextMenuWeightDebugInfo BestMatchDebugInfo = {};
+	FContextMenuWeightDebugInfo CurActionDebugInfo = {};
+
 	for (int32 CurTypeIndex=0; CurTypeIndex < AllActions.GetNumActions(); ++CurTypeIndex)
 	{
 		FGraphActionListBuilderBase::ActionGroup& CurrentAction = AllActions.GetAction( CurTypeIndex );
 
 		// If we're filtering, search check to see if we need to show this action
 		bool bShowAction = true;
-		int32 EachWeight = 0;
+		float EachWeight = 0.0f;
+		CurActionDebugInfo = {};
 		if (bRequiresFiltering)
 		{
 			// Combine the actions string, separate with \n so terms don't run into each other, and remove the spaces (incase the user is searching for a variable)
@@ -710,7 +798,7 @@ void SGraphActionMenu::GenerateFilteredItems(bool bPreserveExpansion)
 			if (bShowAction)
 			{
 				// Get the 'weight' of this in relation to the filter
-				EachWeight = GetActionFilteredWeight( CurrentAction, FilterTerms, SanitizedFilterTerms );
+				EachWeight = GetActionFilteredWeight(CurrentAction, FilterTerms, SanitizedFilterTerms, CurActionDebugInfo);
 			}
 		}
 
@@ -721,7 +809,8 @@ void SGraphActionMenu::GenerateFilteredItems(bool bPreserveExpansion)
 			{
 				BestMatchCount = EachWeight;
 				BestMatchIndex = CurTypeIndex;
-			}								
+				BestMatchDebugInfo = CurActionDebugInfo;
+			}
 			FilteredRootAction->AddChild(CurrentAction);
 		}
 	}
@@ -737,6 +826,12 @@ void SGraphActionMenu::GenerateFilteredItems(bool bPreserveExpansion)
 	TArray< TSharedPtr<FGraphActionNode> > AllNodes;
 	FilteredRootAction->GetAllNodes(AllNodes);
 
+	// Print out the info about which action we picked and why
+	if (ContextMenuConsoleVariables::bPrintDebugContextSelection)
+	{
+		BestMatchDebugInfo.Print();
+	}
+
 	// If theres a BestMatchIndex find it in the actions nodes and select it (maybe this should check the current selected suggestion first ?)
 	if( BestMatchIndex != INDEX_NONE ) 
 	{
@@ -778,30 +873,30 @@ void SGraphActionMenu::GenerateFilteredItems(bool bPreserveExpansion)
 	}
 }
 
-int32 SGraphActionMenu::GetActionFilteredWeight( const FGraphActionListBuilderBase::ActionGroup& InCurrentAction, const TArray<FString>& InFilterTerms, const TArray<FString>& InSanitizedFilterTerms )
+void FContextMenuWeightDebugInfo::Print()
 {
-	// The overall 'weight'
-	int32 TotalWeight = 0;
-
-	// Some simple weight figures to help find the most appropriate match	
-	const int32 WholeMatchWeightMultiplier = 2;
-	const int32 WholeMatchLocalizedWeightMultiplier = 3;
-	const int32 DescriptionWeight = 10;
-	const int32 CategoryWeight = 1;
-	const int32 NodeTitleWeight = 1;
-	const int32 KeywordWeight = 4;
+	UE_LOG(LogTemp, Warning, TEXT("[Weight Debug info] \
+		TotalWeight: %-8.2f | PercentageMatchWeight: %-8.2f | PercMatch: %-8.2f | ShorterWeight: %-8.2f | CategoryBonusWeight: %-8.2f | KeywordArrayWeight: %-8.2f | DescriptionWeight: %-8.2f | NodeTitleWeight: %-8.2f | CategoryWeight: %-8.2f\n"), 
+		TotalWeight, PercentageMatchWeight, PercMatch, ShorterWeight, CategoryBonusWieight, KeywordArrayWeight, DescriptionWeight, NodeTitleWeight, CategoryWeight);
+}
 
+float SGraphActionMenu::GetActionFilteredWeight(const FGraphActionListBuilderBase::ActionGroup& InCurrentAction, const TArray<FString>& InFilterTerms, const TArray<FString>& InSanitizedFilterTerms, FContextMenuWeightDebugInfo& OutDebugInfo)
+{	
+	// The overall 'weight' of this action 
+	float TotalWeight = 0.0f;
 	// Helper array
 	struct FArrayWithWeight
 	{
-		FArrayWithWeight(const TArray< FString >* InArray, int32 InWeight)
+		FArrayWithWeight(const TArray< FString >* InArray, float InWeightModifier, float* OutDebugWeight)
 			: Array(InArray)
-			, Weight(InWeight)
+			, OutWeight(OutDebugWeight)
+			, WeightModifier(InWeightModifier)
 		{
 		}
 
-		const TArray< FString >* Array;
-		int32			  Weight;
+		const TArray< FString >* Array = nullptr;
+		float* OutWeight = nullptr;
+		float WeightModifier = 0.0f;
 	};
 
 	// Setup an array of arrays so we can do a weighted search			
@@ -816,99 +911,149 @@ int32 SGraphActionMenu::GetActionFilteredWeight( const FGraphActionListBuilderBa
 		const FString& SearchText = InCurrentAction.GetSearchTextForFirstAction();
 
 		// First the localized keywords
-		WeightedArrayList.Add(FArrayWithWeight(&InCurrentAction.GetLocalizedSearchKeywordsArrayForFirstAction(), KeywordWeight));
+		WeightedArrayList.Add(FArrayWithWeight(&InCurrentAction.GetLocalizedSearchKeywordsArrayForFirstAction(), ContextMenuConsoleVariables::KeywordWeight, &OutDebugInfo.KeywordArrayWeight));
 
 		// The localized description
-		WeightedArrayList.Add(FArrayWithWeight(&InCurrentAction.GetLocalizedMenuDescriptionArrayForFirstAction(), DescriptionWeight));
+		WeightedArrayList.Add(FArrayWithWeight(&InCurrentAction.GetLocalizedMenuDescriptionArrayForFirstAction(), ContextMenuConsoleVariables::DescriptionWeight, &OutDebugInfo.DescriptionWeight));
 
 		// The node search localized title weight
-		WeightedArrayList.Add(FArrayWithWeight(&InCurrentAction.GetLocalizedSearchTitleArrayForFirstAction(), NodeTitleWeight));
+		WeightedArrayList.Add(FArrayWithWeight(&InCurrentAction.GetLocalizedSearchTitleArrayForFirstAction(), ContextMenuConsoleVariables::NodeTitleWeight, &OutDebugInfo.NodeTitleWeight));
 
 		// The localized category
-		WeightedArrayList.Add(FArrayWithWeight(&InCurrentAction.GetLocalizedSearchCategoryArrayForFirstAction(), CategoryWeight));
+		WeightedArrayList.Add(FArrayWithWeight(&InCurrentAction.GetLocalizedSearchCategoryArrayForFirstAction(), ContextMenuConsoleVariables::CategoryWeight,&OutDebugInfo.CategoryWeight));
 
 		// First the keywords
-		int32 NonLocalizedFirstIndex = WeightedArrayList.Add(FArrayWithWeight(&InCurrentAction.GetSearchKeywordsArrayForFirstAction(), KeywordWeight));
-
+		WeightedArrayList.Add(FArrayWithWeight(&InCurrentAction.GetSearchKeywordsArrayForFirstAction(), ContextMenuConsoleVariables::KeywordWeight, &OutDebugInfo.KeywordArrayWeight));
+		
 		// The description
-		WeightedArrayList.Add(FArrayWithWeight(&InCurrentAction.GetMenuDescriptionArrayForFirstAction(), DescriptionWeight));
+		WeightedArrayList.Add(FArrayWithWeight(&InCurrentAction.GetMenuDescriptionArrayForFirstAction(), ContextMenuConsoleVariables::DescriptionWeight,  &OutDebugInfo.DescriptionWeight));
 
 		// The node search title weight
-		WeightedArrayList.Add(FArrayWithWeight(&InCurrentAction.GetSearchTitleArrayForFirstAction(), NodeTitleWeight));
+		WeightedArrayList.Add(FArrayWithWeight(&InCurrentAction.GetSearchTitleArrayForFirstAction(), ContextMenuConsoleVariables::NodeTitleWeight, &OutDebugInfo.NodeTitleWeight));
 
 		// The category
-		WeightedArrayList.Add(FArrayWithWeight(&InCurrentAction.GetSearchCategoryArrayForFirstAction(), CategoryWeight));
+		WeightedArrayList.Add(FArrayWithWeight(&InCurrentAction.GetSearchCategoryArrayForFirstAction(), ContextMenuConsoleVariables::CategoryWeight,  &OutDebugInfo.CategoryWeight));
+
+		// If this actions category matches the one from a dragged off pin, give it a weight bonus
+		const TArray<FString>& InActionCategories = InCurrentAction.GetCategoryChain();
+		bool bAddMatchBonus = false;
+
+		for (const FString& InActionCategory : InActionCategories)
+		{
+			for (UEdGraphPin* const FromPin : DraggedFromPins)
+			{
+				check(FromPin != nullptr);
+			
+				// If we can't find anything there, check the subcategory of the object
+				// This covers most of the more complex struct types (LinearColor, date time, etc)
+				if (UObject* const SubCatObj = FromPin->PinType.PinSubCategoryObject.Get())
+				{
+					const FString& SubCatObjName = SubCatObj->GetFullName();
+					// The pin SubObjectCategory names don't have any spaces, so split up the category
+					TArray<FString> DelimitedArray;
+					InActionCategory.ParseIntoArray(DelimitedArray, TEXT(" "), true);
+					for (const FString& DelimetedCat : DelimitedArray)
+					{
+						if (SubCatObjName.Contains(DelimetedCat))
+						{
+							bAddMatchBonus = true;
+							break;
+						}
+					}
+				}
+				// Check the category of the pin, this works for basic math types (int, float, byte, etc)
+				else if (InActionCategory.Contains(FromPin->PinType.PinCategory.ToString()))
+				{
+					bAddMatchBonus = true;
+				}
+
+				// If we found match in any cases above then add the weight bonus and stop looking
+				if (bAddMatchBonus)
+				{
+					TotalWeight += ContextMenuConsoleVariables::MatchingFromPinCategory;				
+					OutDebugInfo.CategoryBonusWieight += ContextMenuConsoleVariables::MatchingFromPinCategory;
+					
+					// Break out of the loop so that we don't give any extra bonuses
+					break;
+				}
+			}
+		}
 
 		// Now iterate through all the filter terms and calculate a 'weight' using the values and multipliers
 		const FString* EachTerm = nullptr;
 		const FString* EachTermSanitized = nullptr;
+
+		// For every filter item the user has typed in (a letter)
 		for (int32 FilterIndex = 0; FilterIndex < InFilterTerms.Num(); ++FilterIndex)
 		{
 			EachTerm = &InFilterTerms[FilterIndex];
 			EachTermSanitized = &InSanitizedFilterTerms[FilterIndex];
-			if( SearchText.Contains( *EachTerm, ESearchCase::CaseSensitive ) )
-			{
-				TotalWeight += 2;
-			}
-			else if (SearchText.Contains(*EachTermSanitized, ESearchCase::CaseSensitive))
-			{
-				TotalWeight++;
-			}		
+			
 			// Now check the weighted lists	(We could further improve the hit weight by checking consecutive word matches)
-			for (int32 iFindCount = 0; iFindCount < WeightedArrayList.Num() ; iFindCount++)
+			for (int32 iFindCount = 0; iFindCount < WeightedArrayList.Num() ; ++iFindCount)
 			{
-				int32 WeightPerList = 0;
 				const TArray<FString>& KeywordArray = *WeightedArrayList[iFindCount].Array;
-				int32 EachWeight = WeightedArrayList[iFindCount].Weight;
-				int32 WholeMatchCount = 0;
-				int32 WholeMatchMultiplier = (iFindCount < NonLocalizedFirstIndex) ? WholeMatchLocalizedWeightMultiplier : WholeMatchWeightMultiplier;
+				float WeightPerList = 0.0f;
+				float KeywordArrayWeight = WeightedArrayList[iFindCount].WeightModifier;
+
+				// Count of how many words in this keyword array contain a filter(letter) that the user has typed in
+				int32 WordMatchCount = 0;
 
-				for (int32 iEachWord = 0; iEachWord < KeywordArray.Num() ; iEachWord++)
+				// The number of characters in this keyword array
+				int32 KeywordArrayCharLength = 0;
+
+				// Loop through every word that the user could be looking for
+				for (int32 iEachWord = 0; iEachWord < KeywordArray.Num() ; ++iEachWord)
 				{
-					// If we get an exact match weight the find count to get exact matches higher priority
-					if (KeywordArray[iEachWord].StartsWith(*EachTerm, ESearchCase::CaseSensitive))
-					{
-						if (iEachWord == 0)
-						{
-							WeightPerList += EachWeight * WholeMatchMultiplier;
-						}
-						else
-						{
-							WeightPerList += EachWeight;
-						}
-						WholeMatchCount++;
-					}
-					else if (KeywordArray[iEachWord].Contains(*EachTerm, ESearchCase::CaseSensitive))
+					// Keep track of how long all the words in the array are
+					KeywordArrayCharLength += KeywordArray[iEachWord].Len();
+
+					// If a word contains the letter that the user has typed in, than increment the whole match count
+					// If the word starts with the letter, give it a little extra boost of weight
+					if (KeywordArray[iEachWord].StartsWith(*EachTermSanitized, ESearchCase::CaseSensitive) || KeywordArray[iEachWord].StartsWith(*EachTerm, ESearchCase::CaseSensitive))
 					{
-						WeightPerList += EachWeight;
+						++WordMatchCount;
+						WeightPerList += KeywordArrayWeight * ContextMenuConsoleVariables::StartsWithBonusWeightMultiplier;
 					}
-					if (KeywordArray[iEachWord].StartsWith(*EachTermSanitized, ESearchCase::CaseSensitive))
+					else if (KeywordArray[iEachWord].Contains(*EachTermSanitized, ESearchCase::CaseSensitive) || KeywordArray[iEachWord].Contains(*EachTerm, ESearchCase::CaseSensitive))
 					{
-						if (iEachWord == 0)
-						{
-							WeightPerList += EachWeight * WholeMatchMultiplier;
-						}
-						else
-						{
-							WeightPerList += EachWeight;
-						}
-						WholeMatchCount++;
+						++WordMatchCount;
+						WeightPerList += KeywordArrayWeight * ContextMenuConsoleVariables::WordContainsLetterWeightMultiplier;
 					}
-					else if (KeywordArray[iEachWord].Contains(*EachTermSanitized, ESearchCase::CaseSensitive))
+				}
+
+				if (KeywordArrayCharLength > 0)
+				{
+					// How many matches did we find / the total length of this keyword array
+					float PercMatch = ((float)WordMatchCount / (float)KeywordArrayCharLength);
+					float PercentageBonus = (WeightPerList * PercMatch * ContextMenuConsoleVariables::PercentageMatchWeightMultiplier);
+					WeightPerList += PercentageBonus;
+
+					// Give a bonus for being shorter than a certain amount
+					int32 ShortPoints = ContextMenuConsoleVariables::KeywordLengthDifferenceModifier - KeywordArrayCharLength;
+					float ShortWeight = 0.0f;
+					if (ShortPoints > 0)
 					{
-						WeightPerList += EachWeight / 2;
+						ShortWeight = (float)ShortPoints * ContextMenuConsoleVariables::ShorterWeight;
+						WeightPerList += ShortWeight;
 					}
+
+					OutDebugInfo.PercMatch += PercMatch;
+					OutDebugInfo.ShorterWeight += ShortWeight;
+					OutDebugInfo.PercentageMatchWeight += PercentageBonus;					
 				}
-				// Increase the weight if theres a larger % of matches in the keyword list
-				if( WholeMatchCount != 0 )
+
+				TotalWeight += WeightPerList;
+				if (WeightedArrayList[iFindCount].OutWeight)
 				{
-					int32 PercentAdjust = ( 100 / KeywordArray.Num() ) * WholeMatchCount;
-					WeightPerList *= PercentAdjust;
+					*WeightedArrayList[iFindCount].OutWeight = WeightPerList;
 				}
-				TotalWeight += WeightPerList;
 			}
 		}
 	}
+
+	OutDebugInfo.TotalWeight = TotalWeight;
+	
 	return TotalWeight;
 }
 
@@ -1183,7 +1328,7 @@ FText SGraphActionMenu::GetFilterText() const
 		return OnGetFilterText.Execute();
 	}
 	
-	return FilterTextBox->GetText();;
+	return FilterTextBox->GetText();
 }
 
 void SGraphActionMenu::OnItemSelected( TSharedPtr< FGraphActionNode > InSelectedItem, ESelectInfo::Type SelectInfo )
diff --git a/Engine/Source/Editor/GraphEditor/Public/SGraphActionMenu.h b/Engine/Source/Editor/GraphEditor/Public/SGraphActionMenu.h
--- a/Engine/Source/Editor/GraphEditor/Public/SGraphActionMenu.h
+++ b/Engine/Source/Editor/GraphEditor/Public/SGraphActionMenu.h
@@ -79,6 +79,31 @@ struct FCustomExpanderData
 	TSharedPtr<SPanel> WidgetContainer;
 };
 
+/** 
+* Debug Info about how the preferred context menu action is chosen
+* @see SGraphActionMenu::GetActionFilteredWeight
+*/
+struct FContextMenuWeightDebugInfo
+{
+	float TotalWeight = 0.0f;
+	// Info about which sets of words match up with what weight
+	float KeywordArrayWeight = 0.0f;
+	float DescriptionWeight = 0.0f;
+	float NodeTitleWeight = 0.0f;
+	float CategoryWeight = 0.0f;
+
+	float PercentageMatchWeight = 0.0f;
+	/** The calculated match percentage */
+	float PercMatch = 0.0f;
+	float ShorterWeight = 0.0f;
+	float CategoryBonusWieight = 0.0f;
+
+	/**
+	* Print out the debug info about this weight info to the console 
+	*/
+	void Print();
+};
+
 /** Class that displays a list of graph actions and them to be searched and selected */
 class GRAPHEDITOR_API SGraphActionMenu : public SCompoundWidget, public FGCObject
 {
@@ -141,6 +166,8 @@ public:
 		SLATE_ARGUMENT( bool, AlphaSortItems )
 		SLATE_ARGUMENT( bool, ShowFilterTextBox )
 		SLATE_ARGUMENT( bool, UseSectionStyling )
+		SLATE_ARGUMENT( TArray<UEdGraphPin*>, DraggedFromPins )
+
 	SLATE_END_ARGS()
 
 	void Construct( const FArguments& InArgs, bool bIsReadOnly = true );
@@ -189,7 +216,7 @@ protected:
 	bool bAlphaSortItems;
 	/** Should the rows and sections be styled like the details panel? */
 	bool bUseSectionStyling;
-
+	
 	/** Delegate to call when action is selected */
 	FOnActionSelected OnActionSelected;
 	/** Delegate to call when action is double clicked */
@@ -283,8 +310,17 @@ protected:
 	/** Checks if the passed in node is safe for renaming */
 	bool CanRenameNode(TWeakPtr<FGraphActionNode> InNode) const;
 
-	/** Get the 'weight' of the given action relevance to a given filter list */
-	int32 GetActionFilteredWeight( const FGraphActionListBuilderBase::ActionGroup& InCurrentAction, const TArray<FString>& InFilterTerms, const TArray<FString>& InSanitizedFilterTerms );
+	/** 
+	* Get the 'weight' of the given action relevance to a given filter list
+	* 
+	* @param InCurrentAction			The action to check how good of a match it is
+	* @param InFilterTerms				The filter terms that the user has typed in to go off of
+	* @param InSanitizedFilterTerms		The sanitized filter terms to go off of
+	* @param OutDebugInfo				Debug info about the weight so that the user can tweak how things are matched
+	* 
+	* @return the weight of this action group
+	*/
+	float GetActionFilteredWeight( const FGraphActionListBuilderBase::ActionGroup& InCurrentAction, const TArray<FString>& InFilterTerms, const TArray<FString>& InSanitizedFilterTerms, FContextMenuWeightDebugInfo& OutDebugInfo );
 
 	// Delegates
 
@@ -324,5 +360,9 @@ protected:
 	void OnItemScrolledIntoView( TSharedPtr<FGraphActionNode> InActionNode, const TSharedPtr<ITableRow>& InWidget );
 	/** Callback for expanding tree items recursively */
 	void OnSetExpansionRecursive(TSharedPtr<FGraphActionNode> InTreeNode, bool bInIsItemExpanded);
+
+private:
+	/** The pins that have been dragged off of to prompt the creation of this action menu. */
+	TArray<UEdGraphPin*> DraggedFromPins;
 };
 
diff --git a/Engine/Source/Editor/Kismet/Private/SBlueprintActionMenu.cpp b/Engine/Source/Editor/Kismet/Private/SBlueprintActionMenu.cpp
--- a/Engine/Source/Editor/Kismet/Private/SBlueprintActionMenu.cpp
+++ b/Engine/Source/Editor/Kismet/Private/SBlueprintActionMenu.cpp
@@ -344,6 +344,7 @@ void SBlueprintActionMenu::Construct( const FArguments& InArgs, TSharedPtr<FBlue
 						.OnCreateWidgetForAction(SGraphActionMenu::FOnCreateWidgetForAction::CreateSP(this, &SBlueprintActionMenu::OnCreateWidgetForAction))
 						.OnCollectAllActions(this, &SBlueprintActionMenu::CollectAllActions)
 						.OnCreateCustomRowExpander_Static(&CreateCustomBlueprintActionExpander)
+						.DraggedFromPins(DraggedFromPins)
 				]
 			]
 		]
