Make context menu debug info only print out when the user types, not whenever the context menu is built.

#jria UE-77737

diff --git a/Engine/Source/Editor/GraphEditor/Private/SGraphActionMenu.cpp b/Engine/Source/Editor/GraphEditor/Private/SGraphActionMenu.cpp
--- a/Engine/Source/Editor/GraphEditor/Private/SGraphActionMenu.cpp
+++ b/Engine/Source/Editor/GraphEditor/Private/SGraphActionMenu.cpp
@@ -827,7 +827,7 @@ void SGraphActionMenu::GenerateFilteredItems(bool bPreserveExpansion)
 	FilteredRootAction->GetAllNodes(AllNodes);
 
 	// Print out the info about which action we picked and why
-	if (ContextMenuConsoleVariables::bPrintDebugContextSelection)
+	if (ContextMenuConsoleVariables::bPrintDebugContextSelection && FilterTerms.Num() > 0)
 	{
 		BestMatchDebugInfo.Print();
 	}
