Add undo functionality for creation of signitures based on an event dispatcher

diff --git a/Engine/Source/Editor/GraphEditor/Private/KismetNodes/SGraphNodeK2CreateDelegate.cpp b/Engine/Source/Editor/GraphEditor/Private/KismetNodes/SGraphNodeK2CreateDelegate.cpp
--- a/Engine/Source/Editor/GraphEditor/Private/KismetNodes/SGraphNodeK2CreateDelegate.cpp
+++ b/Engine/Source/Editor/GraphEditor/Private/KismetNodes/SGraphNodeK2CreateDelegate.cpp
@@ -157,13 +157,18 @@ TSharedRef<ITableRow> SGraphNodeK2CreateDelegate::HandleGenerateRowFunction(TSha
 
 void SGraphNodeK2CreateDelegate::OnFunctionSelected(TSharedPtr<FFunctionItemData> FunctionItemData, ESelectInfo::Type SelectInfo)
 {
+	const FScopedTransaction Transaction(NSLOCTEXT("GraphNodeK2Create", "CreateMatchingSigniture", "Create matching signiture"));
+
 	if (FunctionItemData.IsValid())
 	{
 		if (UK2Node_CreateDelegate* Node = Cast<UK2Node_CreateDelegate>(GraphNode))
 		{
 			UBlueprint* NodeBP = Node->GetBlueprint();
-			
-			check(NodeBP);
+			UEdGraph* const SourceGraph = Node->GetGraph();
+			check(NodeBP && SourceGraph);
+			SourceGraph->Modify();
+			NodeBP->Modify();
+			Node->Modify();
 
 			if (FunctionItemData == CreateMatchingFunctionData)
 			{
@@ -172,7 +177,7 @@ void SGraphNodeK2CreateDelegate::OnFunctionSelected(TSharedPtr<FFunctionItemData
 				FName NewFuncName = FBlueprintEditorUtils::GenerateUniqueGraphName(NodeBP, ProposedFuncName);
 				
 				UEdGraph* NewGraph = nullptr;
-				NewGraph = FBlueprintEditorUtils::CreateNewGraph(NodeBP, NewFuncName, UEdGraph::StaticClass(), UEdGraphSchema_K2::StaticClass());
+				NewGraph = FBlueprintEditorUtils::CreateNewGraph(NodeBP, NewFuncName, SourceGraph->GetClass(), SourceGraph->GetSchema() ? SourceGraph->GetSchema()->GetClass() : GetDefault<UEdGraphSchema_K2>()->GetClass());
 
 				if (NewGraph != nullptr)
 				{
