Fix missing include
#rnx

diff --git a/Engine/Source/Editor/GraphEditor/Private/KismetNodes/SGraphNodeK2CreateDelegate.cpp b/Engine/Source/Editor/GraphEditor/Private/KismetNodes/SGraphNodeK2CreateDelegate.cpp
--- a/Engine/Source/Editor/GraphEditor/Private/KismetNodes/SGraphNodeK2CreateDelegate.cpp
+++ b/Engine/Source/Editor/GraphEditor/Private/KismetNodes/SGraphNodeK2CreateDelegate.cpp
@@ -11,6 +11,7 @@
 
 #include "Editor/UnrealEd/Public/Kismet2/BlueprintEditorUtils.h"
 #include "Editor/UnrealEd/Public/Kismet2/KismetEditorUtilities.h"
+#include "Editor/UnrealEd/Public/ScopedTransaction.h"
 #include "Editor/BlueprintGraph/Classes/BlueprintNodeBinder.h"
 #include "Editor/BlueprintGraph/Classes/BlueprintEventNodeSpawner.h"
 #include "Editor/BlueprintGraph/Classes/K2Node_CustomEvent.h"
