PR #4880: Fix use of shift key for manipulating graph pins (Contributed by [REDACTED_PERSON])

#rnx

diff --git a/Engine/Source/Editor/GraphEditor/Private/SGraphPin.cpp b/Engine/Source/Editor/GraphEditor/Private/SGraphPin.cpp
--- a/Engine/Source/Editor/GraphEditor/Private/SGraphPin.cpp
+++ b/Engine/Source/Editor/GraphEditor/Private/SGraphPin.cpp
@@ -438,7 +438,7 @@ FReply SGraphPin::OnPinMouseDown( const FGeometry& SenderGeometry, const FPointe
 				}
 			}
 
-			if (!GraphPinObj->bNotConnectable)
+			if (!MouseEvent.IsShiftDown() && !GraphPinObj->bNotConnectable)
 			{
 				// Start a drag-drop on the pin
 				if (ensure(OwnerNodePinned.IsValid()))
