Fix SGraphPin merge from 5112201

diff --git a/Engine/Source/Editor/GraphEditor/Private/SGraphPin.cpp b/Engine/Source/Editor/GraphEditor/Private/SGraphPin.cpp
--- a/Engine/Source/Editor/GraphEditor/Private/SGraphPin.cpp
+++ b/Engine/Source/Editor/GraphEditor/Private/SGraphPin.cpp
@@ -69,6 +69,7 @@ SGraphPin::SGraphPin()
 	: GraphPinObj(nullptr)
 	, PinColorModifier(FLinearColor::White)
 	, CachedNodeOffset(FVector2D::ZeroVector)
+	, bGraphDataInvalid(false)
 	, bShowLabel(true)
 	, bOnlyShowDefaultValue(false)
 	, bIsMovingLinks(false)
@@ -1002,7 +1003,7 @@ FSlateColor SGraphPin::GetPinTextColor() const
 			FLinearColor PinColor = FLinearColor::Red;
 			if (bDisabled)
 			{
-				PinColor.A = .25;
+				PinColor.A = .25f;
 			}
 			return PinColor;
 		}
