Change the texture preview from fill to either align left or right (depending on what side of the node it's added to)

diff --git a/Engine/Source/Editor/GraphEditor/Private/MaterialNodes/SGraphNodeMaterialBase.cpp b/Engine/Source/Editor/GraphEditor/Private/MaterialNodes/SGraphNodeMaterialBase.cpp
--- a/Engine/Source/Editor/GraphEditor/Private/MaterialNodes/SGraphNodeMaterialBase.cpp
+++ b/Engine/Source/Editor/GraphEditor/Private/MaterialNodes/SGraphNodeMaterialBase.cpp
@@ -357,6 +357,7 @@ void SGraphNodeMaterialBase::CreateBelowPinControls(TSharedPtr<SVerticalBox> Mai
 			LeftNodeBox->AddSlot()
 			.Padding(FMargin(NegativeHPad + ExtraPad, 0.0f, 0.0f, 0.0f))
 			.AutoHeight()
+			.HAlign(HAlign_Left)
 			[
 				CreatePreviewWidget()
 			];
@@ -366,6 +367,7 @@ void SGraphNodeMaterialBase::CreateBelowPinControls(TSharedPtr<SVerticalBox> Mai
 			RightNodeBox->AddSlot()
 			.Padding(FMargin(NegativeHPad + ExtraPad, 0.0f, 0.0f, 0.0f))
 			.AutoHeight()
+			.HAlign(HAlign_Right)
 			[
 				CreatePreviewWidget()
 			];
