Fix for a timing-specific crash caused by a stale graph pin ptr.

diff --git a/Engine/Source/Editor/GraphEditor/Private/KismetPins/SGraphPinText.cpp b/Engine/Source/Editor/GraphEditor/Private/KismetPins/SGraphPinText.cpp
--- a/Engine/Source/Editor/GraphEditor/Private/KismetPins/SGraphPinText.cpp
+++ b/Engine/Source/Editor/GraphEditor/Private/KismetPins/SGraphPinText.cpp
@@ -12,7 +12,7 @@ namespace
 	{
 	public:
 		FEditableTextGraphPin(UEdGraphPin* InGraphPinObj)
-			: GraphPinObj(InGraphPinObj)
+			: GraphPinObjRef(InGraphPinObj)
 		{
 		}
 
@@ -28,14 +28,24 @@ namespace
 
 		virtual bool IsReadOnly() const override
 		{
-			return GraphPinObj->bDefaultValueIsReadOnly;
+			if (const UEdGraphPin* GraphPinObj = GraphPinObjRef.Get())
+			{
+				return GraphPinObj->bDefaultValueIsReadOnly;
+			}
+
+			return false;
 		}
 
 		virtual bool IsDefaultValue() const override
 		{
-			FString TextAsString;
-			FTextStringHelper::WriteToBuffer(TextAsString, GraphPinObj->DefaultTextValue);
-			return TextAsString.Equals(GraphPinObj->AutogeneratedDefaultValue, ESearchCase::CaseSensitive);
+			if (const UEdGraphPin* GraphPinObj = GraphPinObjRef.Get())
+			{
+				FString TextAsString;
+				FTextStringHelper::WriteToBuffer(TextAsString, GraphPinObj->DefaultTextValue);
+				return TextAsString.Equals(GraphPinObj->AutogeneratedDefaultValue, ESearchCase::CaseSensitive);
+			}
+
+			return false;
 		}
 
 		virtual FText GetToolTipText() const override
@@ -51,15 +61,23 @@ namespace
 		virtual FText GetText(const int32 InIndex) const override
 		{
 			check(InIndex == 0);
-			return GraphPinObj->DefaultTextValue;
+			if (const UEdGraphPin* GraphPinObj = GraphPinObjRef.Get())
+			{
+				return GraphPinObj->DefaultTextValue;
+			}
+
+			return FText::GetEmpty();
 		}
 
 		virtual void SetText(const int32 InIndex, const FText& InText) override
 		{
 			check(InIndex == 0);
-			const FScopedTransaction Transaction(NSLOCTEXT("GraphEditor", "ChangeTxtPinValue", "Change Text Pin Value"));
-			GraphPinObj->Modify();
-			GraphPinObj->GetSchema()->TrySetDefaultText(*GraphPinObj, InText);
+			if (UEdGraphPin* GraphPinObj = GraphPinObjRef.Get())
+			{
+				const FScopedTransaction Transaction(NSLOCTEXT("GraphEditor", "ChangeTxtPinValue", "Change Text Pin Value"));
+				GraphPinObj->Modify();
+				GraphPinObj->GetSchema()->TrySetDefaultText(*GraphPinObj, InText);
+			}
 		}
 
 		virtual bool IsValidText(const FText& InText, FText& OutErrorMsg) const override
@@ -71,7 +89,10 @@ namespace
 		virtual void GetStableTextId(const int32 InIndex, const ETextPropertyEditAction InEditAction, const FString& InTextSource, const FString& InProposedNamespace, const FString& InProposedKey, FString& OutStableNamespace, FString& OutStableKey) const override
 		{
 			check(InIndex == 0);
-			return StaticStableTextId(GraphPinObj->GetOwningNodeUnchecked(), InEditAction, InTextSource, InProposedNamespace, InProposedKey, OutStableNamespace, OutStableKey);
+			if (const UEdGraphPin* GraphPinObj = GraphPinObjRef.Get())
+			{
+				StaticStableTextId(GraphPinObj->GetOwningNodeUnchecked(), InEditAction, InTextSource, InProposedNamespace, InProposedKey, OutStableNamespace, OutStableKey);
+			}
 		}
 #endif // USE_STABLE_LOCALIZATION_KEYS
 
@@ -80,7 +101,7 @@ namespace
 		}
 
 	private:
-		UEdGraphPin* GraphPinObj;
+		FEdGraphPinReference GraphPinObjRef;
 	};
 }
 
