Fix various crashes when disconnecting pins that trigger node changes while default pin value widgets are focused

diff --git a/Engine/Source/Editor/GraphEditor/Private/KismetPins/SGraphPinBool.cpp b/Engine/Source/Editor/GraphEditor/Private/KismetPins/SGraphPinBool.cpp
--- a/Engine/Source/Editor/GraphEditor/Private/KismetPins/SGraphPinBool.cpp
+++ b/Engine/Source/Editor/GraphEditor/Private/KismetPins/SGraphPinBool.cpp
@@ -27,6 +27,11 @@ ECheckBoxState SGraphPinBool::IsDefaultValueChecked() const
 
 void SGraphPinBool::OnDefaultValueCheckBoxChanged(ECheckBoxState InIsChecked)
 {
+	if(GraphPinObj->IsPendingKill())
+	{
+		return;
+	}
+
 	const FString BoolString = (InIsChecked == ECheckBoxState::Checked) ? TEXT("true") : TEXT("false");
 	if(GraphPinObj->GetDefaultAsString() != BoolString)
 	{
diff --git a/Engine/Source/Editor/GraphEditor/Private/KismetPins/SGraphPinClass.cpp b/Engine/Source/Editor/GraphEditor/Private/KismetPins/SGraphPinClass.cpp
--- a/Engine/Source/Editor/GraphEditor/Private/KismetPins/SGraphPinClass.cpp
+++ b/Engine/Source/Editor/GraphEditor/Private/KismetPins/SGraphPinClass.cpp
@@ -130,6 +130,11 @@ FOnClicked SGraphPinClass::GetOnUseButtonDelegate()
 
 void SGraphPinClass::OnPickedNewClass(UClass* ChosenClass)
 {
+	if(GraphPinObj->IsPendingKill())
+	{
+		return;
+	}
+
 	FString NewPath;
 	if (ChosenClass)
 	{
diff --git a/Engine/Source/Editor/GraphEditor/Private/KismetPins/SGraphPinColor.cpp b/Engine/Source/Editor/GraphEditor/Private/KismetPins/SGraphPinColor.cpp
--- a/Engine/Source/Editor/GraphEditor/Private/KismetPins/SGraphPinColor.cpp
+++ b/Engine/Source/Editor/GraphEditor/Private/KismetPins/SGraphPinColor.cpp
@@ -75,6 +75,11 @@ FLinearColor SGraphPinColor::GetColor() const
 
 void SGraphPinColor::OnColorCommitted(FLinearColor InColor)
 {
+	if(GraphPinObj->IsPendingKill())
+	{
+		return;
+	}
+
 	// Update pin object
 	FString ColorString = InColor.ToString();
 
diff --git a/Engine/Source/Editor/GraphEditor/Private/KismetPins/SGraphPinEnum.cpp b/Engine/Source/Editor/GraphEditor/Private/KismetPins/SGraphPinEnum.cpp
--- a/Engine/Source/Editor/GraphEditor/Private/KismetPins/SGraphPinEnum.cpp
+++ b/Engine/Source/Editor/GraphEditor/Private/KismetPins/SGraphPinEnum.cpp
@@ -107,6 +107,11 @@ TSharedRef<SWidget>	SGraphPinEnum::GetDefaultValueWidget()
 
 FText SGraphPinEnum::OnGetFriendlyName(int32 EnumIndex)
 {
+	if(GraphPinObj->IsPendingKill())
+	{
+		return FText();
+	}
+
 	UEnum* EnumPtr = Cast<UEnum>(GraphPinObj->PinType.PinSubCategoryObject.Get());
 
 	check(EnumPtr);
diff --git a/Engine/Source/Editor/GraphEditor/Private/KismetPins/SGraphPinKey.cpp b/Engine/Source/Editor/GraphEditor/Private/KismetPins/SGraphPinKey.cpp
--- a/Engine/Source/Editor/GraphEditor/Private/KismetPins/SGraphPinKey.cpp
+++ b/Engine/Source/Editor/GraphEditor/Private/KismetPins/SGraphPinKey.cpp
@@ -51,6 +51,11 @@ TOptional<FKey> SGraphPinKey::GetCurrentKey() const
 
 void SGraphPinKey::OnKeyChanged(TSharedPtr<FKey> InSelectedKey)
 {
+	if(GraphPinObj->IsPendingKill())
+	{
+		return;
+	}
+
 	if(SelectedKey != *InSelectedKey.Get())
 	{
 		const FScopedTransaction Transaction( NSLOCTEXT("GraphEditor", "ChangeKeyPinValue", "Change Key Pin Value" ) );
diff --git a/Engine/Source/Editor/GraphEditor/Private/KismetPins/SGraphPinObject.cpp b/Engine/Source/Editor/GraphEditor/Private/KismetPins/SGraphPinObject.cpp
--- a/Engine/Source/Editor/GraphEditor/Private/KismetPins/SGraphPinObject.cpp
+++ b/Engine/Source/Editor/GraphEditor/Private/KismetPins/SGraphPinObject.cpp
@@ -252,6 +252,11 @@ TSharedRef<SWidget> SGraphPinObject::GenerateAssetPicker()
 
 void SGraphPinObject::OnAssetSelectedFromPicker(const struct FAssetData& AssetData)
 {
+	if(GraphPinObj->IsPendingKill())
+	{
+		return;
+	}
+
 	const FAssetData& CurrentAssetData = GetAssetData(true);
 	if(CurrentAssetData != AssetData)
 	{
diff --git a/Engine/Source/Editor/GraphEditor/Private/KismetPins/SGraphPinString.cpp b/Engine/Source/Editor/GraphEditor/Private/KismetPins/SGraphPinString.cpp
--- a/Engine/Source/Editor/GraphEditor/Private/KismetPins/SGraphPinString.cpp
+++ b/Engine/Source/Editor/GraphEditor/Private/KismetPins/SGraphPinString.cpp
@@ -62,6 +62,11 @@ FText SGraphPinString::GetTypeInValue() const
 
 void SGraphPinString::SetTypeInValue(const FText& NewTypeInValue, ETextCommit::Type /*CommitInfo*/)
 {
+	if(GraphPinObj->IsPendingKill())
+	{
+		return;
+	}
+
 	if(!GraphPinObj->GetDefaultAsString().Equals(NewTypeInValue.ToString()))
 	{
 		const FScopedTransaction Transaction( NSLOCTEXT("GraphEditor", "ChangeStringPinValue", "Change String Pin Value" ) );
diff --git a/Engine/Source/Editor/GraphEditor/Private/KismetPins/SGraphPinStruct.cpp b/Engine/Source/Editor/GraphEditor/Private/KismetPins/SGraphPinStruct.cpp
--- a/Engine/Source/Editor/GraphEditor/Private/KismetPins/SGraphPinStruct.cpp
+++ b/Engine/Source/Editor/GraphEditor/Private/KismetPins/SGraphPinStruct.cpp
@@ -108,6 +108,11 @@ FOnClicked SGraphPinStruct::GetOnUseButtonDelegate()
 
 void SGraphPinStruct::OnPickedNewStruct(const UScriptStruct* ChosenStruct)
 {
+	if(GraphPinObj->IsPendingKill())
+	{
+		return;
+	}
+
 	FString NewPath;
 	if (ChosenStruct)
 	{
diff --git a/Engine/Source/Editor/GraphEditor/Private/KismetPins/SGraphPinVector.cpp b/Engine/Source/Editor/GraphEditor/Private/KismetPins/SGraphPinVector.cpp
--- a/Engine/Source/Editor/GraphEditor/Private/KismetPins/SGraphPinVector.cpp
+++ b/Engine/Source/Editor/GraphEditor/Private/KismetPins/SGraphPinVector.cpp
@@ -210,6 +210,11 @@ FString SGraphPinVector::GetValue( ETextBoxIndex Index ) const
 
 void SGraphPinVector::OnChangedValueTextBox_0(float NewValue, ETextCommit::Type CommitInfo)
 {
+	if(GraphPinObj->IsPendingKill())
+	{
+		return;
+	}
+
 	const FString ValueStr = FString::Printf( TEXT("%f"), NewValue );
 
 	FString DefaultValue;
@@ -236,6 +241,11 @@ void SGraphPinVector::OnChangedValueTextBox_0(float NewValue, ETextCommit::Type
 
 void SGraphPinVector::OnChangedValueTextBox_1(float NewValue, ETextCommit::Type CommitInfo)
 {
+	if(GraphPinObj->IsPendingKill())
+	{
+		return;
+	}
+
 	const FString ValueStr = FString::Printf( TEXT("%f"), NewValue );
 
 	FString DefaultValue;
@@ -262,6 +272,11 @@ void SGraphPinVector::OnChangedValueTextBox_1(float NewValue, ETextCommit::Type
 
 void SGraphPinVector::OnChangedValueTextBox_2(float NewValue, ETextCommit::Type CommitInfo)
 {
+	if(GraphPinObj->IsPendingKill())
+	{
+		return;
+	}
+
 	const FString ValueStr = FString::Printf( TEXT("%f"), NewValue );
 
 	FString DefaultValue;
diff --git a/Engine/Source/Editor/GraphEditor/Private/KismetPins/SGraphPinVector2D.cpp b/Engine/Source/Editor/GraphEditor/Private/KismetPins/SGraphPinVector2D.cpp
--- a/Engine/Source/Editor/GraphEditor/Private/KismetPins/SGraphPinVector2D.cpp
+++ b/Engine/Source/Editor/GraphEditor/Private/KismetPins/SGraphPinVector2D.cpp
@@ -157,6 +157,11 @@ FString MakeVector2DString(const FString& X, const FString& Y)
 
 void SGraphPinVector2D::OnChangedValueTextBox_X(float NewValue, ETextCommit::Type CommitInfo)
 {
+	if(GraphPinObj->IsPendingKill())
+	{
+		return;
+	}
+
 	const FString ValueStr = FString::Printf( TEXT("%f"), NewValue );
 	const FString Vector2DString = MakeVector2DString(ValueStr, GetValue(TextBox_Y));
 
@@ -172,6 +177,11 @@ void SGraphPinVector2D::OnChangedValueTextBox_X(float NewValue, ETextCommit::Typ
 
 void SGraphPinVector2D::OnChangedValueTextBox_Y(float NewValue, ETextCommit::Type CommitInfo)
 {
+	if(GraphPinObj->IsPendingKill())
+	{
+		return;
+	}
+
 	const FString ValueStr = FString::Printf( TEXT("%f"), NewValue );
 	const FString Vector2DString =MakeVector2DString(GetValue(TextBox_X), ValueStr);
 
diff --git a/Engine/Source/Editor/GraphEditor/Public/KismetPins/SGraphPinNum.h b/Engine/Source/Editor/GraphEditor/Public/KismetPins/SGraphPinNum.h
--- a/Engine/Source/Editor/GraphEditor/Public/KismetPins/SGraphPinNum.h
+++ b/Engine/Source/Editor/GraphEditor/Public/KismetPins/SGraphPinNum.h
@@ -55,6 +55,11 @@ protected:
 
 	void SetNumericValue(NumericType InValue, ETextCommit::Type CommitType)
 	{
+		if(GraphPinObj->IsPendingKill())
+		{
+			return;
+		}
+
 		const FString TypeValueString = LexToString(InValue);
 		if (GraphPinObj->GetDefaultAsString() != TypeValueString)
 		{
