Graph Editor context menus now setup in a way that allows them to be customized

diff --git a/Engine/Source/Editor/GraphEditor/Private/SGraphEditorImpl.cpp b/Engine/Source/Editor/GraphEditor/Private/SGraphEditorImpl.cpp
--- a/Engine/Source/Editor/GraphEditor/Private/SGraphEditorImpl.cpp
+++ b/Engine/Source/Editor/GraphEditor/Private/SGraphEditorImpl.cpp
@@ -500,7 +500,7 @@ void SGraphEditorImpl::OnClosedActionMenu()
 	GraphPanel->OnStopMakingConnection(/*bForceStop=*/ true);
 }
 
-void SGraphEditorImpl::RegisterContextMenuFor_EdGraphSchema()
+void SGraphEditorImpl::RegisterContextMenuFor_EdGraphSchema(const FName ParentMenuName)
 {
 	static const FName MenuName = UEdGraphSchema::GetContextMenuName(UEdGraphSchema::StaticClass());
 	UToolMenus* ToolMenus = UToolMenus::Get();
@@ -509,7 +509,7 @@ void SGraphEditorImpl::RegisterContextMenuFor_EdGraphSchema()
 		return;
 	}
 
-	UToolMenu* Menu = ToolMenus->RegisterMenu(MenuName);
+	UToolMenu* Menu = ToolMenus->RegisterMenu(MenuName, ParentMenuName);
 
 	Menu->AddDynamicSection("GetNodeContextMenuActions", FNewToolMenuDelegate::CreateLambda([](UToolMenu* InMenu)
 	{
@@ -520,132 +520,134 @@ void SGraphEditorImpl::RegisterContextMenuFor_EdGraphSchema()
 		}
 	}));
 
-	Menu->AddDynamicSection("EdGraphSchema", FNewToolMenuDelegate::CreateLambda([](UToolMenu* InMenu)
+	Menu->AddDynamicSection("EdGraphSchema", FNewToolMenuDelegate::CreateStatic(&SGraphEditorImpl::AddContextMenuCommentSection));
+}
+
+void SGraphEditorImpl::AddContextMenuCommentSection(UToolMenu* InMenu)
+{
+	UGraphNodeContextMenuContext* Context = InMenu->FindContext<UGraphNodeContextMenuContext>();
+	if (!Context)
 	{
-		UGraphNodeContextMenuContext* Context = InMenu->FindContext<UGraphNodeContextMenuContext>();
-		if (!Context)
-		{
-			return;
-		}
+		return;
+	}
 
-		const UEdGraphSchema* GraphSchema = Context->Graph->GetSchema();
-		if (!GraphSchema)
-		{
-			return;
-		}
+	const UEdGraphSchema* GraphSchema = Context->Graph->GetSchema();
+	if (!GraphSchema)
+	{
+		return;
+	}
 
-		// Helper to do the node comment editing
-		struct Local
+	// Helper to do the node comment editing
+	struct Local
+	{
+		// Called by the EditableText widget to get the current comment for the node
+		static FString GetNodeComment(TWeakObjectPtr<UEdGraphNode> NodeWeakPtr)
 		{
-			// Called by the EditableText widget to get the current comment for the node
-			static FString GetNodeComment(TWeakObjectPtr<UEdGraphNode> NodeWeakPtr)
+			if (UEdGraphNode* SelectedNode = NodeWeakPtr.Get())
 			{
-				if (UEdGraphNode* SelectedNode = NodeWeakPtr.Get())
-				{
-					return SelectedNode->NodeComment;
-				}
-				return FString();
+				return SelectedNode->NodeComment;
 			}
+			return FString();
+		}
 
-			// Called by the EditableText widget when the user types a new comment for the selected node
-			static void OnNodeCommentTextCommitted(const FText& NewText, ETextCommit::Type CommitInfo, TWeakObjectPtr<UEdGraphNode> NodeWeakPtr)
+		// Called by the EditableText widget when the user types a new comment for the selected node
+		static void OnNodeCommentTextCommitted(const FText& NewText, ETextCommit::Type CommitInfo, TWeakObjectPtr<UEdGraphNode> NodeWeakPtr)
+		{
+			// Apply the change to the selected actor
+			UEdGraphNode* SelectedNode = NodeWeakPtr.Get();
+			FString NewString = NewText.ToString();
+			if (SelectedNode && !SelectedNode->NodeComment.Equals(NewString, ESearchCase::CaseSensitive))
 			{
-				// Apply the change to the selected actor
-				UEdGraphNode* SelectedNode = NodeWeakPtr.Get();
-				FString NewString = NewText.ToString();
-				if (SelectedNode && !SelectedNode->NodeComment.Equals(NewString, ESearchCase::CaseSensitive))
+				// send property changed events
+				const FScopedTransaction Transaction(LOCTEXT("EditNodeComment", "Change Node Comment"));
+				SelectedNode->Modify();
+				UProperty* NodeCommentProperty = FindField<UProperty>(SelectedNode->GetClass(), "NodeComment");
+				if (NodeCommentProperty != nullptr)
 				{
-					// send property changed events
-					const FScopedTransaction Transaction(LOCTEXT("EditNodeComment", "Change Node Comment"));
-					SelectedNode->Modify();
-					UProperty* NodeCommentProperty = FindField<UProperty>(SelectedNode->GetClass(), "NodeComment");
-					if (NodeCommentProperty != nullptr)
-					{
-						SelectedNode->PreEditChange(NodeCommentProperty);
+					SelectedNode->PreEditChange(NodeCommentProperty);
 
-						SelectedNode->NodeComment = NewString;
-						SelectedNode->SetMakeCommentBubbleVisible(true);
+					SelectedNode->NodeComment = NewString;
+					SelectedNode->SetMakeCommentBubbleVisible(true);
 
-						FPropertyChangedEvent NodeCommentPropertyChangedEvent(NodeCommentProperty);
-						SelectedNode->PostEditChangeProperty(NodeCommentPropertyChangedEvent);
-					}
+					FPropertyChangedEvent NodeCommentPropertyChangedEvent(NodeCommentProperty);
+					SelectedNode->PostEditChangeProperty(NodeCommentPropertyChangedEvent);
 				}
+			}
 
-				// Only dismiss all menus if the text was committed via some user action.
-				// ETextCommit::Default implies that focus was switched by some other means. If this is because a submenu was opened, we don't want to close all the menus as a consequence.
-				if (CommitInfo != ETextCommit::Default)
-				{
-					FSlateApplication::Get().DismissAllMenus();
-				}
+			// Only dismiss all menus if the text was committed via some user action.
+			// ETextCommit::Default implies that focus was switched by some other means. If this is because a submenu was opened, we don't want to close all the menus as a consequence.
+			if (CommitInfo != ETextCommit::Default)
+			{
+				FSlateApplication::Get().DismissAllMenus();
 			}
-		};
+		}
+	};
 
-		if (!Context->Pin)
+	if (!Context->Pin)
+	{
+		int32 SelectionCount = GraphSchema->GetNodeSelectionCount(Context->Graph);
+
+		if (SelectionCount == 1)
 		{
-			int32 SelectionCount = GraphSchema->GetNodeSelectionCount(Context->Graph);
+			// Node comment area
+			TSharedRef<SHorizontalBox> NodeCommentBox = SNew(SHorizontalBox);
 
-			if (SelectionCount == 1)
 			{
-				// Node comment area
-				TSharedRef<SHorizontalBox> NodeCommentBox = SNew(SHorizontalBox);
+				FToolMenuSection& Section = InMenu->AddSection("GraphNodeComment", LOCTEXT("NodeCommentMenuHeader", "Node Comment"));
+				Section.AddEntry(FToolMenuEntry::InitWidget("NodeCommentBox", NodeCommentBox, FText::GetEmpty()));
+			}
+			TWeakObjectPtr<UEdGraphNode> SelectedNodeWeakPtr = MakeWeakObjectPtr(const_cast<UEdGraphNode*>(Context->Node));
 
-				{
-					FToolMenuSection& Section = InMenu->AddSection("GraphNodeComment", LOCTEXT("NodeCommentMenuHeader", "Node Comment"));
-					Section.AddEntry(FToolMenuEntry::InitWidget("NodeCommentBox", NodeCommentBox, FText::GetEmpty()));
-				}
-				TWeakObjectPtr<UEdGraphNode> SelectedNodeWeakPtr = MakeWeakObjectPtr(const_cast<UEdGraphNode*>(Context->Node));
+			FText NodeCommentText;
+			if (UEdGraphNode* SelectedNode = SelectedNodeWeakPtr.Get())
+			{
+				NodeCommentText = FText::FromString(SelectedNode->NodeComment);
+			}
 
-				FText NodeCommentText;
-				if (UEdGraphNode* SelectedNode = SelectedNodeWeakPtr.Get())
-				{
-					NodeCommentText = FText::FromString(SelectedNode->NodeComment);
-				}
+			const FSlateBrush* NodeIcon = FCoreStyle::Get().GetDefaultBrush();//@TODO: FActorIconFinder::FindIconForActor(SelectedActors(0).Get());
 
-				const FSlateBrush* NodeIcon = FCoreStyle::Get().GetDefaultBrush();//@TODO: FActorIconFinder::FindIconForActor(SelectedActors(0).Get());
-
-				// Comment label
-				NodeCommentBox->AddSlot()
-					.VAlign(VAlign_Center)
-					.FillWidth(1.0f)
-					[
-						SNew(SMultiLineEditableTextBox)
-						.Text(NodeCommentText)
-						.ToolTipText(LOCTEXT("NodeComment_ToolTip", "Comment for this node"))
-						.OnTextCommitted_Static(&Local::OnNodeCommentTextCommitted, SelectedNodeWeakPtr)
-						.SelectAllTextWhenFocused(true)
-						.RevertTextOnEscape(true)
-						.ModiferKeyForNewLine(EModifierKey::Control)
-					];
-			}
-			else if (SelectionCount > 1)
+			// Comment label
+			NodeCommentBox->AddSlot()
+				.VAlign(VAlign_Center)
+				.FillWidth(1.0f)
+				[
+					SNew(SMultiLineEditableTextBox)
+					.Text(NodeCommentText)
+					.ToolTipText(LOCTEXT("NodeComment_ToolTip", "Comment for this node"))
+					.OnTextCommitted_Static(&Local::OnNodeCommentTextCommitted, SelectedNodeWeakPtr)
+					.SelectAllTextWhenFocused(true)
+					.RevertTextOnEscape(true)
+					.ModiferKeyForNewLine(EModifierKey::Control)
+				];
+		}
+		else if (SelectionCount > 1)
+		{
+			struct SCommentUtility
 			{
-				struct SCommentUtility
+				static void CreateComment(const UEdGraphSchema* Schema, UEdGraph* Graph)
 				{
-					static void CreateComment(const UEdGraphSchema* Schema, UEdGraph* Graph)
+					if (Schema && Graph)
 					{
-						if (Schema && Graph)
-						{
-							TSharedPtr<FEdGraphSchemaAction> Action = Schema->GetCreateCommentAction();
+						TSharedPtr<FEdGraphSchemaAction> Action = Schema->GetCreateCommentAction();
 
-							if (Action.IsValid())
-							{
-								Action->PerformAction(Graph, nullptr, FVector2D());
-							}
+						if (Action.IsValid())
+						{
+							Action->PerformAction(Graph, nullptr, FVector2D());
 						}
 					}
-				};
-
-				FToolMenuSection& Section = InMenu->AddSection("SchemaActionComment", LOCTEXT("MultiCommentHeader", "Comment Group"));
-				Section.AddMenuEntry(
-					"MultiCommentDesc",
-					LOCTEXT("MultiCommentDesc", "Create Comment from Selection"),
-					LOCTEXT("CommentToolTip", "Create a resizable comment box around selection."),
-					FSlateIcon(),
-					FExecuteAction::CreateStatic(SCommentUtility::CreateComment, GraphSchema, const_cast<UEdGraph*>(Context->Graph)
-				));
-			}
+				}
+			};
+
+			FToolMenuSection& Section = InMenu->AddSection("SchemaActionComment", LOCTEXT("MultiCommentHeader", "Comment Group"));
+			Section.AddMenuEntry(
+				"MultiCommentDesc",
+				LOCTEXT("MultiCommentDesc", "Create Comment from Selection"),
+				LOCTEXT("CommentToolTip", "Create a resizable comment box around selection."),
+				FSlateIcon(),
+				FExecuteAction::CreateStatic(SCommentUtility::CreateComment, GraphSchema, const_cast<UEdGraph*>(Context->Graph)
+			));
 		}
-	}));
+	}
 }
 
 FName SGraphEditorImpl::GetNodeParentContextMenuName(UClass* InClass)
@@ -669,78 +671,119 @@ FName SGraphEditorImpl::GetNodeContextMenuName(UClass* InClass)
 	return FName(*(FString(TEXT("GraphEditor.GraphNodeContextMenu.")) + InClass->GetName()));
 }
 
-UToolMenu* SGraphEditorImpl::GenerateContextMenu(const UEdGraphSchema* Schema, FToolMenuContext& MenuContext) const
+void SGraphEditorImpl::RegisterContextMenu(const UEdGraphSchema* Schema, FToolMenuContext& MenuContext) const
 {
+	UToolMenus* ToolMenus = UToolMenus::Get();
+	const FName SchemaMenuName = Schema->GetContextMenuName();
 	UGraphNodeContextMenuContext* Context = MenuContext.Find<UGraphNodeContextMenuContext>();
 
-	// Register UEdGraphSchema's menu here to reduce amount of editor code in engine module
-	RegisterContextMenuFor_EdGraphSchema();
+	// Root menu
+	// "GraphEditor.GraphContextMenu.Common"
+	// contains: GraphSchema->GetContextMenuActions(Menu, Context)
+	const FName CommonRootMenuName = "GraphEditor.GraphContextMenu.Common";
+	if (!ToolMenus->IsMenuRegistered(CommonRootMenuName))
+	{
+		UToolMenu* CommonRootMenu = ToolMenus->RegisterMenu(CommonRootMenuName);
+		CommonRootMenu->AddDynamicSection("GetContextMenuActions", FNewToolMenuDelegate::CreateLambda([](UToolMenu* InMenu)
+		{
+			if (UGraphNodeContextMenuContext* ContextObject = InMenu->FindContext<UGraphNodeContextMenuContext>())
+			{
+				if (const UEdGraphSchema* GraphSchema = ContextObject->Graph->GetSchema())
+				{
+					GraphSchema->GetContextMenuActions(InMenu, ContextObject);
+				}
+			}
+		}));
+	}
 
-	// Register any missing menus otherwise there is a break in the hierarchy chain
-	UToolMenus* ToolMenus = UToolMenus::Get();
+	// Menu for EdGraphSchema, used most but not all schema subclasses
+	// contains: Node->GetNodeContextMenuActions()
+	// contains: Comment for node
+	RegisterContextMenuFor_EdGraphSchema(CommonRootMenuName);
+
+	// Menus for subclasses of EdGraphSchema
 	for (UClass* CurrentClass = Schema->GetClass(); CurrentClass && CurrentClass->IsChildOf(UEdGraphSchema::StaticClass()); CurrentClass = CurrentClass->GetSuperClass())
 	{
+		// EdGraphSchema class is registered instead by RegisterContextMenuFor_EdGraphSchema()
+		if (CurrentClass == UEdGraphSchema::StaticClass())
+		{
+			break;
+		}
+
 		const UEdGraphSchema* CurrentSchema = CurrentClass->GetDefaultObject<UEdGraphSchema>();
 		const FName CheckMenuName = CurrentSchema->GetContextMenuName();
+
+		// Some subclasses of UEdGraphSchema chose not to include UEdGraphSchema's menu
+		// Note: menu "GraphEditor.GraphContextMenu.Common" calls GraphSchema->GetContextMenuActions() and adds entry for node comment
+		const FName CheckParentName = CurrentSchema->GetParentContextMenuName();
+
 		if (!ToolMenus->IsMenuRegistered(CheckMenuName))
 		{
-			ToolMenus->RegisterMenu(CheckMenuName, CurrentSchema->GetParentContextMenuName());
+			FName ParentNameToUse = CheckParentName;
+			
+			// Connect final menu in chain to the common root
+			if (ParentNameToUse == NAME_None)
+			{
+				ParentNameToUse = CommonRootMenuName;
+			}
+
+			ToolMenus->RegisterMenu(CheckMenuName, ParentNameToUse);
 		}
-	}
 
-	const FName MenuName = Schema->GetContextMenuName();
-	TArray<UToolMenu*> Hierarchy = ToolMenus->CollectHierarchy(MenuName);
+		if (CheckParentName == NAME_None)
+		{
+			break;
+		}
+	}
 
-	// Insert list of node specific menus
+	// Node menus
 	if (Context->Node)
 	{
-		// Only insert list if graph schema includes actions for nodes
-		static const FName EdGraphSchemaMenuName = UEdGraphSchema::GetContextMenuName(UEdGraphSchema::StaticClass());
-		if (Hierarchy.ContainsByPredicate([=](UToolMenu* OtherMenu) { return OtherMenu->MenuName == EdGraphSchemaMenuName; }))
+		for (UClass* CurrentClass = Context->Node->GetClass(); CurrentClass && CurrentClass->IsChildOf(UEdGraphNode::StaticClass()); CurrentClass = CurrentClass->GetSuperClass())
 		{
-			// Walk class hierarchy
-			for (UClass* CurrentClass = Context->Node->GetClass(); CurrentClass && CurrentClass->IsChildOf(UEdGraphNode::StaticClass()); CurrentClass = CurrentClass->GetSuperClass())
+			const FName CheckMenuName = GetNodeContextMenuName(CurrentClass);
+			const FName CheckParentName = GetNodeParentContextMenuName(CurrentClass);
+
+			if (!ToolMenus->IsMenuRegistered(CheckMenuName))
 			{
-				const FName CheckMenuName = GetNodeContextMenuName(CurrentClass);
-				const FName CheckParentName = GetNodeParentContextMenuName(CurrentClass);
+				FName ParentNameToUse = CheckParentName;
 
-				// Register if not already registered
-				if (!ToolMenus->IsMenuRegistered(CheckMenuName))
+				// Connect final menu in chain to schema's chain of menus
+				if (CheckParentName == NAME_None)
 				{
-					ToolMenus->RegisterMenu(CheckMenuName, CheckParentName);
+					ParentNameToUse = SchemaMenuName;
 				}
 
-				// Parent rest of hierarchy to this menu
-				Hierarchy.Insert(ToolMenus->FindMenu(CheckMenuName), 0);
+				ToolMenus->RegisterMenu(CheckMenuName, ParentNameToUse);
+			}
 
-				// Stop walking class hierarchy if IncludeParentNodeContextMenu() returns false for current class
-				if (CheckParentName == NAME_None)
-				{
-					break;
-				}
+			if (CheckParentName == NAME_None)
+			{
+				break;
 			}
 		}
 	}
+}
 
-	// Add extra menu into hierarchy for calling GetContextMenuActions()
-	const FName CommonRootMenuName = "GraphEditor.GraphContextMenu.Common";
-	if (!ToolMenus->IsMenuRegistered(CommonRootMenuName))
+UToolMenu* SGraphEditorImpl::GenerateContextMenu(const UEdGraphSchema* Schema, FToolMenuContext& MenuContext) const
+{
+	// Register all the menu's needed
+	RegisterContextMenu(Schema, MenuContext);
+
+	UToolMenus* ToolMenus = UToolMenus::Get();
+	UGraphNodeContextMenuContext* Context = MenuContext.Find<UGraphNodeContextMenuContext>();
+
+	FName MenuName = NAME_None;
+	if (Context->Node)
 	{
-		UToolMenu* CommonRootMenu = ToolMenus->RegisterMenu(CommonRootMenuName);
-		CommonRootMenu->AddDynamicSection("GetContextMenuActions", FNewToolMenuDelegate::CreateLambda([](UToolMenu* InMenu)
-		{
-			if (UGraphNodeContextMenuContext* ContextObject = InMenu->FindContext<UGraphNodeContextMenuContext>())
-			{
-				if (const UEdGraphSchema* GraphSchema = ContextObject->Graph->GetSchema())
-				{
-					GraphSchema->GetContextMenuActions(InMenu, ContextObject);
-				}
-			}
-		}));
+		MenuName = GetNodeContextMenuName(Context->Node->GetClass());
+	}
+	else
+	{
+		MenuName = Schema->GetContextMenuName();
 	}
-	Hierarchy.Insert(ToolMenus->FindMenu(CommonRootMenuName), 0);
 
-	return ToolMenus->GenerateMenu(Hierarchy, MenuContext);
+	return ToolMenus->GenerateMenu(MenuName, MenuContext);
 }
 
 FActionMenuContent SGraphEditorImpl::GraphEd_OnGetContextMenuFor(const FGraphContextMenuArguments& SpawnInfo)
diff --git a/Engine/Source/Editor/GraphEditor/Private/SGraphEditorImpl.h b/Engine/Source/Editor/GraphEditor/Private/SGraphEditorImpl.h
--- a/Engine/Source/Editor/GraphEditor/Private/SGraphEditorImpl.h
+++ b/Engine/Source/Editor/GraphEditor/Private/SGraphEditorImpl.h
@@ -334,11 +334,13 @@ private:
 
 	bool IsLocked() const;
 	
+	void RegisterContextMenu(const class UEdGraphSchema* Schema, struct FToolMenuContext& MenuContext) const;
 	class UToolMenu* GenerateContextMenu(const class UEdGraphSchema* Schema, struct FToolMenuContext& MenuContext) const;
 
-	static void RegisterContextMenuFor_EdGraphSchema();
+	static void RegisterContextMenuFor_EdGraphSchema(const FName ParentMenuName);
 	static FName GetNodeParentContextMenuName(UClass* InClass);
 	static FName GetNodeContextMenuName(UClass* InClass);
+	static void AddContextMenuCommentSection(UToolMenu* InMenu);
 
 private:
 	bool bIsActiveTimerRegistered;
