Updated blueprint graph to cancel ZoomToFit on nodes when panning with the middle mouse button

diff --git a/Engine/Source/Editor/GraphEditor/Private/SNodePanel.cpp b/Engine/Source/Editor/GraphEditor/Private/SNodePanel.cpp
--- a/Engine/Source/Editor/GraphEditor/Private/SNodePanel.cpp
+++ b/Engine/Source/Editor/GraphEditor/Private/SNodePanel.cpp
@@ -609,6 +609,7 @@ FReply SNodePanel::OnMouseButtonDown( const FGeometry& MyGeometry, const FPointe
 		SoftwareCursorPosition = PanelCoordToGraphCoord(MyGeometry.AbsoluteToLocal(MouseEvent.GetScreenSpacePosition()));
 
 		DeferredMovementTargetObject = nullptr; // clear any interpolation when you manually pan
+		CancelZoomToFit();
 
 		// MIDDLE BUTTON is for dragging only.
 		return ReplyState;
