Added SGraphNode::CreateTitleRightWidget() to allow for adding widgets to the far right of the title in a graph node.

Fixed copy-paste mistake in code to get SGraphNode error tooltip.

diff --git a/Engine/Source/Editor/GraphEditor/Private/SGraphNode.cpp b/Engine/Source/Editor/GraphEditor/Private/SGraphNode.cpp
--- a/Engine/Source/Editor/GraphEditor/Private/SGraphNode.cpp
+++ b/Engine/Source/Editor/GraphEditor/Private/SGraphNode.cpp
@@ -739,6 +739,11 @@ TSharedRef<SWidget> SGraphNode::CreateTitleWidget(TSharedPtr<SNodeTitle> NodeTit
 	return InlineEditableText.ToSharedRef();
 }
 
+TSharedRef<SWidget> SGraphNode::CreateTitleRightWidget()
+{
+	return SNullWidget::NullWidget;
+}
+
 /**
  * Update this GraphNode to match the data that it is observing
  */
@@ -770,7 +775,7 @@ void SGraphNode::UpdateGraphNode()
 
 	// Get node icon
 	IconColor = FLinearColor::White;
-	const FSlateBrush* IconBrush = NULL;
+	const FSlateBrush* IconBrush = nullptr;
 	if (GraphNode != NULL && GraphNode->ShowPaletteIconOnNode())
 	{
 		IconBrush = GraphNode->GetIconAndTint(IconColor).GetOptionalIcon();
@@ -785,41 +790,54 @@ void SGraphNode::UpdateGraphNode()
 			.ColorAndOpacity( this, &SGraphNode::GetNodeTitleIconColor )
 		]
 		+SOverlay::Slot()
-		.HAlign(HAlign_Left)
+		.HAlign(HAlign_Fill)
 		.VAlign(VAlign_Center)
 		[
-			SNew(SBorder)
-			.BorderImage( FEditorStyle::GetBrush("Graph.Node.ColorSpill") )
-			// The extra margin on the right
-			// is for making the color spill stretch well past the node title
-			.Padding( FMargin(10,5,30,3) )
-			.BorderBackgroundColor( this, &SGraphNode::GetNodeTitleColor )
+			SNew(SHorizontalBox)
+			+ SHorizontalBox::Slot()
+			.HAlign(HAlign_Fill)
 			[
-				SNew(SHorizontalBox)
-				+ SHorizontalBox::Slot()
-				.VAlign(VAlign_Top)
-				.Padding(FMargin(0.f, 0.f, 4.f, 0.f))
-				.AutoWidth()
-				[
-					SNew(SImage)
-					.Image(IconBrush)
-					.ColorAndOpacity(this, &SGraphNode::GetNodeTitleIconColor)
-				]
-				+ SHorizontalBox::Slot()
+				SNew(SBorder)
+				.BorderImage( FEditorStyle::GetBrush("Graph.Node.ColorSpill") )
+				// The extra margin on the right
+				// is for making the color spill stretch well past the node title
+				.Padding( FMargin(10,5,30,3) )
+				.BorderBackgroundColor( this, &SGraphNode::GetNodeTitleColor )
 				[
-					SNew(SVerticalBox)
-					+ SVerticalBox::Slot()
-					.AutoHeight()
+					SNew(SHorizontalBox)
+					+ SHorizontalBox::Slot()
+					.VAlign(VAlign_Top)
+					.Padding(FMargin(0.f, 0.f, 4.f, 0.f))
+					.AutoWidth()
 					[
-						CreateTitleWidget(NodeTitle)
+						SNew(SImage)
+						.Image(IconBrush)
+						.ColorAndOpacity(this, &SGraphNode::GetNodeTitleIconColor)
 					]
-					+ SVerticalBox::Slot()
-					.AutoHeight()
+					+ SHorizontalBox::Slot()
 					[
-						NodeTitle.ToSharedRef()
+						SNew(SVerticalBox)
+						+ SVerticalBox::Slot()
+						.AutoHeight()
+						[
+							CreateTitleWidget(NodeTitle)
+						]
+						+ SVerticalBox::Slot()
+						.AutoHeight()
+						[
+							NodeTitle.ToSharedRef()
+						]
 					]
 				]
 			]
+			+ SHorizontalBox::Slot()
+			.HAlign(HAlign_Right)
+			.VAlign(VAlign_Center)
+			.Padding(0, 0, 5, 0)
+			.AutoWidth()
+			[
+				CreateTitleRightWidget()
+			]
 		]
 		+SOverlay::Slot()
 		.VAlign(VAlign_Top)
@@ -1364,25 +1382,27 @@ void SGraphNode::PositionThisNodeBetweenOtherNodes(const FVector2D& PrevPos, con
 	GraphNode->NodePosY = NewCorner.Y;
 }
 
-FText SGraphNode::GetErrorMsgToolTip( ) const
+FText SGraphNode::GetErrorMsgToolTip() const
 {
-	FText Result;
-	// Append the node's upgrade message, if any.
-	if (!GraphNode->NodeUpgradeMessage.IsEmpty())
+	FText Result = FText::GetEmpty();
+	if (GraphNode != nullptr)
 	{
-		if (Result.IsEmpty())
-		{
-			Result = GraphNode->NodeUpgradeMessage;
-		}
-		else
+		Result = FText::FromString(GraphNode->ErrorMsg);
+
+		// Append the node's upgrade message, if any.
+		if (!GraphNode->NodeUpgradeMessage.IsEmpty())
 		{
-			Result = FText::Format(FText::FromString(TEXT("{0}\n\n{1}")), Result, GraphNode->NodeUpgradeMessage);
+			if (Result.IsEmpty())
+			{
+				Result = GraphNode->NodeUpgradeMessage;
+			}
+			else
+			{
+				Result = FText::Format(FText::FromString(TEXT("{0}\n\n{1}")), Result, GraphNode->NodeUpgradeMessage);
+			}
 		}
 	}
-	else
-	{
-		Result = FText::FromString(GraphNode->ErrorMsg);
-	}
+
 	return Result;
 }
 
diff --git a/Engine/Source/Editor/GraphEditor/Public/SGraphNode.h b/Engine/Source/Editor/GraphEditor/Public/SGraphNode.h
--- a/Engine/Source/Editor/GraphEditor/Public/SGraphNode.h
+++ b/Engine/Source/Editor/GraphEditor/Public/SGraphNode.h
@@ -333,6 +333,9 @@ protected:
 	virtual void SetDefaultTitleAreaWidget(TSharedRef<SOverlay> DefaultTitleAreaWidget) {}
 	virtual TSharedRef<SWidget> CreateTitleWidget(TSharedPtr<SNodeTitle> NodeTitle);
 	
+	/** Optionally create a widget to dock to the right in the title bar of the node. */
+	virtual TSharedRef<SWidget> CreateTitleRightWidget();
+	
 	/** Create the inner node content area, including the left/right pin boxes */
 	virtual TSharedRef<SWidget> CreateNodeContentArea();
 
