Lines to pins on nodes that are outside the current viewport are now correctly drawn when using UI scaling on a high DPI monitor.

diff --git a/Engine/Source/Editor/GraphEditor/Private/SGraphPanel.cpp b/Engine/Source/Editor/GraphEditor/Private/SGraphPanel.cpp
--- a/Engine/Source/Editor/GraphEditor/Private/SGraphPanel.cpp
+++ b/Engine/Source/Editor/GraphEditor/Private/SGraphPanel.cpp
@@ -346,14 +346,14 @@ int32 SGraphPanel::OnPaint( const FPaintArgs& Args, const FGeometry& AllottedGeo
 				ChildNode->GetPins(NodePins);
 
 				const FVector2D NodeLoc = ChildNode->GetPosition();
-				const FGeometry SynthesizedNodeGeometry(GraphCoordToPanelCoord(NodeLoc), AllottedGeometry.AbsolutePosition, FVector2D::ZeroVector, 1.f);
+				const FGeometry SynthesizedNodeGeometry(GraphCoordToPanelCoord(NodeLoc) * AllottedGeometry.Scale, AllottedGeometry.AbsolutePosition, FVector2D::ZeroVector, 1.f);
 
 				for (TArray< TSharedRef<SWidget> >::TConstIterator NodePinIterator(NodePins); NodePinIterator; ++NodePinIterator)
 				{
 					const SGraphPin& PinWidget = static_cast<const SGraphPin&>((*NodePinIterator).Get());
 					FVector2D PinLoc = NodeLoc + PinWidget.GetNodeOffset();
 
-					const FGeometry SynthesizedPinGeometry(GraphCoordToPanelCoord(PinLoc), AllottedGeometry.AbsolutePosition, FVector2D::ZeroVector, 1.f);
+					const FGeometry SynthesizedPinGeometry(GraphCoordToPanelCoord(PinLoc) * AllottedGeometry.Scale, AllottedGeometry.AbsolutePosition, FVector2D::ZeroVector, 1.f);
 					PinGeometries.Add(*NodePinIterator, FArrangedWidget(*NodePinIterator, SynthesizedPinGeometry));
 				}
 
