Blueprints - We now sort the CreateDelegate entries in the blueprint node.

diff --git a/Engine/Source/Editor/GraphEditor/Private/KismetNodes/SGraphNodeK2CreateDelegate.cpp b/Engine/Source/Editor/GraphEditor/Private/KismetNodes/SGraphNodeK2CreateDelegate.cpp
--- a/Engine/Source/Editor/GraphEditor/Private/KismetNodes/SGraphNodeK2CreateDelegate.cpp
+++ b/Engine/Source/Editor/GraphEditor/Private/KismetNodes/SGraphNodeK2CreateDelegate.cpp
@@ -279,19 +279,27 @@ void SGraphNodeK2CreateDelegate::CreateBelowPinControls(TSharedPtr<SVerticalBox>
 				CreateMatchingEventData = AddDefaultFunctionDataOption(NSLOCTEXT("GraphNodeK2Create", "CreateMatchingEventOption", "[Create a matching event]"));
 			}
 
+			TArray<TSharedRef<FFunctionItemData>> ClassFunctions;
+
 			for (TFieldIterator<UFunction> It(ScopeClass); It; ++It)
 			{
 				UFunction* Func = *It;
 				if (Func && FunctionSignature->IsSignatureCompatibleWith(Func) &&
 					UEdGraphSchema_K2::FunctionCanBeUsedInDelegate(Func))
 				{
-					TSharedPtr<FFunctionItemData> ItemData = MakeShareable(new FFunctionItemData());
+					TSharedRef<FFunctionItemData> ItemData = MakeShared<FFunctionItemData>();
 					ItemData->Name = Func->GetFName();
 					ItemData->Description = FunctionDescription(Func);
-					FunctionDataItems.Add(ItemData);
+					ClassFunctions.Add(ItemData);
 				}
 			}
 
+			ClassFunctions.Sort([](const TSharedRef<FFunctionItemData>& A, const TSharedRef<FFunctionItemData>& B) {
+				return A->Description.CompareTo(B->Description) < 0;
+			});
+
+			FunctionDataItems.Append(ClassFunctions);
+
 			TSharedRef<SComboButton> SelectFunctionWidgetRef = SNew(SComboButton)
 				.Method(EPopupMethod::CreateNewWindow)
 				.ButtonContent()
