Make the "Select Function" drop down searchable on Create event nodes

diff --git a/Engine/Source/Editor/GraphEditor/Private/KismetNodes/SGraphNodeK2CreateDelegate.cpp b/Engine/Source/Editor/GraphEditor/Private/KismetNodes/SGraphNodeK2CreateDelegate.cpp
--- a/Engine/Source/Editor/GraphEditor/Private/KismetNodes/SGraphNodeK2CreateDelegate.cpp
+++ b/Engine/Source/Editor/GraphEditor/Private/KismetNodes/SGraphNodeK2CreateDelegate.cpp
@@ -1,14 +1,9 @@
 // Copyright Epic Games, Inc. All Rights Reserved.
 
 #include "KismetNodes/SGraphNodeK2CreateDelegate.h"
-#include "Widgets/Input/SComboButton.h"
-#include "Widgets/Views/STableViewBase.h"
-#include "Widgets/Views/STableRow.h"
-#include "Widgets/Views/SListView.h"
-#include "Widgets/SBoxPanel.h"
+#include "SSearchableComboBox.h"
 #include "EdGraphSchema_K2.h"
 #include "K2Node_CreateDelegate.h"
-
 #include "Editor/UnrealEd/Public/Kismet2/BlueprintEditorUtils.h"
 #include "Editor/UnrealEd/Public/Kismet2/KismetEditorUtilities.h"
 #include "Editor/UnrealEd/Public/ScopedTransaction.h"
@@ -147,16 +142,7 @@ FText SGraphNodeK2CreateDelegate::GetCurrentFunctionDescription() const
 	return NSLOCTEXT("GraphNodeK2Create", "SelectFunctionLabel", "Select Function...");
 }
 
-TSharedRef<ITableRow> SGraphNodeK2CreateDelegate::HandleGenerateRowFunction(TSharedPtr<FFunctionItemData> FunctionItemData, const TSharedRef<STableViewBase>& OwnerTable)
-{
-	check(FunctionItemData.IsValid());
-	return SNew(STableRow< TSharedPtr<FFunctionItemData> >, OwnerTable).Content()
-		[
-			SNew(STextBlock).Text(FunctionItemData->Description)
-		];
-}
-
-void SGraphNodeK2CreateDelegate::OnFunctionSelected(TSharedPtr<FFunctionItemData> FunctionItemData, ESelectInfo::Type SelectInfo)
+void SGraphNodeK2CreateDelegate::OnFunctionSelected(TSharedPtr<FString> FunctionItemData, ESelectInfo::Type SelectInfo)
 {
 	const FScopedTransaction Transaction(NSLOCTEXT("GraphNodeK2Create", "CreateMatchingSigniture", "Create matching signiture"));
 
@@ -194,7 +180,13 @@ void SGraphNodeK2CreateDelegate::OnFunctionSelected(TSharedPtr<FFunctionItemData
 				FName NewEventName = FBlueprintEditorUtils::FindUniqueCustomEventName(NodeBP);
 
 				UBlueprintEventNodeSpawner* Spawner = UBlueprintEventNodeSpawner::Create(UK2Node_CustomEvent::StaticClass(), NewEventName);
-				UEdGraphNode* NewNode = Spawner->Invoke(Node->GetGraph(), IBlueprintNodeBinder::FBindingSet(), FVector2D(Node->NodePosX, Node->NodePosY + 200));
+				
+				// Get a good spawn location for the new event
+				UEdGraph* NodeGraph = Node->GetGraph();
+				check(NodeGraph);
+				FVector2D SpawnPos = NodeGraph->GetGoodPlaceForNewNode();
+
+				UEdGraphNode* NewNode = Spawner->Invoke(NodeGraph, IBlueprintNodeBinder::FBindingSet(), SpawnPos);
 
 				if (UK2Node_CustomEvent* NewEventNode = Cast<UK2Node_CustomEvent>(NewNode))
 				{
@@ -210,12 +202,13 @@ void SGraphNodeK2CreateDelegate::OnFunctionSelected(TSharedPtr<FFunctionItemData
 			}
 			else
 			{
-				Node->SetFunction(FunctionItemData->Name);
+				FName FuncName( **FunctionItemData.Get());
+				Node->SetFunction(FuncName);
 			}
 			
 			Node->HandleAnyChange(true);
 
-			auto SelectFunctionWidgetPtr = SelectFunctionWidget.Pin();
+			TSharedPtr<SSearchableComboBox> SelectFunctionWidgetPtr = FunctionOptionComboBox.Pin();
 			if (SelectFunctionWidgetPtr.IsValid())
 			{
 				SelectFunctionWidgetPtr->SetIsOpen(false);
@@ -224,12 +217,11 @@ void SGraphNodeK2CreateDelegate::OnFunctionSelected(TSharedPtr<FFunctionItemData
 	}
 }
 
-TSharedPtr<SGraphNodeK2CreateDelegate::FFunctionItemData> SGraphNodeK2CreateDelegate::AddDefaultFunctionDataOption(const FText& DisplayName)
+TSharedPtr<FString> SGraphNodeK2CreateDelegate::AddDefaultFunctionDataOption(const FText& DisplayName)
 {
-	TSharedPtr<FFunctionItemData> NewEntry = MakeShareable(new FFunctionItemData());
-	NewEntry->Description = DisplayName;
-	FunctionDataItems.Add(NewEntry);
-	return NewEntry;
+	TSharedPtr<FString> Res = MakeShareable(new FString(DisplayName.ToString()));
+	FunctionOptionList.Add(Res);
+	return Res;
 }
 
 void SGraphNodeK2CreateDelegate::CreateBelowPinControls(TSharedPtr<SVerticalBox> MainBox)
@@ -265,7 +257,7 @@ void SGraphNodeK2CreateDelegate::CreateBelowPinControls(TSharedPtr<SVerticalBox>
 					.ToolTipText(FunctionSignatureToolTipText)
 				];
 
-			FunctionDataItems.Empty();
+			FunctionOptionList.Empty();
 
 			// add an empty row, so the user can clear the selection if they want
 			AddDefaultFunctionDataOption(NSLOCTEXT("GraphNodeK2Create", "EmptyFunctionOption", "[None]"));
@@ -285,34 +277,24 @@ void SGraphNodeK2CreateDelegate::CreateBelowPinControls(TSharedPtr<SVerticalBox>
 				if (Func && FunctionSignature->IsSignatureCompatibleWith(Func) &&
 					UEdGraphSchema_K2::FunctionCanBeUsedInDelegate(Func))
 				{
-					TSharedPtr<FFunctionItemData> ItemData = MakeShareable(new FFunctionItemData());
-					ItemData->Name = Func->GetFName();
-					ItemData->Description = FunctionDescription(Func);
-					FunctionDataItems.Add(ItemData);
+					// Add this to the searchable text box as an FString so users can type and find it
+					FunctionOptionList.Add(MakeShareable(new FString(Func->GetFName().ToString())));
 				}
 			}
 
-			TSharedRef<SComboButton> SelectFunctionWidgetRef = SNew(SComboButton)
-				.Method(EPopupMethod::CreateNewWindow)
-				.ButtonContent()
+			TSharedRef<SSearchableComboBox> SelectFunctionWidgetRef =
+				SNew(SSearchableComboBox)
+				.OptionsSource(&FunctionOptionList)
+				.OnGenerateWidget(this, &SGraphNodeK2CreateDelegate::MakeFunctionOptionComboWidget)
+				.OnSelectionChanged(this, &SGraphNodeK2CreateDelegate::OnFunctionSelected)
+				.ContentPadding(2)
+				.MaxListHeight(200.0f)
+				.Content()
 				[
 					SNew(STextBlock)
-					.Text(this, &SGraphNodeK2CreateDelegate::GetCurrentFunctionDescription)
-				]
-			.MenuContent()
-				[
-					SNew(SVerticalBox)
- 					 + SVerticalBox::Slot()
- 					 .AutoHeight()
- 					 .MaxHeight(500.f)
- 					 [
- 						SNew(SListView<TSharedPtr<FFunctionItemData> >)
- 						.ListItemsSource(&FunctionDataItems)
- 						.OnGenerateRow(this, &SGraphNodeK2CreateDelegate::HandleGenerateRowFunction)
- 						.OnSelectionChanged(this, &SGraphNodeK2CreateDelegate::OnFunctionSelected)
- 					 ]
+					.Text(GetCurrentFunctionDescription())
 				];
-
+			
 			MainBox->AddSlot()
 				.AutoHeight()
 				.VAlign(VAlign_Fill)
@@ -321,16 +303,21 @@ void SGraphNodeK2CreateDelegate::CreateBelowPinControls(TSharedPtr<SVerticalBox>
 					SelectFunctionWidgetRef
 				];
 
-			SelectFunctionWidget = SelectFunctionWidgetRef;
+			FunctionOptionComboBox = SelectFunctionWidgetRef;
 		}
 	}
 }
 
+TSharedRef<SWidget> SGraphNodeK2CreateDelegate::MakeFunctionOptionComboWidget(TSharedPtr<FString> InItem)
+{
+	return SNew(STextBlock).Text(FText::FromString(*InItem));
+}
+
 SGraphNodeK2CreateDelegate::~SGraphNodeK2CreateDelegate()
 {
-	auto SelectFunctionWidgetPtr = SelectFunctionWidget.Pin();
+	TSharedPtr<SSearchableComboBox> SelectFunctionWidgetPtr = FunctionOptionComboBox.Pin();
 	if (SelectFunctionWidgetPtr.IsValid())
 	{
 		SelectFunctionWidgetPtr->SetIsOpen(false);
 	}
-}
+}
\ No newline at end of file
diff --git a/Engine/Source/Editor/GraphEditor/Private/KismetNodes/SGraphNodeK2CreateDelegate.h b/Engine/Source/Editor/GraphEditor/Private/KismetNodes/SGraphNodeK2CreateDelegate.h
--- a/Engine/Source/Editor/GraphEditor/Private/KismetNodes/SGraphNodeK2CreateDelegate.h
+++ b/Engine/Source/Editor/GraphEditor/Private/KismetNodes/SGraphNodeK2CreateDelegate.h
@@ -6,34 +6,26 @@
 #include "Widgets/DeclarativeSyntaxSupport.h"
 #include "KismetNodes/SGraphNodeK2Base.h"
 
-class ITableRow;
-class SComboButton;
+class SSearchableComboBox;
 class STableViewBase;
 class SVerticalBox;
 class UK2Node;
 
+/**
+* @brief	The CreateDelegate node will allow users to bind to event dispatchers 
+*			based off of appropriate function signatures, or create a matching one.
+*/
 class SGraphNodeK2CreateDelegate : public SGraphNodeK2Base
 {
 public:
 	SLATE_BEGIN_ARGS(SGraphNodeK2CreateDelegate) {}
 	SLATE_END_ARGS()
 
-	/** Data that defines a delegate function */
-	struct FFunctionItemData
-	{
-		FName Name;
-		FText Description;
-	};
-
-	/** Collection of function items that have matching function delegates to this node */
-	TArray<TSharedPtr<FFunctionItemData>> FunctionDataItems;
-	TWeakPtr<SComboButton> SelectFunctionWidget;
-
 	/** Data that can be used to create a matching function based on the parameters of a create event node */
-	TSharedPtr<FFunctionItemData> CreateMatchingFunctionData;
+	TSharedPtr<FString> CreateMatchingFunctionData;
 	
 	/** Data that can be used to create a matching event based on based on the parameters of a create event node */
-	TSharedPtr<FFunctionItemData> CreateMatchingEventData;
+	TSharedPtr<FString> CreateMatchingEventData;
 
 public:
 	virtual ~SGraphNodeK2CreateDelegate();
@@ -44,17 +36,24 @@ protected:
 	static FText FunctionDescription(const UFunction* Function, const bool bOnlyDescribeSignature = false, const int32 CharacterLimit = 32);
 
 	FText GetCurrentFunctionDescription() const;
-	TSharedRef<ITableRow> HandleGenerateRowFunction(TSharedPtr<FFunctionItemData> FunctionItemData, const TSharedRef<STableViewBase>& OwnerTable);
-	void OnFunctionSelected(TSharedPtr<FFunctionItemData> FunctionItemData, ESelectInfo::Type SelectInfo);
 
 private:
 
+	//~ Begin Searchable combo box options
+	TWeakPtr<SSearchableComboBox>	FunctionOptionComboBox;
+	TArray< TSharedPtr< FString > >	FunctionOptionList;
+
+	TSharedRef<SWidget> MakeFunctionOptionComboWidget(TSharedPtr<FString> InItem);
+
+	/** Callback for when the function selection has changed from the dropdown */
+	void OnFunctionSelected(TSharedPtr<FString> FunctionItemData, ESelectInfo::Type SelectInfo);
+	//~ End Searchable combo box options
+
 	/**
 	* Adds a FunctionItemData with a given description to the array of FunctionDataItems. 
 	* 
 	* @param	DescriptionName		Description of the option to give the user
 	* @return	Shared pointer to the FunctionItemData
 	*/
-	TSharedPtr<SGraphNodeK2CreateDelegate::FFunctionItemData> AddDefaultFunctionDataOption(const FText& DescriptionName);
-};
-
+	TSharedPtr<FString> AddDefaultFunctionDataOption(const FText& DescriptionName);
+};
\ No newline at end of file
