Fix compile error
#rnx

diff --git a/Engine/Source/Editor/GraphEditor/Private/KismetNodes/SGraphNodeK2CreateDelegate.cpp b/Engine/Source/Editor/GraphEditor/Private/KismetNodes/SGraphNodeK2CreateDelegate.cpp
--- a/Engine/Source/Editor/GraphEditor/Private/KismetNodes/SGraphNodeK2CreateDelegate.cpp
+++ b/Engine/Source/Editor/GraphEditor/Private/KismetNodes/SGraphNodeK2CreateDelegate.cpp
@@ -286,8 +286,8 @@ void SGraphNodeK2CreateDelegate::CreateBelowPinControls(TSharedPtr<SVerticalBox>
 					UEdGraphSchema_K2::FunctionCanBeUsedInDelegate(Func))
 				{
 					FFunctionItemData ItemData;
-					ItemData->Name = Func->GetFName();
-					ItemData->Description = FunctionDescription(Func);
+					ItemData.Name = Func->GetFName();
+					ItemData.Description = FunctionDescription(Func);
 					ClassFunctions.Emplace(MoveTemp(ItemData));
 				}
 			}
