Fix another compile error
#rnx

diff --git a/Engine/Source/Editor/GraphEditor/Private/KismetNodes/SGraphNodeK2CreateDelegate.cpp b/Engine/Source/Editor/GraphEditor/Private/KismetNodes/SGraphNodeK2CreateDelegate.cpp
--- a/Engine/Source/Editor/GraphEditor/Private/KismetNodes/SGraphNodeK2CreateDelegate.cpp
+++ b/Engine/Source/Editor/GraphEditor/Private/KismetNodes/SGraphNodeK2CreateDelegate.cpp
@@ -299,7 +299,7 @@ void SGraphNodeK2CreateDelegate::CreateBelowPinControls(TSharedPtr<SVerticalBox>
 			for (const FFunctionItemData& ItemData : ClassFunctions)
 			{
 				// Add this to the searchable text box as an FString so users can type and find it
-				FunctionOptionList.Add(MakeShareable(new FString(ItemData.Name().ToString())));
+				FunctionOptionList.Add(MakeShareable(new FString(ItemData.Name.ToString())));
 			}
 
 			TSharedRef<SSearchableComboBox> SelectFunctionWidgetRef =
