Make the conteext search prefer options that have been favorited by the user

diff --git a/Engine/Source/Editor/GraphEditor/Private/SGraphActionMenu.cpp b/Engine/Source/Editor/GraphEditor/Private/SGraphActionMenu.cpp
--- a/Engine/Source/Editor/GraphEditor/Private/SGraphActionMenu.cpp
+++ b/Engine/Source/Editor/GraphEditor/Private/SGraphActionMenu.cpp
@@ -15,6 +15,8 @@
 #include "Widgets/SToolTip.h"
 #include "IDocumentation.h"
 #include "EditorCategoryUtils.h"
+#include "Editor/EditorPerProjectUserSettings.h"
+#include "BlueprintPaletteFavorites.h"
 #include "Widgets/Input/SEditableTextBox.h"
 #include "Widgets/Input/SSearchBox.h"
 #include "Widgets/Text/SInlineEditableTextBlock.h"
@@ -94,6 +96,13 @@ namespace ContextMenuConsoleVariables
 		TEXT("The multiplier given if the keyword only contains a letter the user typed in"),
 		ECVF_Default);
 
+	/** The bonus given if node is a favorite */
+	static float FavoriteBonus = 1000.0f;
+	static FAutoConsoleVariableRef CVarWordContainsLetterFavoriteBonus(
+		TEXT("ContextMenu.FavoriteBonus"), FavoriteBonus,
+		TEXT("The bonus given if node is a favorite"),
+		ECVF_Default);
+
 	/** Enabling the debug printing of context menu selections */
 	static bool bPrintDebugContextSelection = false;
 	static FAutoConsoleVariableRef CVarPrintDebugContextSelection(
@@ -876,8 +885,8 @@ void SGraphActionMenu::GenerateFilteredItems(bool bPreserveExpansion)
 void FContextMenuWeightDebugInfo::Print()
 {
 	UE_LOG(LogTemp, Warning, TEXT("[Weight Debug info] \
-		TotalWeight: %-8.2f | PercentageMatchWeight: %-8.2f | PercMatch: %-8.2f | ShorterWeight: %-8.2f | CategoryBonusWeight: %-8.2f | KeywordArrayWeight: %-8.2f | DescriptionWeight: %-8.2f | NodeTitleWeight: %-8.2f | CategoryWeight: %-8.2f\n"), 
-		TotalWeight, PercentageMatchWeight, PercMatch, ShorterWeight, CategoryBonusWieight, KeywordArrayWeight, DescriptionWeight, NodeTitleWeight, CategoryWeight);
+		TotalWeight: %-8.2f | PercentageMatchWeight: %-8.2f | PercMatch: %-8.2f | ShorterWeight: %-8.2f | CategoryBonusWeight: %-8.2f | KeywordArrayWeight: %-8.2f | DescriptionWeight: %-8.2f | NodeTitleWeight: %-8.2f | CategoryWeight: %-8.2f | Fav. Bonus:%-8.2f\n"), 
+		TotalWeight, PercentageMatchWeight, PercMatch, ShorterWeight, CategoryBonusWieight, KeywordArrayWeight, DescriptionWeight, NodeTitleWeight, CategoryWeight, FavoriteBonusWeight);
 }
 
 float SGraphActionMenu::GetActionFilteredWeight(const FGraphActionListBuilderBase::ActionGroup& InCurrentAction, const TArray<FString>& InFilterTerms, const TArray<FString>& InSanitizedFilterTerms, FContextMenuWeightDebugInfo& OutDebugInfo)
@@ -905,6 +914,8 @@ float SGraphActionMenu::GetActionFilteredWeight(const FGraphActionListBuilderBas
 	int32 Action = 0;
 	if( InCurrentAction.Actions[Action].IsValid() == true )
 	{
+		TSharedPtr<FEdGraphSchemaAction> CurrentAction = InCurrentAction.Actions[Action];
+		
 		// Combine the actions string, separate with \n so terms don't run into each other, and remove the spaces (incase the user is searching for a variable)
 		// In the case of groups containing multiple actions, they will have been created and added at the same place in the code, using the same description
 		// and keywords, so we only need to use the first one for filtering.
@@ -979,6 +990,17 @@ float SGraphActionMenu::GetActionFilteredWeight(const FGraphActionListBuilderBas
 			}
 		}
 
+		// If the user has favorite this action, then give it a hefty bonus
+		const UEditorPerProjectUserSettings& EditorSettings = *GetDefault<UEditorPerProjectUserSettings>();
+		if (UBlueprintPaletteFavorites* BlueprintFavorites = EditorSettings.BlueprintFavorites)
+		{
+			if (BlueprintFavorites->IsFavorited(CurrentAction))
+			{
+				TotalWeight += ContextMenuConsoleVariables::FavoriteBonus;
+				OutDebugInfo.FavoriteBonusWeight += ContextMenuConsoleVariables::FavoriteBonus;
+			}
+		}
+		
 		// Now iterate through all the filter terms and calculate a 'weight' using the values and multipliers
 		const FString* EachTerm = nullptr;
 		const FString* EachTermSanitized = nullptr;
diff --git a/Engine/Source/Editor/GraphEditor/Public/SGraphActionMenu.h b/Engine/Source/Editor/GraphEditor/Public/SGraphActionMenu.h
--- a/Engine/Source/Editor/GraphEditor/Public/SGraphActionMenu.h
+++ b/Engine/Source/Editor/GraphEditor/Public/SGraphActionMenu.h
@@ -91,6 +91,7 @@ struct FContextMenuWeightDebugInfo
 	float DescriptionWeight = 0.0f;
 	float NodeTitleWeight = 0.0f;
 	float CategoryWeight = 0.0f;
+	float FavoriteBonusWeight = 0.0f;
 
 	float PercentageMatchWeight = 0.0f;
 	/** The calculated match percentage */
