Fixed graph pin default values being editable even when set to read-only

diff --git a/Engine/Source/Editor/GraphEditor/Private/KismetPins/SGraphPinBool.cpp b/Engine/Source/Editor/GraphEditor/Private/KismetPins/SGraphPinBool.cpp
--- a/Engine/Source/Editor/GraphEditor/Private/KismetPins/SGraphPinBool.cpp
+++ b/Engine/Source/Editor/GraphEditor/Private/KismetPins/SGraphPinBool.cpp
@@ -15,6 +15,7 @@ TSharedRef<SWidget>	SGraphPinBool::GetDefaultValueWidget()
 	return SNew(SCheckBox)
 		.Style(FEditorStyle::Get(), "Graph.Checkbox")
 		.IsChecked(this, &SGraphPinBool::IsDefaultValueChecked)
+		.IsEnabled(this, &SGraphPin::GetDefaultValueIsEditable)
 		.OnCheckStateChanged(this, &SGraphPinBool::OnDefaultValueCheckBoxChanged)
 		.Visibility( this, &SGraphPin::GetDefaultValueVisibility );
 }
diff --git a/Engine/Source/Editor/GraphEditor/Private/KismetPins/SGraphPinCollisionProfile.cpp b/Engine/Source/Editor/GraphEditor/Private/KismetPins/SGraphPinCollisionProfile.cpp
--- a/Engine/Source/Editor/GraphEditor/Private/KismetPins/SGraphPinCollisionProfile.cpp
+++ b/Engine/Source/Editor/GraphEditor/Private/KismetPins/SGraphPinCollisionProfile.cpp
@@ -24,6 +24,7 @@ TSharedRef<SWidget>	SGraphPinCollisionProfile::GetDefaultValueWidget()
 	return SAssignNew(NameComboBox, SNameComboBox)
 		.ContentPadding(FMargin(6.0f, 2.0f))
 		.OptionsSource(&NameList)
+		.IsEnabled(this, &SGraphPin::GetDefaultValueIsEditable)
 		.InitiallySelectedItem(InitialSelectedName)
 		.OnSelectionChanged(this, &SGraphPinCollisionProfile::OnSelectionChanged)
 		.OnComboBoxOpening(this, &SGraphPinCollisionProfile::OnComboBoxOpening)
diff --git a/Engine/Source/Editor/GraphEditor/Private/KismetPins/SGraphPinColor.cpp b/Engine/Source/Editor/GraphEditor/Private/KismetPins/SGraphPinColor.cpp
--- a/Engine/Source/Editor/GraphEditor/Private/KismetPins/SGraphPinColor.cpp
+++ b/Engine/Source/Editor/GraphEditor/Private/KismetPins/SGraphPinColor.cpp
@@ -20,6 +20,7 @@ TSharedRef<SWidget>	SGraphPinColor::GetDefaultValueWidget()
 	return SAssignNew(DefaultValueWidget, SBorder)
 		.BorderImage( FEditorStyle::GetBrush("FilledBorder") )
 		.Visibility( this, &SGraphPin::GetDefaultValueVisibility )
+		.IsEnabled(this, &SGraphPin::GetDefaultValueIsEditable)
 		.Padding(1)
 		[
 			SNew(SHorizontalBox)
diff --git a/Engine/Source/Editor/GraphEditor/Private/KismetPins/SGraphPinEnum.cpp b/Engine/Source/Editor/GraphEditor/Private/KismetPins/SGraphPinEnum.cpp
--- a/Engine/Source/Editor/GraphEditor/Private/KismetPins/SGraphPinEnum.cpp
+++ b/Engine/Source/Editor/GraphEditor/Private/KismetPins/SGraphPinEnum.cpp
@@ -100,6 +100,7 @@ TSharedRef<SWidget>	SGraphPinEnum::GetDefaultValueWidget()
 		.ComboItemList( ComboItems )
 		.VisibleText( this, &SGraphPinEnum::OnGetText )
 		.OnSelectionChanged( this, &SGraphPinEnum::ComboBoxSelectionChanged )
+		.IsEnabled(this, &SGraphPin::GetDefaultValueIsEditable)
 		.Visibility( this, &SGraphPin::GetDefaultValueVisibility )
 		.OnGetDisplayName(this, &SGraphPinEnum::OnGetFriendlyName)
 		.OnGetTooltip(this, &SGraphPinEnum::OnGetTooltip);
diff --git a/Engine/Source/Editor/GraphEditor/Private/KismetPins/SGraphPinKey.cpp b/Engine/Source/Editor/GraphEditor/Private/KismetPins/SGraphPinKey.cpp
--- a/Engine/Source/Editor/GraphEditor/Private/KismetPins/SGraphPinKey.cpp
+++ b/Engine/Source/Editor/GraphEditor/Private/KismetPins/SGraphPinKey.cpp
@@ -41,6 +41,7 @@ TSharedRef<SWidget>	SGraphPinKey::GetDefaultValueWidget()
 	return SNew(SKeySelector)
 		.Visibility(this, &SGraphPin::GetDefaultValueVisibility)
 		.CurrentKey(this, &SGraphPinKey::GetCurrentKey)
+		.IsEnabled(this, &SGraphPin::GetDefaultValueIsEditable)
 		.OnKeyChanged(this, &SGraphPinKey::OnKeyChanged);
 }
 
diff --git a/Engine/Source/Editor/GraphEditor/Private/KismetPins/SGraphPinText.cpp b/Engine/Source/Editor/GraphEditor/Private/KismetPins/SGraphPinText.cpp
--- a/Engine/Source/Editor/GraphEditor/Private/KismetPins/SGraphPinText.cpp
+++ b/Engine/Source/Editor/GraphEditor/Private/KismetPins/SGraphPinText.cpp
@@ -116,6 +116,7 @@ TSharedRef<SWidget>	SGraphPinText::GetDefaultValueWidget()
 		.Style(FEditorStyle::Get(), "Graph.EditableTextBox")
 		.Visibility(this, &SGraphPin::GetDefaultValueVisibility)
 		.ForegroundColor(FSlateColor::UseForeground())
+		.IsEnabled(this, &SGraphPin::GetDefaultValueIsEditable)
 		.WrapTextAt(400)
 		.MinDesiredWidth(18.0f)
 		.MaxDesiredHeight(200);
diff --git a/Engine/Source/Editor/GraphEditor/Private/KismetPins/SGraphPinVector.cpp b/Engine/Source/Editor/GraphEditor/Private/KismetPins/SGraphPinVector.cpp
--- a/Engine/Source/Editor/GraphEditor/Private/KismetPins/SGraphPinVector.cpp
+++ b/Engine/Source/Editor/GraphEditor/Private/KismetPins/SGraphPinVector.cpp
@@ -162,6 +162,7 @@ TSharedRef<SWidget>	SGraphPinVector::GetDefaultValueWidget()
 			.VisibleText_1( this, &SGraphPinVector::GetCurrentValue_1 )
 			.VisibleText_2( this, &SGraphPinVector::GetCurrentValue_2 )
 			.Visibility( this, &SGraphPin::GetDefaultValueVisibility )
+			.IsEnabled(this, &SGraphPin::GetDefaultValueIsEditable )
 			.OnFloatCommitted_Box_0( this, &SGraphPinVector::OnChangedValueTextBox_0 )
 			.OnFloatCommitted_Box_1( this, &SGraphPinVector::OnChangedValueTextBox_1 )
 			.OnFloatCommitted_Box_2( this, &SGraphPinVector::OnChangedValueTextBox_2 );
diff --git a/Engine/Source/Editor/GraphEditor/Private/KismetPins/SGraphPinVector2D.cpp b/Engine/Source/Editor/GraphEditor/Private/KismetPins/SGraphPinVector2D.cpp
--- a/Engine/Source/Editor/GraphEditor/Private/KismetPins/SGraphPinVector2D.cpp
+++ b/Engine/Source/Editor/GraphEditor/Private/KismetPins/SGraphPinVector2D.cpp
@@ -118,6 +118,7 @@ TSharedRef<SWidget>	SGraphPinVector2D::GetDefaultValueWidget()
 			.VisibleText_X( this, &SGraphPinVector2D::GetCurrentValue_X )
 			.VisibleText_Y( this, &SGraphPinVector2D::GetCurrentValue_Y )
 			.Visibility( this, &SGraphPin::GetDefaultValueVisibility )
+			.IsEnabled(this, &SGraphPin::GetDefaultValueIsEditable )
 			.OnFloatCommitted_Box_X( this, &SGraphPinVector2D::OnChangedValueTextBox_X )
 			.OnFloatCommitted_Box_Y( this, &SGraphPinVector2D::OnChangedValueTextBox_Y );
 }
diff --git a/Engine/Source/Editor/GraphEditor/Private/NiagaraPins/SGraphPinVector4.cpp b/Engine/Source/Editor/GraphEditor/Private/NiagaraPins/SGraphPinVector4.cpp
--- a/Engine/Source/Editor/GraphEditor/Private/NiagaraPins/SGraphPinVector4.cpp
+++ b/Engine/Source/Editor/GraphEditor/Private/NiagaraPins/SGraphPinVector4.cpp
@@ -191,6 +191,7 @@ TSharedRef<SWidget>	SGraphPinVector4::GetDefaultValueWidget()
 		.VisibleText_2(this, &SGraphPinVector4::GetCurrentValue_2)
 		.VisibleText_3(this, &SGraphPinVector4::GetCurrentValue_3)
 		.Visibility(this, &SGraphPin::GetDefaultValueVisibility)
+		.IsEnabled(this, &SGraphPin::GetDefaultValueIsEditable)
 		.OnFloatCommitted_Box_0(this, &SGraphPinVector4::OnChangedValueTextBox_0)
 		.OnFloatCommitted_Box_1(this, &SGraphPinVector4::OnChangedValueTextBox_1)
 		.OnFloatCommitted_Box_2(this, &SGraphPinVector4::OnChangedValueTextBox_2)
diff --git a/Engine/Source/Editor/GraphEditor/Public/KismetPins/SGraphPinNum.h b/Engine/Source/Editor/GraphEditor/Public/KismetPins/SGraphPinNum.h
--- a/Engine/Source/Editor/GraphEditor/Public/KismetPins/SGraphPinNum.h
+++ b/Engine/Source/Editor/GraphEditor/Public/KismetPins/SGraphPinNum.h
@@ -35,17 +35,12 @@ protected:
 				.EditableTextBoxStyle(FEditorStyle::Get(), "Graph.EditableTextBox")
 				.BorderForegroundColor(FSlateColor::UseForeground())
 				.Visibility(this, &SGraphPinNum::GetDefaultValueVisibility)
-				.IsEnabled(this, &SGraphPinNum::GetDefaultValueIsEnabled)
+				.IsEnabled(this, &SGraphPinNum::GetDefaultValueIsEditable)
 				.Value(this, &SGraphPinNum::GetNumericValue)
 				.OnValueCommitted(this, &SGraphPinNum::SetNumericValue)
 			];
 	}
 
-	bool GetDefaultValueIsEnabled() const
-	{
-		return !GraphPinObj->bDefaultValueIsReadOnly;
-	}
-
 	TOptional<NumericType> GetNumericValue() const
 	{
 		NumericType Num = NumericType();
diff --git a/Engine/Source/Editor/GraphEditor/Public/SGraphPin.h b/Engine/Source/Editor/GraphEditor/Public/SGraphPin.h
--- a/Engine/Source/Editor/GraphEditor/Public/SGraphPin.h
+++ b/Engine/Source/Editor/GraphEditor/Public/SGraphPin.h
@@ -190,6 +190,9 @@ public:
 
 	/** True if pin can be edited */
 	bool IsEditingEnabled() const;
+
+	/** True if the pin's default value can be edited, false if it is read only */
+	bool GetDefaultValueIsEditable() const { return !GraphPinObj->bDefaultValueIsReadOnly; }
 	
 	/** 
 	  * Called when ed graph data is cleared, indicating this widget can no longer safely access GraphPinObj 
