GraphEditor - Fix an issue where a deferred zoom to fit would be lost due to deferred ui updates.  Fixes issues in niagara scratch pad graphs.

diff --git a/Engine/Source/Editor/GraphEditor/Private/SGraphPanel.cpp b/Engine/Source/Editor/GraphEditor/Private/SGraphPanel.cpp
--- a/Engine/Source/Editor/GraphEditor/Private/SGraphPanel.cpp
+++ b/Engine/Source/Editor/GraphEditor/Private/SGraphPanel.cpp
@@ -1554,6 +1554,9 @@ void SGraphPanel::Update()
 
 	// Invoke any delegate methods
 	OnUpdateGraphPanel.ExecuteIfBound();
+
+	// Clear the update pending flag to allow deferred zoom commands to run.
+	bVisualUpdatePending = false;
 }
 
 // Purges the existing visual representation (typically followed by an Update call in the next tick)
@@ -1564,6 +1567,9 @@ void SGraphPanel::PurgeVisualRepresentation()
 
 	// Clear all of the nodes and pins
 	RemoveAllNodes();
+
+	// Set a flag to know that an update is pending to prevent running pending commands like zoom to fit until widgets are generated.
+	bVisualUpdatePending = true;
 }
 
 bool SGraphPanel::IsNodeTitleVisible(const class UEdGraphNode* Node, bool bRequestRename)
diff --git a/Engine/Source/Editor/GraphEditor/Private/SNodePanel.cpp b/Engine/Source/Editor/GraphEditor/Private/SNodePanel.cpp
--- a/Engine/Source/Editor/GraphEditor/Private/SNodePanel.cpp
+++ b/Engine/Source/Editor/GraphEditor/Private/SNodePanel.cpp
@@ -327,6 +327,8 @@ void SNodePanel::Construct()
 	TotalGestureMagnify = 0.0f;
 
 	ScopedTransactionPtr.Reset();
+
+	bVisualUpdatePending = false;
 }
 
 FVector2D SNodePanel::ComputeEdgePanAmount(const FGeometry& MyGeometry, const FVector2D& TargetPosition)
@@ -474,7 +476,7 @@ void SNodePanel::Tick( const FGeometry& AllottedGeometry, const double InCurrent
 	}
 
 	// Zoom to node extents
-	if( bDeferredZoomToNodeExtents )
+	if( bDeferredZoomToNodeExtents && bVisualUpdatePending == false )
 	{
 		bDeferredZoomToNodeExtents = false;
 		ZoomPadding = NodePanelDefs::DefaultZoomPadding;
diff --git a/Engine/Source/Editor/GraphEditor/Public/SNodePanel.h b/Engine/Source/Editor/GraphEditor/Public/SNodePanel.h
--- a/Engine/Source/Editor/GraphEditor/Public/SNodePanel.h
+++ b/Engine/Source/Editor/GraphEditor/Public/SNodePanel.h
@@ -956,6 +956,10 @@ protected:
 	/** Cached geometry for use within the active timer */
 	FGeometry CachedGeometry;
 
+	/** A flag to detect when a visual update is pending to prevent deferred
+	    commands like zoom to fit from running when there are no widgets */
+	bool bVisualUpdatePending;
+
 private:
 	/** Active timer that handles deferred zooming until the target zoom is reached */
 	EActiveTimerReturnType HandleZoomToFit(double InCurrentTime, float InDeltaTime);
