Hide some internal Python types from the class picker

#rnx

diff --git a/Engine/Plugins/Experimental/PythonScriptPlugin/Source/PythonScriptPlugin/Private/PyWrapperDelegate.h b/Engine/Plugins/Experimental/PythonScriptPlugin/Source/PythonScriptPlugin/Private/PyWrapperDelegate.h
--- a/Engine/Plugins/Experimental/PythonScriptPlugin/Source/PythonScriptPlugin/Private/PyWrapperDelegate.h
+++ b/Engine/Plugins/Experimental/PythonScriptPlugin/Source/PythonScriptPlugin/Private/PyWrapperDelegate.h
@@ -12,7 +12,7 @@
  * UObject proxy base used to wrap a callable Python object so that it can be used with an Unreal delegate
  * @note This can't go inside the WITH_PYTHON block due to UHT parsing limitations (it doesn't understand that macro)
  */
-UCLASS()
+UCLASS(HideDropDown)
 class UPythonCallableForDelegate : public UObject, public IPythonResourceOwner
 {
 	GENERATED_BODY()
diff --git a/Engine/Plugins/Experimental/PythonScriptPlugin/Source/PythonScriptPlugin/Private/PyWrapperTypeRegistry.cpp b/Engine/Plugins/Experimental/PythonScriptPlugin/Source/PythonScriptPlugin/Private/PyWrapperTypeRegistry.cpp
--- a/Engine/Plugins/Experimental/PythonScriptPlugin/Source/PythonScriptPlugin/Private/PyWrapperTypeRegistry.cpp
+++ b/Engine/Plugins/Experimental/PythonScriptPlugin/Source/PythonScriptPlugin/Private/PyWrapperTypeRegistry.cpp
@@ -2026,6 +2026,7 @@ PyTypeObject* FPyWrapperTypeRegistry::GenerateWrappedDelegateType(const UFunctio
 		PythonCallableForDelegateFunc->StaticLink(true);
 		PythonCallableForDelegateClass->AddFunctionToFunctionMap(PythonCallableForDelegateFunc, PythonCallableForDelegateFunc->GetFName());
 		PythonCallableForDelegateClass->SetSuperStruct(UPythonCallableForDelegate::StaticClass());
+		PythonCallableForDelegateClass->ClassFlags |= CLASS_HideDropDown;
 		PythonCallableForDelegateClass->Bind();
 		PythonCallableForDelegateClass->StaticLink(true);
 		PythonCallableForDelegateClass->AssembleReferenceTokenStream();
diff --git a/Engine/Source/Editor/GraphEditor/Private/KismetPins/SGraphPinClass.cpp b/Engine/Source/Editor/GraphEditor/Private/KismetPins/SGraphPinClass.cpp
--- a/Engine/Source/Editor/GraphEditor/Private/KismetPins/SGraphPinClass.cpp
+++ b/Engine/Source/Editor/GraphEditor/Private/KismetPins/SGraphPinClass.cpp
@@ -67,6 +67,7 @@ public:
 		
 			// Don't allow classes from a loaded map (e.g. LSBPs) unless we're already working inside that package context. Otherwise, choosing the class would lead to a GLEO at save time.
 			Result &= !ClassPackage->ContainsMap() || ClassPackage == GraphPinOutermostPackage;
+			Result &= !InClass->HasAnyClassFlags(CLASS_Hidden | CLASS_HideDropDown | CLASS_Deprecated);
 			Result &= bAllowAbstractClasses || !InClass->HasAnyClassFlags(CLASS_Abstract);
 		}
 
@@ -75,7 +76,9 @@ public:
 
 	virtual bool IsUnloadedClassAllowed(const FClassViewerInitializationOptions& InInitOptions, const TSharedRef< const IUnloadedBlueprintData > InUnloadedClassData, TSharedRef< FClassViewerFilterFuncs > InFilterFuncs) override
 	{
-		return (InFilterFuncs->IfInChildOfClassesSet( AllowedChildrenOfClasses, InUnloadedClassData) != EFilterReturn::Failed) && (bAllowAbstractClasses || !InUnloadedClassData->HasAnyClassFlags(CLASS_Abstract));
+		return (InFilterFuncs->IfInChildOfClassesSet( AllowedChildrenOfClasses, InUnloadedClassData) != EFilterReturn::Failed) 
+			&& (!InUnloadedClassData->HasAnyClassFlags(CLASS_Hidden | CLASS_HideDropDown | CLASS_Deprecated))
+			&& (bAllowAbstractClasses || !InUnloadedClassData->HasAnyClassFlags(CLASS_Abstract));
 	}
 };
 
