Remove meaningless check for deprecated classes, that's handled elsewhere
#rnx

diff --git a/Engine/Source/Editor/GraphEditor/Private/KismetPins/SGraphPinClass.cpp b/Engine/Source/Editor/GraphEditor/Private/KismetPins/SGraphPinClass.cpp
--- a/Engine/Source/Editor/GraphEditor/Private/KismetPins/SGraphPinClass.cpp
+++ b/Engine/Source/Editor/GraphEditor/Private/KismetPins/SGraphPinClass.cpp
@@ -67,7 +67,7 @@ public:
 		
 			// Don't allow classes from a loaded map (e.g. LSBPs) unless we're already working inside that package context. Otherwise, choosing the class would lead to a GLEO at save time.
 			Result &= !ClassPackage->ContainsMap() || ClassPackage == GraphPinOutermostPackage;
-			Result &= !InClass->HasAnyClassFlags(CLASS_Hidden | CLASS_HideDropDown | CLASS_Deprecated);
+			Result &= !InClass->HasAnyClassFlags(CLASS_Hidden | CLASS_HideDropDown);
 			Result &= bAllowAbstractClasses || !InClass->HasAnyClassFlags(CLASS_Abstract);
 		}
 
@@ -77,7 +77,7 @@ public:
 	virtual bool IsUnloadedClassAllowed(const FClassViewerInitializationOptions& InInitOptions, const TSharedRef< const IUnloadedBlueprintData > InUnloadedClassData, TSharedRef< FClassViewerFilterFuncs > InFilterFuncs) override
 	{
 		return (InFilterFuncs->IfInChildOfClassesSet( AllowedChildrenOfClasses, InUnloadedClassData) != EFilterReturn::Failed) 
-			&& (!InUnloadedClassData->HasAnyClassFlags(CLASS_Hidden | CLASS_HideDropDown | CLASS_Deprecated))
+			&& (!InUnloadedClassData->HasAnyClassFlags(CLASS_Hidden | CLASS_HideDropDown))
 			&& (bAllowAbstractClasses || !InUnloadedClassData->HasAnyClassFlags(CLASS_Abstract));
 	}
 };
