Update context menu search calculation to prefer shorter things even more than before
Reduce keyword weight to improve Niagara search functionality

#rnx

diff --git a/Engine/Source/Editor/GraphEditor/Private/SGraphActionMenu.cpp b/Engine/Source/Editor/GraphEditor/Private/SGraphActionMenu.cpp
--- a/Engine/Source/Editor/GraphEditor/Private/SGraphActionMenu.cpp
+++ b/Engine/Source/Editor/GraphEditor/Private/SGraphActionMenu.cpp
@@ -18,6 +18,7 @@
 #include "Widgets/Input/SEditableTextBox.h"
 #include "Widgets/Input/SSearchBox.h"
 #include "Widgets/Text/SInlineEditableTextBlock.h"
+#include "Math/NumericLimits.h"
 
 #define LOCTEXT_NAMESPACE "GraphActionMenu"
 
@@ -26,18 +27,12 @@
 namespace ContextMenuConsoleVariables
 {
 	/** Increasing this weight will make shorter words preferred */
-	static float ShorterWeight = 2.0f;
+	static float ShorterWeight = 15.0f;
 	static FAutoConsoleVariableRef CVarShorterWeight(
 		TEXT("ContextMenu.ShorterWeight"), ShorterWeight,
 		TEXT("Increasing this weight will make shorter words preferred"),
 		ECVF_Default);
 
-	static int32 KeywordLengthDifferenceModifier = 50;
-	static FAutoConsoleVariableRef CVarKeywordLengthDifferenceModifier(
-		TEXT("ContextMenu.KeywordLengthDifferenceModifier"), KeywordLengthDifferenceModifier,
-		TEXT("Used to calculate how much weight should be given to the current keyword (KeywordLengthDifferenceModifier - KeywordLength) * ShorterWeight"),
-		ECVF_Default);
-
 	/** Increasing this will prefer whole percentage matches when comparing the keyword to what the user has typed in */
 	static float PercentageMatchWeightMultiplier = 1.0f;
 	static FAutoConsoleVariableRef CVarPercentageMatchWeightMultiplier(
@@ -67,14 +62,14 @@ namespace ContextMenuConsoleVariables
 		ECVF_Default);
 
 	/** How much weight the node's title has */
-	static float NodeTitleWeight = 50.0f;
+	static float NodeTitleWeight = 5.0f;
 	static FAutoConsoleVariableRef CVarNodeTitleWeight(
 		TEXT("ContextMenu.NodeTitleWeight"), NodeTitleWeight,
 		TEXT("The amount of weight placed on the search items title"),
 		ECVF_Default);
 
 	/** Weight used to prefer keywords of actions  */
-	static float KeywordWeight = 100.0f;
+	static float KeywordWeight = 30.0f;
 	static FAutoConsoleVariableRef CVarKeywordWeight(
 		TEXT("ContextMenu.KeywordWeight"), KeywordWeight,
 		TEXT("The amount of weight placed on search items keyword"),
@@ -766,7 +761,7 @@ void SGraphActionMenu::GenerateFilteredItems(bool bPreserveExpansion)
 	ensure( SanitizedFilterTerms.Num() == FilterTerms.Num() );// Both of these should match !
 
 	const bool bRequiresFiltering = FilterTerms.Num() > 0;
-	float BestMatchCount = 0.0f;
+	float BestMatchCount = TNumericLimits<float>::Lowest();
 	int32 BestMatchIndex = INDEX_NONE;
 
 	FContextMenuWeightDebugInfo BestMatchDebugInfo = {};
@@ -778,7 +773,7 @@ void SGraphActionMenu::GenerateFilteredItems(bool bPreserveExpansion)
 
 		// If we're filtering, search check to see if we need to show this action
 		bool bShowAction = true;
-		float EachWeight = 0.0f;
+		float EachWeight = TNumericLimits<float>::Lowest();
 		CurActionDebugInfo = {};
 		if (bRequiresFiltering)
 		{
@@ -884,7 +879,7 @@ float SGraphActionMenu::GetActionFilteredWeight(const FGraphActionListBuilderBas
 {	
 	// The overall 'weight' of this action 
 	float TotalWeight = 0.0f;
-	// Helper array
+	// Helper array for tracking actions and the weight that they have added
 	struct FArrayWithWeight
 	{
 		FArrayWithWeight(const TArray< FString >* InArray, float InWeightModifier, float* OutDebugWeight)
@@ -901,6 +896,8 @@ float SGraphActionMenu::GetActionFilteredWeight(const FGraphActionListBuilderBas
 
 	// Setup an array of arrays so we can do a weighted search			
 	TArray< FArrayWithWeight > WeightedArrayList;
+
+	const bool bIsFromDrag = DraggedFromPins.Num();
 	
 	int32 Action = 0;
 	if( InCurrentAction.Actions[Action].IsValid() == true )
@@ -934,47 +931,50 @@ float SGraphActionMenu::GetActionFilteredWeight(const FGraphActionListBuilderBas
 		// The category
 		WeightedArrayList.Add(FArrayWithWeight(&InCurrentAction.GetSearchCategoryArrayForFirstAction(), ContextMenuConsoleVariables::CategoryWeight,  &OutDebugInfo.CategoryWeight));
 
-		// If this actions category matches the one from a dragged off pin, give it a weight bonus
-		const TArray<FString>& InActionCategories = InCurrentAction.GetCategoryChain();
-		bool bAddMatchBonus = false;
-
-		for (const FString& InActionCategory : InActionCategories)
+		// Give a weight bonus to actions whose category matches what was dragged off of
+		if(bIsFromDrag)
 		{
-			for (UEdGraphPin* const FromPin : DraggedFromPins)
+			const TArray<FString>& InActionCategories = InCurrentAction.GetCategoryChain();
+			bool bAddMatchBonus = false;
+
+			for (const FString& InActionCategory : InActionCategories)
 			{
-				check(FromPin != nullptr);
-			
-				// If we can't find anything there, check the subcategory of the object
-				// This covers most of the more complex struct types (LinearColor, date time, etc)
-				if (UObject* const SubCatObj = FromPin->PinType.PinSubCategoryObject.Get())
+				for (UEdGraphPin* const FromPin : DraggedFromPins)
 				{
-					const FString& SubCatObjName = SubCatObj->GetFullName();
-					// The pin SubObjectCategory names don't have any spaces, so split up the category
-					TArray<FString> DelimitedArray;
-					InActionCategory.ParseIntoArray(DelimitedArray, TEXT(" "), true);
-					for (const FString& DelimetedCat : DelimitedArray)
+					check(FromPin != nullptr);
+
+					// If we can't find anything there, check the subcategory of the object
+					// This covers most of the more complex struct types (LinearColor, date time, etc)
+					if (UObject* const SubCatObj = FromPin->PinType.PinSubCategoryObject.Get())
 					{
-						if (SubCatObjName.Contains(DelimetedCat))
+						const FString& SubCatObjName = SubCatObj->GetFullName();
+						// The pin SubObjectCategory names don't have any spaces, so split up the category
+						TArray<FString> DelimitedArray;
+						InActionCategory.ParseIntoArray(DelimitedArray, TEXT(" "), true);
+						for (const FString& DelimetedCat : DelimitedArray)
 						{
-							bAddMatchBonus = true;
-							break;
+							if (SubCatObjName.Contains(DelimetedCat))
+							{
+								bAddMatchBonus = true;
+								break;
+							}
 						}
 					}
-				}
-				// Check the category of the pin, this works for basic math types (int, float, byte, etc)
-				else if (InActionCategory.Contains(FromPin->PinType.PinCategory.ToString()))
-				{
-					bAddMatchBonus = true;
-				}
+					// Check the category of the pin, this works for basic math types (int, float, byte, etc)
+					else if (InActionCategory.Contains(FromPin->PinType.PinCategory.ToString()))
+					{
+						bAddMatchBonus = true;
+					}
 
-				// If we found match in any cases above then add the weight bonus and stop looking
-				if (bAddMatchBonus)
-				{
-					TotalWeight += ContextMenuConsoleVariables::MatchingFromPinCategory;				
-					OutDebugInfo.CategoryBonusWieight += ContextMenuConsoleVariables::MatchingFromPinCategory;
-					
-					// Break out of the loop so that we don't give any extra bonuses
-					break;
+					// If we found match in any cases above then add the weight bonus and stop looking
+					if (bAddMatchBonus)
+					{
+						TotalWeight += ContextMenuConsoleVariables::MatchingFromPinCategory;
+						OutDebugInfo.CategoryBonusWieight += ContextMenuConsoleVariables::MatchingFromPinCategory;
+
+						// Break out of the loop so that we don't give any extra bonuses
+						break;
+					}
 				}
 			}
 		}
@@ -983,12 +983,13 @@ float SGraphActionMenu::GetActionFilteredWeight(const FGraphActionListBuilderBas
 		const FString* EachTerm = nullptr;
 		const FString* EachTermSanitized = nullptr;
 
-		// For every filter item the user has typed in (a letter)
+		// For every filter item the user has typed in (the text in the search bar, seperated by spaces)
 		for (int32 FilterIndex = 0; FilterIndex < InFilterTerms.Num(); ++FilterIndex)
 		{
-			EachTerm = &InFilterTerms[FilterIndex];
+			EachTerm = &InFilterTerms[FilterIndex];			
 			EachTermSanitized = &InSanitizedFilterTerms[FilterIndex];
-			
+			int32 TermLen = EachTerm->Len();
+
 			// Now check the weighted lists	(We could further improve the hit weight by checking consecutive word matches)
 			for (int32 iFindCount = 0; iFindCount < WeightedArrayList.Num() ; ++iFindCount)
 			{
@@ -1003,43 +1004,41 @@ float SGraphActionMenu::GetActionFilteredWeight(const FGraphActionListBuilderBas
 				int32 KeywordArrayCharLength = 0;
 
 				// Loop through every word that the user could be looking for
-				for (int32 iEachWord = 0; iEachWord < KeywordArray.Num() ; ++iEachWord)
+				for (int32 iEachWord = 0; iEachWord < KeywordArray.Num(); ++iEachWord)
 				{
 					// Keep track of how long all the words in the array are
 					KeywordArrayCharLength += KeywordArray[iEachWord].Len();
 
-					// If a word contains the letter that the user has typed in, than increment the whole match count
-					// If the word starts with the letter, give it a little extra boost of weight
-					if (KeywordArray[iEachWord].StartsWith(*EachTermSanitized, ESearchCase::CaseSensitive) || KeywordArray[iEachWord].StartsWith(*EachTerm, ESearchCase::CaseSensitive))
-					{
-						++WordMatchCount;
-						WeightPerList += KeywordArrayWeight * ContextMenuConsoleVariables::StartsWithBonusWeightMultiplier;
-					}
-					else if (KeywordArray[iEachWord].Contains(*EachTermSanitized, ESearchCase::CaseSensitive) || KeywordArray[iEachWord].Contains(*EachTerm, ESearchCase::CaseSensitive))
+					// If a word contains the letter that the user has typed in, than increment the whole match count					
+					if (KeywordArray[iEachWord].Contains(*EachTermSanitized, ESearchCase::CaseSensitive) || KeywordArray[iEachWord].Contains(*EachTerm, ESearchCase::CaseSensitive))
 					{
 						++WordMatchCount;
 						WeightPerList += KeywordArrayWeight * ContextMenuConsoleVariables::WordContainsLetterWeightMultiplier;
+						
+						// If the word starts with the letter, give it a little extra boost of weight
+						if (KeywordArray[iEachWord].StartsWith(*EachTermSanitized, ESearchCase::CaseSensitive) || KeywordArray[iEachWord].StartsWith(*EachTerm, ESearchCase::CaseSensitive))
+						{
+							WeightPerList += KeywordArrayWeight * ContextMenuConsoleVariables::StartsWithBonusWeightMultiplier;
+						}
 					}
 				}
 
+				// If the user has dragged off of a pin then do not prefer shorter things, because that will result
+				// in the matching of "Add" for a container instead of "+" for numeric types
 				if (KeywordArrayCharLength > 0)
 				{
-					// How many matches did we find / the total length of this keyword array
-					float PercMatch = ((float)WordMatchCount / (float)KeywordArrayCharLength);
+					// How many words that we are checking had partial matches compared to what the user typed in?
+					float PercMatch = (float)WordMatchCount / (float)KeywordArray.Num();
+
 					float PercentageBonus = (WeightPerList * PercMatch * ContextMenuConsoleVariables::PercentageMatchWeightMultiplier);
 					WeightPerList += PercentageBonus;
 
-					// Give a bonus for being shorter than a certain amount
-					int32 ShortPoints = ContextMenuConsoleVariables::KeywordLengthDifferenceModifier - KeywordArrayCharLength;
-					float ShortWeight = 0.0f;
-					if (ShortPoints > 0)
-					{
-						ShortWeight = (float)ShortPoints * ContextMenuConsoleVariables::ShorterWeight;
-						WeightPerList += ShortWeight;
-					}
+					// The longer the match is, the more points it loses
+					float ShortWeight = KeywordArrayCharLength * ContextMenuConsoleVariables::ShorterWeight * (bIsFromDrag ? 0.25f : 1.0f);
+					WeightPerList -= ShortWeight;
 
 					OutDebugInfo.PercMatch += PercMatch;
-					OutDebugInfo.ShorterWeight += ShortWeight;
+					OutDebugInfo.ShorterWeight -= ShortWeight;
 					OutDebugInfo.PercentageMatchWeight += PercentageBonus;					
 				}
 
diff --git a/Engine/Source/Runtime/Engine/Classes/Kismet/KismetMathLibrary.h b/Engine/Source/Runtime/Engine/Classes/Kismet/KismetMathLibrary.h
--- a/Engine/Source/Runtime/Engine/Classes/Kismet/KismetMathLibrary.h
+++ b/Engine/Source/Runtime/Engine/Classes/Kismet/KismetMathLibrary.h
@@ -268,11 +268,11 @@ class ENGINE_API UKismetMathLibrary : public UBlueprintFunctionLibrary
 	//
 
 	/** Multiplication (A * B) */
-	UFUNCTION(BlueprintPure, meta=(DisplayName = "integer * integer", CompactNodeTitle = "*", Keywords = "* multiply", CommutativeAssociativeBinaryOperator = "true"), Category="Math|Integer")
+	UFUNCTION(BlueprintPure, meta=(DisplayName = "int * int", CompactNodeTitle = "*", Keywords = "* multiply", CommutativeAssociativeBinaryOperator = "true"), Category="Math|Integer")
 	static int32 Multiply_IntInt(int32 A, int32 B);
 
 	/** Division (A / B) */
-	UFUNCTION(BlueprintPure, meta=(DisplayName = "integer / integer", CompactNodeTitle = "/", Keywords = "/ divide division"), Category="Math|Integer")
+	UFUNCTION(BlueprintPure, meta=(DisplayName = "int / int", CompactNodeTitle = "/", Keywords = "/ divide division"), Category="Math|Integer")
 	static int32 Divide_IntInt(int32 A, int32 B = 1);
 
 	/** Modulo (A % B) */
@@ -280,11 +280,11 @@ class ENGINE_API UKismetMathLibrary : public UBlueprintFunctionLibrary
 	static int32 Percent_IntInt(int32 A, int32 B = 1);
 
 	/** Addition (A + B) */
-	UFUNCTION(BlueprintPure, meta=(DisplayName = "integer + integer", CompactNodeTitle = "+", Keywords = "+ add plus", CommutativeAssociativeBinaryOperator = "true"), Category="Math|Integer")
+	UFUNCTION(BlueprintPure, meta=(DisplayName = "int + int", CompactNodeTitle = "+", Keywords = "+ add plus", CommutativeAssociativeBinaryOperator = "true"), Category="Math|Integer")
 	static int32 Add_IntInt(int32 A, int32 B = 1);
 
 	/** Subtraction (A - B) */
-	UFUNCTION(BlueprintPure, meta=(DisplayName = "integer - integer", CompactNodeTitle = "-", Keywords = "- subtract minus"), Category="Math|Integer")
+	UFUNCTION(BlueprintPure, meta=(DisplayName = "int - int", CompactNodeTitle = "-", Keywords = "- subtract minus"), Category="Math|Integer")
 	static int32 Subtract_IntInt(int32 A, int32 B = 1);
 
 	/** Returns true if A is less than B (A < B) */
@@ -1300,7 +1300,7 @@ class ENGINE_API UKismetMathLibrary : public UBlueprintFunctionLibrary
 	static FVector Add_VectorFloat(FVector A, float B);
 	
 	/** Adds an integer to each component of a vector */
-	UFUNCTION(BlueprintPure, meta=(DisplayName = "vector + int", CompactNodeTitle = "+", ScriptMethod = "AddInt", Keywords = "+ add plus"), Category="Math|Vector")
+	UFUNCTION(BlueprintPure, meta=(DisplayName = "vector + integer", CompactNodeTitle = "+", ScriptMethod = "AddInt", Keywords = "+ add plus"), Category="Math|Vector")
 	static FVector Add_VectorInt(FVector A, int32 B);
 
 	/** Vector subtraction */
@@ -1312,7 +1312,7 @@ class ENGINE_API UKismetMathLibrary : public UBlueprintFunctionLibrary
 	static FVector Subtract_VectorFloat(FVector A, float B);
 
 	/** Subtracts an integer from each component of a vector */
-	UFUNCTION(BlueprintPure, meta=(DisplayName = "vector - int", CompactNodeTitle = "-", ScriptMethod = "SubtractInt", Keywords = "- subtract minus"), Category="Math|Vector")
+	UFUNCTION(BlueprintPure, meta=(DisplayName = "vector - integer", CompactNodeTitle = "-", ScriptMethod = "SubtractInt", Keywords = "- subtract minus"), Category="Math|Vector")
 	static FVector Subtract_VectorInt(FVector A, int32 B);
 
 	/** Element-wise Vector multiplication (Result = {A.x*B.x, A.y*B.y, A.z*B.z}) */
@@ -1324,7 +1324,7 @@ class ENGINE_API UKismetMathLibrary : public UBlueprintFunctionLibrary
 	static FVector Multiply_VectorFloat(FVector A, float B);
 	
 	/** Scales Vector A by B */
-	UFUNCTION(BlueprintPure, meta=(DisplayName = "vector * int", CompactNodeTitle = "*", ScriptMethod = "MultiplyInt", Keywords = "* multiply"), Category="Math|Vector")
+	UFUNCTION(BlueprintPure, meta=(DisplayName = "vector * integer", CompactNodeTitle = "*", ScriptMethod = "MultiplyInt", Keywords = "* multiply"), Category="Math|Vector")
 	static FVector Multiply_VectorInt(FVector A, int32 B);
 
 	/** Element-wise Vector division (Result = {A.x/B.x, A.y/B.y, A.z/B.z}) */
@@ -1336,7 +1336,7 @@ class ENGINE_API UKismetMathLibrary : public UBlueprintFunctionLibrary
 	static FVector Divide_VectorFloat(FVector A, float B = 1.f);
 
 	/** Vector divide by an integer */
-	UFUNCTION(BlueprintPure, meta=(DisplayName = "vector / int", CompactNodeTitle = "/", ScriptMethod = "DivideInt", Keywords = "/ divide division"), Category="Math|Vector")
+	UFUNCTION(BlueprintPure, meta=(DisplayName = "vector / integer", CompactNodeTitle = "/", ScriptMethod = "DivideInt", Keywords = "/ divide division"), Category="Math|Vector")
 	static FVector Divide_VectorInt(FVector A, int32 B = 1);
 
 	/** Negate a vector. */
