Fix CIS warning about implicit bool cast

#rnx

diff --git a/Engine/Source/Editor/GraphEditor/Private/SGraphActionMenu.cpp b/Engine/Source/Editor/GraphEditor/Private/SGraphActionMenu.cpp
--- a/Engine/Source/Editor/GraphEditor/Private/SGraphActionMenu.cpp
+++ b/Engine/Source/Editor/GraphEditor/Private/SGraphActionMenu.cpp
@@ -897,7 +897,7 @@ float SGraphActionMenu::GetActionFilteredWeight(const FGraphActionListBuilderBas
 	// Setup an array of arrays so we can do a weighted search			
 	TArray< FArrayWithWeight > WeightedArrayList;
 
-	const bool bIsFromDrag = DraggedFromPins.Num();
+	const bool bIsFromDrag = (DraggedFromPins.Num() > 0);
 	
 	int32 Action = 0;
 	if( InCurrentAction.Actions[Action].IsValid() == true )
