Edigrate CL 14268540 from UE4 Main

Fix crashes related to context menu for graph links.

diff --git a/Engine/Source/Editor/GraphEditor/Private/SGraphEditorImpl.cpp b/Engine/Source/Editor/GraphEditor/Private/SGraphEditorImpl.cpp
--- a/Engine/Source/Editor/GraphEditor/Private/SGraphEditorImpl.cpp
+++ b/Engine/Source/Editor/GraphEditor/Private/SGraphEditorImpl.cpp
@@ -1060,16 +1060,19 @@ void SGraphEditorImpl::RegisterContextMenu(const UEdGraphSchema* Schema, FToolMe
 	{
 		UToolMenu* Menu = ToolMenus->FindMenu(EdGraphSchemaContextMenuName);
 
-		Menu->AddDynamicSection("EdGraphSchemaPinActions", FNewToolMenuDelegate::CreateLambda([this](UToolMenu* InMenu)
+		Menu->AddDynamicSection("EdGraphSchemaPinActions", FNewToolMenuDelegate::CreateLambda([](UToolMenu* InMenu)
 			{
 				UGraphNodeContextMenuContext* NodeContext = InMenu->FindContext<UGraphNodeContextMenuContext>();
-				if (NodeContext && NodeContext->Pin)
+				if (NodeContext && NodeContext->Graph)
 				{
-					GetPinContextMenuActionsForSchema(InMenu);
+					if (TSharedPtr<SGraphEditorImpl> GraphEditor = StaticCastSharedPtr<SGraphEditorImpl>(FindGraphEditorForGraph(NodeContext->Graph)))
+					{
+						GraphEditor->GetPinContextMenuActionsForSchema(InMenu);
+					}
 				}
 			}));
 
-		Menu->AddDynamicSection("GetContextMenuActions", FNewToolMenuDelegate::CreateLambda([this](UToolMenu* InMenu)
+		Menu->AddDynamicSection("GetContextMenuActions", FNewToolMenuDelegate::CreateLambda([](UToolMenu* InMenu)
 		{
 			if (UGraphNodeContextMenuContext* ContextObject = InMenu->FindContext<UGraphNodeContextMenuContext>())
 			{
@@ -1087,7 +1090,7 @@ void SGraphEditorImpl::RegisterContextMenu(const UEdGraphSchema* Schema, FToolMe
 	{
 		UToolMenu* Menu = ToolMenus->FindMenu(GetNodeContextMenuName(Context->Node->GetClass()));
 		 
-		Menu->AddDynamicSection("GetNodeContextMenuActions", FNewToolMenuDelegate::CreateLambda([this](UToolMenu* InMenu)
+		Menu->AddDynamicSection("GetNodeContextMenuActions", FNewToolMenuDelegate::CreateLambda([](UToolMenu* InMenu)
 		{
 			UGraphNodeContextMenuContext* Context = InMenu->FindContext<UGraphNodeContextMenuContext>();
 			if (Context && Context->Node)
