Multiple entries per graph are visible again in pin context menu.

#rnx

diff --git a/Engine/Source/Editor/GraphEditor/Private/SGraphEditorImpl.cpp b/Engine/Source/Editor/GraphEditor/Private/SGraphEditorImpl.cpp
--- a/Engine/Source/Editor/GraphEditor/Private/SGraphEditorImpl.cpp
+++ b/Engine/Source/Editor/GraphEditor/Private/SGraphEditorImpl.cpp
@@ -205,12 +205,9 @@ void SGraphEditorImpl::GetPinContextMenuActionsForSchema(UToolMenu* InMenu) cons
 	auto GetMenuEntryForPin = [](const UEdGraphPin* TargetPin, const FToolUIAction& Action, const FText& SingleDescFormat, const FText& MultiDescFormat, TMap< FString, uint32 >& LinkTitleCount)
 	{
 		FText Title = TargetPin->GetOwningNode()->GetNodeTitle(ENodeTitleType::ListView);
-		FString TitleString = Title.ToString();
 		const FText PinDisplayName = TargetPin->GetDisplayName();
 		if (!PinDisplayName.IsEmpty())
 		{
-			TitleString = FString::Printf(TEXT("%s (%s)"), *TitleString, *PinDisplayName.ToString());
-
 			// Add name of connection if possible
 			FFormatNamedArguments Args;
 			Args.Add(TEXT("NodeTitle"), Title);
@@ -218,7 +215,7 @@ void SGraphEditorImpl::GetPinContextMenuActionsForSchema(UToolMenu* InMenu) cons
 			Title = FText::Format(LOCTEXT("JumpToDescPin", "{NodeTitle} ({PinName})"), Args);
 		}
 
-		uint32& Count = LinkTitleCount.FindOrAdd(TitleString);
+		uint32& Count = LinkTitleCount.FindOrAdd(Title.ToString());
 
 		FText Description;
 		FFormatNamedArguments Args;
@@ -236,7 +233,7 @@ void SGraphEditorImpl::GetPinContextMenuActionsForSchema(UToolMenu* InMenu) cons
 		++Count;
 
 		return FToolMenuEntry::InitMenuEntry(
-			TargetPin->GetOwningNode()->GetFName(),
+			NAME_None,
 			Description,
 			Description,
 			FSlateIcon(),
