PR #7288: UE-97445: Retrieve font size for material comment from correct location (Contributed by [REDACTED_PERSON])

diff --git a/Engine/Source/Editor/GraphEditor/Private/SGraphNodeComment.cpp b/Engine/Source/Editor/GraphEditor/Private/SGraphNodeComment.cpp
--- a/Engine/Source/Editor/GraphEditor/Private/SGraphNodeComment.cpp
+++ b/Engine/Source/Editor/GraphEditor/Private/SGraphNodeComment.cpp
@@ -81,7 +81,7 @@ void SGraphNodeComment::Tick( const FGeometry& AllottedGeometry, const double In
 		bCachedBubbleVisibility = CommentNode->bCommentBubbleVisible_InDetailsPanel;
 	}
 
-	if (CachedFontSize != CommentNode->FontSize)
+	if (CachedFontSize != CommentNode->GetFontSize())
 	{
 		UpdateGraphNode();
 	}
@@ -139,10 +139,11 @@ void SGraphNodeComment::UpdateGraphNode()
 	FGraphNodeMetaData TagMeta(TEXT("Graphnode"));
 	PopulateMetaTag(&TagMeta);
 
+	CachedFontSize = CommentNode->GetFontSize();
+
 	CommentStyle = FEditorStyle::Get().GetWidgetStyle<FInlineEditableTextBlockStyle>("Graph.CommentBlock.TitleInlineEditableText");
-	CommentStyle.EditableTextBoxStyle.Font.Size = CommentNode->FontSize;
-	CommentStyle.TextStyle.Font.Size = CommentNode->FontSize;
-	CachedFontSize = CommentNode->FontSize;
+	CommentStyle.EditableTextBoxStyle.Font.Size = CachedFontSize;
+	CommentStyle.TextStyle.Font.Size = CachedFontSize;
 
 	bool bIsSet = GraphNode->IsA(UEdGraphNode_Comment::StaticClass());
 	this->ContentScale.Bind( this, &SGraphNode::GetContentScale );
diff --git a/Engine/Source/Editor/UnrealEd/Classes/MaterialGraph/MaterialGraphNode_Comment.h b/Engine/Source/Editor/UnrealEd/Classes/MaterialGraph/MaterialGraphNode_Comment.h
--- a/Engine/Source/Editor/UnrealEd/Classes/MaterialGraph/MaterialGraphNode_Comment.h
+++ b/Engine/Source/Editor/UnrealEd/Classes/MaterialGraph/MaterialGraphNode_Comment.h
@@ -37,6 +37,7 @@ class UMaterialGraphNode_Comment : public UEdGraphNode_Comment
 	virtual void PostPlacedNewNode() override;
 	virtual void OnRenameNode(const FString& NewName) override;
 	virtual void ResizeNode(const FVector2D& NewSize) override;
+	virtual int32 GetFontSize() const override;
 	//~ End UEdGraphNode Interface.
 
 private:
diff --git a/Engine/Source/Editor/UnrealEd/Private/MaterialGraphNode_Comment.cpp b/Engine/Source/Editor/UnrealEd/Private/MaterialGraphNode_Comment.cpp
--- a/Engine/Source/Editor/UnrealEd/Private/MaterialGraphNode_Comment.cpp
+++ b/Engine/Source/Editor/UnrealEd/Private/MaterialGraphNode_Comment.cpp
@@ -164,6 +164,15 @@ void UMaterialGraphNode_Comment::ResizeNode(const FVector2D& NewSize)
 	MaterialDirtyDelegate.ExecuteIfBound();
 }
 
+int32 UMaterialGraphNode_Comment::GetFontSize() const
+{
+	if (MaterialExpressionComment)
+	{
+		return MaterialExpressionComment->FontSize;
+	}
+	return Super::GetFontSize();
+}
+
 void UMaterialGraphNode_Comment::ResetMaterialExpressionOwner()
 {
 	if (MaterialExpressionComment)
diff --git a/Engine/Source/Editor/UnrealEd/Public/EdGraphNode_Comment.h b/Engine/Source/Editor/UnrealEd/Public/EdGraphNode_Comment.h
--- a/Engine/Source/Editor/UnrealEd/Public/EdGraphNode_Comment.h
+++ b/Engine/Source/Editor/UnrealEd/Public/EdGraphNode_Comment.h
@@ -96,6 +96,9 @@ public:
 	/** Return the set of nodes underneath the comment */
 	UNREALED_API const FCommentNodeSet& GetNodesUnderComment() const;
 
+	/** Return the font size of the comment */
+	UNREALED_API virtual int32 GetFontSize() const { return FontSize; }
+
 private:
 	/** Nodes currently within the region of the comment */
 	FCommentNodeSet	NodesUnderComment;
