Fix graph editor comment bubbles appearing black on black text

diff --git a/Engine/Source/Editor/GraphEditor/Private/SCommentBubble.cpp b/Engine/Source/Editor/GraphEditor/Private/SCommentBubble.cpp
--- a/Engine/Source/Editor/GraphEditor/Private/SCommentBubble.cpp
+++ b/Engine/Source/Editor/GraphEditor/Private/SCommentBubble.cpp
@@ -257,16 +257,15 @@ void SCommentBubble::UpdateBubble()
 							SAssignNew(TextBlock, SMultiLineEditableTextBox)
 							.Text(MakeAttributeLambda([this] { return CachedCommentText; }))
 							.HintText( NSLOCTEXT( "CommentBubble", "EditCommentHint", "Click to edit" ))
-							.IsReadOnly( this, &SCommentBubble::IsReadOnly )
-							.Font( FEditorStyle::GetFontStyle( TEXT("Graph.Node.CommentFont")))
+							.IsReadOnly(this, &SCommentBubble::IsReadOnly)
 							.SelectAllTextWhenFocused( true )
 							.RevertTextOnEscape( true )
 							.ClearKeyboardFocusOnCommit( true )
-							.ModiferKeyForNewLine( EModifierKey::Shift )
-							.ForegroundColor( this, &SCommentBubble::GetTextForegroundColor )
-							.ReadOnlyForegroundColor( this, &SCommentBubble::GetTextForegroundColor )
-							.BackgroundColor( this, &SCommentBubble::GetTextBackgroundColor )
-							.OnTextCommitted( this, &SCommentBubble::OnCommentTextCommitted )
+							.ModiferKeyForNewLine(EModifierKey::Shift)
+							.ForegroundColor(this, &SCommentBubble::GetTextForegroundColor)
+							.ReadOnlyForegroundColor(this, &SCommentBubble::GetReadOnlyTextForegroundColor)
+							.BackgroundColor(this, &SCommentBubble::GetTextBackgroundColor)
+							.OnTextCommitted(this, &SCommentBubble::OnCommentTextCommitted)
 						]
 						+SHorizontalBox::Slot()
 						.AutoWidth()
@@ -410,7 +409,12 @@ FSlateColor SCommentBubble::GetTextBackgroundColor() const
 
 FSlateColor SCommentBubble::GetTextForegroundColor() const
 {
-	return TextBlock->HasKeyboardFocus() ? FLinearColor::Black : ForegroundColor;
+	return TextBlock->HasKeyboardFocus() ? FSlateColor::UseStyle() : FLinearColor::Black;
+}
+
+FSlateColor SCommentBubble::GetReadOnlyTextForegroundColor() const
+{
+	return TextBlock->HasKeyboardFocus() ? FLinearColor::White : FLinearColor::Black;
 }
 
 void SCommentBubble::OnCommentTextCommitted( const FText& NewText, ETextCommit::Type CommitInfo )
diff --git a/Engine/Source/Editor/GraphEditor/Public/SCommentBubble.h b/Engine/Source/Editor/GraphEditor/Public/SCommentBubble.h
--- a/Engine/Source/Editor/GraphEditor/Public/SCommentBubble.h
+++ b/Engine/Source/Editor/GraphEditor/Public/SCommentBubble.h
@@ -139,6 +139,9 @@ protected:
 	/** Returns the current foreground color for the textbox */
 	FSlateColor GetTextForegroundColor() const;
 
+	/** Returns the current foreground color for the textbox when it is read only*/
+	FSlateColor GetReadOnlyTextForegroundColor() const;
+
 	/** Returns the foreground color for the text and buttons, taking into account the bubble color */
 	FSlateColor GetForegroundColor() const override { return ForegroundColor; }
 
diff --git a/Engine/Source/Runtime/Slate/Private/Widgets/Input/SMultiLineEditableTextBox.cpp b/Engine/Source/Runtime/Slate/Private/Widgets/Input/SMultiLineEditableTextBox.cpp
--- a/Engine/Source/Runtime/Slate/Private/Widgets/Input/SMultiLineEditableTextBox.cpp
+++ b/Engine/Source/Runtime/Slate/Private/Widgets/Input/SMultiLineEditableTextBox.cpp
@@ -222,7 +222,7 @@ FSlateColor SMultiLineEditableTextBox::DetermineForegroundColor() const
 		{
 			Result = ReadOnlyForegroundColorOverride.Get();
 		}
-		if (ForegroundColorOverride.IsSet())
+		else if (ForegroundColorOverride.IsSet())
 		{
 			Result = ForegroundColorOverride.Get();
 		}
