Control Rig: Fixed node placement not saved when using any keyboard shortcuts

diff --git a/Engine/Plugins/Experimental/ControlRig/Source/ControlRigEditor/Private/Graph/SControlRigGraphNode.cpp b/Engine/Plugins/Experimental/ControlRig/Source/ControlRigEditor/Private/Graph/SControlRigGraphNode.cpp
--- a/Engine/Plugins/Experimental/ControlRig/Source/ControlRigEditor/Private/Graph/SControlRigGraphNode.cpp
+++ b/Engine/Plugins/Experimental/ControlRig/Source/ControlRigEditor/Private/Graph/SControlRigGraphNode.cpp
@@ -240,6 +240,20 @@ void SControlRigGraphNode::EndUserInteraction() const
 	SGraphNode::EndUserInteraction();
 }
 
+void SControlRigGraphNode::MoveTo(const FVector2D& NewPosition, FNodeSet& NodeFilter)
+{
+	if (!NodeFilter.Find(SharedThis(this)))
+	{
+		if (GraphNode && !RequiresSecondPassLayout())
+		{
+			if (const UControlRigGraphSchema* RigSchema = Cast<UControlRigGraphSchema>(GraphNode->GetSchema()))
+			{
+				RigSchema->SetNodePosition(GraphNode, NewPosition);
+			}
+		}
+	}
+}
+
 void SControlRigGraphNode::AddPin(const TSharedRef<SGraphPin>& PinToAdd) 
 {
 	// We show our own label
diff --git a/Engine/Plugins/Experimental/ControlRig/Source/ControlRigEditor/Private/Graph/SControlRigGraphNode.h b/Engine/Plugins/Experimental/ControlRig/Source/ControlRigEditor/Private/Graph/SControlRigGraphNode.h
--- a/Engine/Plugins/Experimental/ControlRig/Source/ControlRigEditor/Private/Graph/SControlRigGraphNode.h
+++ b/Engine/Plugins/Experimental/ControlRig/Source/ControlRigEditor/Private/Graph/SControlRigGraphNode.h
@@ -33,6 +33,7 @@ public:
 	virtual TSharedRef<SWidget> CreateTitleWidget(TSharedPtr<SNodeTitle> InNodeTitle) override;
 	virtual bool UseLowDetailNodeTitles() const override;
 	virtual void EndUserInteraction() const override;
+	virtual void MoveTo(const FVector2D& NewPosition, FNodeSet& NodeFilter) override;
 	virtual void AddPin( const TSharedRef<SGraphPin>& PinToAdd ) override;
 	virtual void SetDefaultTitleAreaWidget(TSharedRef<SOverlay> DefaultTitleAreaWidget) override
 	{
diff --git a/Engine/Source/Editor/GraphEditor/Private/SGraphEditorImpl.h b/Engine/Source/Editor/GraphEditor/Private/SGraphEditorImpl.h
--- a/Engine/Source/Editor/GraphEditor/Private/SGraphEditorImpl.h
+++ b/Engine/Source/Editor/GraphEditor/Private/SGraphEditorImpl.h
@@ -82,12 +82,33 @@ struct FAlignmentHelper
 	{
 		if (AlignmentData.Num() > 1)
 		{
-			float Target = DetermineAlignmentTarget();
+			UEdGraph* Graph = AlignmentData[0].Node->GetGraph();
 
-			for (FAlignmentData& Entry : AlignmentData)
+			if (Graph)
 			{
-				Entry.Node->Modify();
-				Entry.TargetProperty = Target - Entry.TargetOffset;
+				const UEdGraphSchema* Schema = Graph->GetSchema();
+
+				if (Schema)
+				{ 
+					float Target = DetermineAlignmentTarget();
+					
+					for (FAlignmentData& Entry : AlignmentData)
+					{
+						int32 TargetProperty = Target - Entry.TargetOffset; 
+						FVector2D TargetPosition(Entry.Node->NodePosX, Entry.Node->NodePosY);
+
+						if (Orientation == EOrientation::Orient_Horizontal)
+						{
+							TargetPosition.X = TargetProperty;
+						}
+						else
+						{ 
+							TargetPosition.Y = TargetProperty;
+						}
+
+						Schema->SetNodePosition(Entry.Node, TargetPosition);
+					}
+				}
 			}
 		}
 	}
