Control Rig: added missing includes in SGraphEditorImpl.h, which was changed as part of a Control Rig UI bug fix (UE-101840)

diff --git a/Engine/Source/Editor/GraphEditor/Private/SGraphEditorImpl.h b/Engine/Source/Editor/GraphEditor/Private/SGraphEditorImpl.h
--- a/Engine/Source/Editor/GraphEditor/Private/SGraphEditorImpl.h
+++ b/Engine/Source/Editor/GraphEditor/Private/SGraphEditorImpl.h
@@ -12,9 +12,12 @@
 #include "Widgets/SOverlay.h"
 #include "GraphEditor.h"
 #include "EdGraph/EdGraphPin.h" // for FEdGraphPinReference
+#include "EdGraph/EdGraph.h"
+#include "EdGraph/EdGraphSchema.h"
 
 class SGraphPanel;
 class UEdGraph;
+class UEdGraphSchema;
 struct FEdGraphEditAction;
 struct FGraphContextMenuArguments;
 struct FNotificationInfo;
