Added a Nudge Node Scoped Transaction when moving graph node with keyboard keys in the Blueprint or Material editor

diff --git a/Engine/Source/Editor/GraphEditor/Private/SGraphPanel.cpp b/Engine/Source/Editor/GraphEditor/Private/SGraphPanel.cpp
--- a/Engine/Source/Editor/GraphEditor/Private/SGraphPanel.cpp
+++ b/Engine/Source/Editor/GraphEditor/Private/SGraphPanel.cpp
@@ -557,6 +557,7 @@ TSharedPtr<IToolTip> SGraphPanel::GetToolTip()
 
 void SGraphPanel::UpdateSelectedNodesPositions(FVector2D PositionIncrement)
 {
+	FScopedTransaction Transaction(NSLOCTEXT("GraphEditor", "NudgeNodeAction", "Nudge Node"));
 	for (FGraphPanelSelectionSet::TIterator NodeIt(SelectionManager.SelectedNodes); NodeIt; ++NodeIt)
 	{
 		TSharedRef<SNode>* pWidget = NodeToWidgetLookup.Find(*NodeIt);
