Fix CIS missing include

diff --git a/Engine/Source/Editor/GraphEditor/Private/SGraphPanel.cpp b/Engine/Source/Editor/GraphEditor/Private/SGraphPanel.cpp
--- a/Engine/Source/Editor/GraphEditor/Private/SGraphPanel.cpp
+++ b/Engine/Source/Editor/GraphEditor/Private/SGraphPanel.cpp
@@ -32,6 +32,7 @@
 
 #include "Editor/UnrealEdEngine.h"
 #include "UnrealEdGlobals.h"
+#include "ScopedTransaction.h"
 
 DEFINE_LOG_CATEGORY_STATIC(LogGraphPanel, Log, All);
 
