Control Rig: Added Right Click Node Alignment Menu and Use Schema for Distribute Node positions, tested basic functionality, undo/redo, save/load, in both CR editor and BP Editor

diff --git a/Engine/Plugins/Experimental/ControlRig/Source/ControlRigEditor/Private/ControlRigEditorModule.cpp b/Engine/Plugins/Experimental/ControlRig/Source/ControlRigEditor/Private/ControlRigEditorModule.cpp
--- a/Engine/Plugins/Experimental/ControlRig/Source/ControlRigEditor/Private/ControlRigEditorModule.cpp
+++ b/Engine/Plugins/Experimental/ControlRig/Source/ControlRigEditor/Private/ControlRigEditorModule.cpp
@@ -1771,6 +1771,26 @@ void FControlRigEditorModule::GetContextMenuActions(const UControlRigGraphSchema
 						));
 					}
 
+					OrganizationSection.AddSubMenu("Alignment", LOCTEXT("AlignmentHeader", "Alignment"), FText(), FNewToolMenuDelegate::CreateLambda([](UToolMenu* AlignmentMenu)
+					{
+						{
+							FToolMenuSection& InSection = AlignmentMenu->AddSection("EdGraphSchemaAlignment", LOCTEXT("AlignHeader", "Align"));
+							InSection.AddMenuEntry(FGraphEditorCommands::Get().AlignNodesTop);
+							InSection.AddMenuEntry(FGraphEditorCommands::Get().AlignNodesMiddle);
+							InSection.AddMenuEntry(FGraphEditorCommands::Get().AlignNodesBottom);
+							InSection.AddMenuEntry(FGraphEditorCommands::Get().AlignNodesLeft);
+							InSection.AddMenuEntry(FGraphEditorCommands::Get().AlignNodesCenter);
+							InSection.AddMenuEntry(FGraphEditorCommands::Get().AlignNodesRight);
+							InSection.AddMenuEntry(FGraphEditorCommands::Get().StraightenConnections);
+						}
+
+						{
+							FToolMenuSection& InSection = AlignmentMenu->AddSection("EdGraphSchemaDistribution", LOCTEXT("DistributionHeader", "Distribution"));
+							InSection.AddMenuEntry(FGraphEditorCommands::Get().DistributeNodesHorizontally);
+							InSection.AddMenuEntry(FGraphEditorCommands::Get().DistributeNodesVertically);
+						}
+					}));
+
 					if (URigVMFunctionReferenceNode* FunctionRefNode = Cast<URigVMFunctionReferenceNode>(RigNode->GetModelNode()))
 					{
 						OrganizationSection.AddMenuEntry(
diff --git a/Engine/Source/Editor/GraphEditor/Private/SGraphEditorImpl.cpp b/Engine/Source/Editor/GraphEditor/Private/SGraphEditorImpl.cpp
--- a/Engine/Source/Editor/GraphEditor/Private/SGraphEditorImpl.cpp
+++ b/Engine/Source/Editor/GraphEditor/Private/SGraphEditorImpl.cpp
@@ -1530,7 +1530,7 @@ void SGraphEditorImpl::OnStraightenConnections()
 }
 
 /** Distribute the specified array of node data evenly */
-void DistributeNodes(TArray<FAlignmentData>& InData)
+void DistributeNodes(TArray<FAlignmentData>& InData, bool bIsHorizontal)
 {
 	// Sort the data
 	InData.Sort([](const FAlignmentData& A, const FAlignmentData& B) {
@@ -1550,14 +1550,50 @@ void DistributeNodes(TArray<FAlignmentData>& InData)
 	float TargetPosition = InData[0].GetTarget() + PaddingAmount;
 
 	// Now set all the properties on the target
-	for (int32 Index = 1; Index < InData.Num() - 1; ++Index)
+	if (InData.Num() > 1)
 	{
-		FAlignmentData& Entry = InData[Index];
+		UEdGraph* Graph = InData[0].Node->GetGraph();
+		if (Graph)
+		{
+			const UEdGraphSchema* Schema = Graph->GetSchema();
+
+			// similar to FAlignmentHelper::Align(), first try using GraphSchema to move the nodes if applicable
+			if (Schema)
+			{
+				for (int32 Index = 1; Index < InData.Num() - 1; ++Index)
+				{
+					FAlignmentData& Entry = InData[Index];
+
+					FVector2D Target2DPosition(Entry.Node->NodePosX, Entry.Node->NodePosY);
+ 
+					if (bIsHorizontal)
+					{
+						Target2DPosition.X = TargetPosition;
+					}
+					else
+					{ 
+						Target2DPosition.Y = TargetPosition;
+					}
 
-		Entry.Node->Modify();
-		Entry.TargetProperty = TargetPosition;
+					Schema->SetNodePosition(Entry.Node, Target2DPosition);
 
-		TargetPosition = Entry.GetTarget() + PaddingAmount;
+					TargetPosition = Entry.GetTarget() + PaddingAmount;
+				} 
+
+				return;
+			}
+		}
+
+		// fall back to the old approach if there isn't a schema
+		for (int32 Index = 1; Index < InData.Num() - 1; ++Index)
+		{
+			FAlignmentData& Entry = InData[Index];
+
+			Entry.Node->Modify();
+			Entry.TargetProperty = TargetPosition;
+
+			TargetPosition = Entry.GetTarget() + PaddingAmount;
+		}
 	}
 }
 
@@ -1575,7 +1611,7 @@ void SGraphEditorImpl::OnDistributeNodesH()
 	if (AlignData.Num() > 2)
 	{
 		const FScopedTransaction Transaction(FGraphEditorCommands::Get().DistributeNodesHorizontally->GetLabel());
-		DistributeNodes(AlignData);
+		DistributeNodes(AlignData, true);
 	}
 }
 
@@ -1593,7 +1629,7 @@ void SGraphEditorImpl::OnDistributeNodesV()
 	if (AlignData.Num() > 2)
 	{
 		const FScopedTransaction Transaction(FGraphEditorCommands::Get().DistributeNodesVertically->GetLabel());
-		DistributeNodes(AlignData);
+		DistributeNodes(AlignData, false);
 	}
 }
 
