Fixing comment bubble text when bubbles have non-default background color

diff --git a/Engine/Source/Editor/GraphEditor/Private/SCommentBubble.cpp b/Engine/Source/Editor/GraphEditor/Private/SCommentBubble.cpp
--- a/Engine/Source/Editor/GraphEditor/Private/SCommentBubble.cpp
+++ b/Engine/Source/Editor/GraphEditor/Private/SCommentBubble.cpp
@@ -31,12 +31,6 @@ namespace SCommentBubbleDefs
 	/** Luminance CoEficients */
 	static const FLinearColor LuminanceCoEff( 0.2126f, 0.7152f, 0.0722f, 0.f );
 
-	/** Light foreground color */
-	static const FLinearColor LightForegroundClr( 0.f, 0.f, 0.f, 0.65f );
-
-	/** Dark foreground color */
-	static const FLinearColor DarkForegroundClr( 1.f, 1.f, 1.f, 0.65f );
-
 	/** Clear text box background color */
 	static const FLinearColor TextClearBackground( 0.f, 0.f, 0.f, 0.f );
 
@@ -65,7 +59,7 @@ void SCommentBubble::Construct( const FArguments& InArgs )
 	ToggleButtonCheck		= InArgs._ToggleButtonCheck.IsBound() ?	InArgs._ToggleButtonCheck : 
 																	TAttribute<ECheckBoxState>( this, &SCommentBubble::GetToggleButtonCheck );
 	// Ensue this value is set to something sensible
-	ForegroundColor = SCommentBubbleDefs::LightForegroundClr;
+	ForegroundColor = FStyleColors::Foreground;
 
 	// Cache the comment
 	CachedComment = CommentAttribute.Get();
@@ -100,7 +94,7 @@ void SCommentBubble::Tick( const FGeometry& AllottedGeometry, const double InCur
 	{
 		const FLinearColor BubbleColor = GetBubbleColor().GetSpecifiedColor() * SCommentBubbleDefs::LuminanceCoEff;
 		const float BubbleLuminance = BubbleColor.R + BubbleColor.G + BubbleColor.B;
-		ForegroundColor = BubbleLuminance < 0.5f ? SCommentBubbleDefs::DarkForegroundClr : SCommentBubbleDefs::LightForegroundClr;
+		ForegroundColor = BubbleLuminance < 0.5f ? FStyleColors::Foreground : FStyleColors::Background;
 	}
 
 	TickVisibility(InCurrentTime, InDeltaTime);
@@ -205,7 +199,7 @@ void SCommentBubble::UpdateBubble()
 					.OnCheckStateChanged( this, &SCommentBubble::OnCommentBubbleToggle )
 					.ToolTipText( NSLOCTEXT( "CommentBubble", "ToggleCommentTooltip", "Toggle Comment Bubble" ))
 					.Cursor( EMouseCursor::Default )
-					.ForegroundColor( FLinearColor::White )
+					.ForegroundColor(FStyleColors::Foreground)
 				];
 			}
 			else
@@ -223,7 +217,7 @@ void SCommentBubble::UpdateBubble()
 					.OnCheckStateChanged( this, &SCommentBubble::OnCommentBubbleToggle )
 					.ToolTipText( NSLOCTEXT( "CommentBubble", "ToggleCommentTooltip", "Toggle Comment Bubble" ))
 					.Cursor( EMouseCursor::Default )
-					.ForegroundColor( FLinearColor::White )
+					.ForegroundColor( FStyleColors::Foreground )
 				];
 			}
 		}
@@ -265,6 +259,9 @@ void SCommentBubble::UpdateBubble()
 							.ClearKeyboardFocusOnCommit( true )
 							.ModiferKeyForNewLine(EModifierKey::Shift)
 							.OnTextCommitted(this, &SCommentBubble::OnCommentTextCommitted)
+							.ForegroundColor(this, &SCommentBubble::GetTextForegroundColor)
+							.ReadOnlyForegroundColor(this, &SCommentBubble::GetReadOnlyTextForegroundColor)
+							.BackgroundColor(this, &SCommentBubble::GetTextBackgroundColor)
 						]
 						+SHorizontalBox::Slot()
 						.AutoWidth()
@@ -410,6 +407,21 @@ void SCommentBubble::OnCommentTextCommitted( const FText& NewText, ETextCommit::
 	}
 }
 
+FSlateColor SCommentBubble::GetTextBackgroundColor() const
+{
+	return TextBlock->HasKeyboardFocus() ? FStyleColors::Foreground : SCommentBubbleDefs::TextClearBackground;
+}
+
+FSlateColor SCommentBubble::GetTextForegroundColor() const
+{
+	return TextBlock->HasKeyboardFocus() ? FSlateColor::UseStyle() : ForegroundColor;
+}
+
+FSlateColor SCommentBubble::GetReadOnlyTextForegroundColor() const
+{
+	return ForegroundColor;
+}
+
 EVisibility SCommentBubble::GetToggleButtonVisibility() const
 {
 	EVisibility ButtonVisibility = EVisibility::Collapsed;
diff --git a/Engine/Source/Editor/GraphEditor/Public/SCommentBubble.h b/Engine/Source/Editor/GraphEditor/Public/SCommentBubble.h
--- a/Engine/Source/Editor/GraphEditor/Public/SCommentBubble.h
+++ b/Engine/Source/Editor/GraphEditor/Public/SCommentBubble.h
@@ -139,6 +139,9 @@ protected:
 	/** Called when the comment text is committed */
 	void OnCommentTextCommitted( const FText& NewText, ETextCommit::Type CommitInfo );
 
+	FSlateColor GetTextBackgroundColor() const;
+	FSlateColor GetTextForegroundColor() const;
+	FSlateColor GetReadOnlyTextForegroundColor() const;
 	/** Returns bubble toggle check state */
 	ECheckBoxState GetToggleButtonCheck() const;
 
@@ -174,7 +177,7 @@ protected:
 	FIsGraphNodeHovered IsGraphNodeHovered;
 
 	/** Current Foreground Color */
-	FLinearColor ForegroundColor;
+	FSlateColor ForegroundColor;
 	/** Allow pin behaviour */
 	bool bAllowPinning;
 	/** Allow in bubble controls */
