Fix bad style. Introduced by 15410841

diff --git a/Engine/Source/Editor/GraphEditor/Private/SCommentBubble.cpp b/Engine/Source/Editor/GraphEditor/Private/SCommentBubble.cpp
--- a/Engine/Source/Editor/GraphEditor/Private/SCommentBubble.cpp
+++ b/Engine/Source/Editor/GraphEditor/Private/SCommentBubble.cpp
@@ -251,7 +251,7 @@ void SCommentBubble::UpdateBubble()
 						[
 							SAssignNew(TextBlock, SMultiLineEditableTextBox)
 							.Text(MakeAttributeLambda([this] { return CachedCommentText; }))
-							.Style(FAppStyle::Get(), "Graph.Node.NodeTitleInlineEditableText")
+							.Style(FAppStyle::Get(), "Graph.Node.NodeTitleEditableText")
 							.HintText( NSLOCTEXT( "CommentBubble", "EditCommentHint", "Click to edit" ))
 							.IsReadOnly(this, &SCommentBubble::IsReadOnly)
 							.SelectAllTextWhenFocused( true )
