Fix giant separators in graph action menus

diff --git a/Engine/Source/Editor/GraphEditor/Private/SGraphActionMenu.cpp b/Engine/Source/Editor/GraphEditor/Private/SGraphActionMenu.cpp
--- a/Engine/Source/Editor/GraphEditor/Private/SGraphActionMenu.cpp
+++ b/Engine/Source/Editor/GraphEditor/Private/SGraphActionMenu.cpp
@@ -21,6 +21,7 @@
 #include "Widgets/Input/SSearchBox.h"
 #include "Widgets/Text/SInlineEditableTextBlock.h"
 #include "Math/NumericLimits.h"
+#include "Widgets/Layout/SSeparator.h"
 
 #define LOCTEXT_NAMESPACE "GraphActionMenu"
 
@@ -986,17 +987,13 @@ TSharedRef<ITableRow> SGraphActionMenu::MakeWidget( TSharedPtr<FGraphActionNode>
 			.Visibility(EVisibility::HitTestInvisible)
 
 			+ SVerticalBox::Slot()
-			.AutoHeight()
+			.VAlign(VAlign_Center)
 			// Add some empty space before the line, and a tiny bit after it
-			.Padding( 0.0f, 5.f, 0.0f, 5.f )
+			.Padding( 0.0f, 1.f, 0.0f, 1.f )
 			[
-				SNew( SBorder )
-
-				// We'll use the border's padding to actually create the horizontal line
-				.Padding(FEditorStyle::GetMargin(TEXT("Menu.Separator.Padding")))
-
-				// Separator graphic
-				.BorderImage( FEditorStyle::GetBrush( TEXT( "Menu.Separator" ) ) )
+				SNew(SSeparator)
+				.SeparatorImage(FEditorStyle::Get().GetBrush("Menu.Separator"))
+				.Thickness(1.0f)
 			];
 		}
 		else
