Only add extra padding for pure compact nodes instead of all compact nodes.

#rnx

diff --git a/Engine/Source/Editor/GraphEditor/Private/KismetNodes/SGraphNodeK2Base.cpp b/Engine/Source/Editor/GraphEditor/Private/KismetNodes/SGraphNodeK2Base.cpp
--- a/Engine/Source/Editor/GraphEditor/Private/KismetNodes/SGraphNodeK2Base.cpp
+++ b/Engine/Source/Editor/GraphEditor/Private/KismetNodes/SGraphNodeK2Base.cpp
@@ -112,16 +112,9 @@ void SGraphNodeK2Base::UpdateCompactNode()
 	// since don't need to worry about alignment with other nodes
 	float PinPaddingTop = 0.f;
 
-	// If the center node title is 2 or more, then make the node bigger
-	// so that the text box isn't over top of the label
 	static float MinNodePadding = 55.f;
-	static float MaxNodePadding = 180.0f;
-	static float PaddingIncrementSize = 20.0f;
-
-	int32 HeadTitleLength = NodeTitle.Get() ? NodeTitle.Get()->GetHeadTitle().ToString().Len() : 0;
-	
 	// Calculate a padding amount clamping to the min/max settings
-	const float PinPaddingRight = FMath::Clamp(MinNodePadding + ((float)HeadTitleLength * PaddingIncrementSize), MinNodePadding, MaxNodePadding);
+	float PinPaddingRight = MinNodePadding;
 
 	EVerticalAlignment PinVerticalAlignment = VAlign_Center;
 
@@ -129,11 +122,24 @@ void SGraphNodeK2Base::UpdateCompactNode()
 	// and add some padding so that the exec pins line up with the exec pins of other nodes
 	if (UK2Node* K2Node = Cast<UK2Node>(GraphNode))
 	{
-		if (!K2Node->IsNodePure())
+		const bool bIsPure = K2Node->IsNodePure();
+		if (!bIsPure)
 		{
 			PinPaddingTop += 8.0f;
 			PinVerticalAlignment = VAlign_Top;
 		}
+
+		if (K2Node->ShouldDrawCompact() && bIsPure)
+		{
+			// If the center node title is 2 or more, then make the node bigger
+			// so that the text box isn't over top of the label
+			static float MaxNodePadding = 180.0f;
+			static float PaddingIncrementSize = 20.0f;
+
+			int32 HeadTitleLength = NodeTitle.Get() ? NodeTitle.Get()->GetHeadTitle().ToString().Len() : 0;
+
+			PinPaddingRight = FMath::Clamp(MinNodePadding + ((float)HeadTitleLength * PaddingIncrementSize), MinNodePadding, MaxNodePadding);
+		}
 	}
 
 	NodeOverlay->AddSlot()
