Visually style DefaultToSelf blueprint pins to look like normal Self pins to make it clear they will auto set to self, the functionality is completely unchanged
To do this without storing something on the pin itself, extend GetPinMetaData on function call nodes to handle "FunctionOption=PinName" function metadata as metadata for PinName
Modified version of PR #5949

diff --git a/Engine/Source/Editor/BlueprintGraph/Classes/K2Node_BaseAsyncTask.h b/Engine/Source/Editor/BlueprintGraph/Classes/K2Node_BaseAsyncTask.h
--- a/Engine/Source/Editor/BlueprintGraph/Classes/K2Node_BaseAsyncTask.h
+++ b/Engine/Source/Editor/BlueprintGraph/Classes/K2Node_BaseAsyncTask.h
@@ -42,6 +42,7 @@ class BLUEPRINTGRAPH_API UK2Node_BaseAsyncTask : public UK2Node
 	virtual bool IsCompatibleWithGraph(const UEdGraph* TargetGraph) const override;
 	virtual void ValidateNodeDuringCompilation(class FCompilerResultsLog& MessageLog) const override;
 	virtual void GetPinHoverText(const UEdGraphPin& Pin, FString& HoverTextOut) const override;
+	virtual FString GetPinMetaData(FName InPinName, FName InKey) override;
 	// End of UEdGraphNode interface
 
 	// UK2Node interface
diff --git a/Engine/Source/Editor/BlueprintGraph/Private/K2Node_BaseAsyncTask.cpp b/Engine/Source/Editor/BlueprintGraph/Private/K2Node_BaseAsyncTask.cpp
--- a/Engine/Source/Editor/BlueprintGraph/Private/K2Node_BaseAsyncTask.cpp
+++ b/Engine/Source/Editor/BlueprintGraph/Private/K2Node_BaseAsyncTask.cpp
@@ -691,4 +691,35 @@ void UK2Node_BaseAsyncTask::GetPinHoverText(const UEdGraphPin& Pin, FString& Hov
 	return UK2Node::GetPinHoverText(Pin, HoverTextOut);
 }
 
+FString UK2Node_BaseAsyncTask::GetPinMetaData(FName InPinName, FName InKey)
+{
+	FString MetaData = Super::GetPinMetaData(InPinName, InKey);
+
+	// If there's no metadata directly on the pin then check for metadata on the function
+	if (MetaData.IsEmpty())
+	{
+		if (UFunction* Function = GetFactoryFunction())
+		{
+			// Find the corresponding property for the pin and search that first
+			if (FProperty* Property = Function->FindPropertyByName(InPinName))
+			{
+				MetaData = Property->GetMetaData(InKey);
+			}
+
+			// Also look for metadata like DefaultToSelf on the function itself
+			if (MetaData.IsEmpty())
+			{
+				MetaData = Function->GetMetaData(InKey);
+				if (MetaData != InPinName.ToString())
+				{
+					// Only return if the value matches the pin name as we don't want general function metadata
+					MetaData.Empty();
+				}
+			}
+		}
+	}
+
+	return MetaData;
+}
+
 #undef LOCTEXT_NAMESPACE
diff --git a/Engine/Source/Editor/BlueprintGraph/Private/K2Node_CallFunction.cpp b/Engine/Source/Editor/BlueprintGraph/Private/K2Node_CallFunction.cpp
--- a/Engine/Source/Editor/BlueprintGraph/Private/K2Node_CallFunction.cpp
+++ b/Engine/Source/Editor/BlueprintGraph/Private/K2Node_CallFunction.cpp
@@ -3349,11 +3349,22 @@ FString UK2Node_CallFunction::GetPinMetaData(FName InPinName, FName InKey)
 	{
 		if (UFunction* Function = GetTargetFunction())
 		{
-			// Find the corresponding property for the pin
+			// Find the corresponding property for the pin and search that first
 			if (FProperty* Property = Function->FindPropertyByName(InPinName))
 			{
 				MetaData = Property->GetMetaData(InKey);
 			}
+
+			// Also look for metadata like DefaultToSelf on the function itself
+			if (MetaData.IsEmpty())
+			{
+				MetaData = Function->GetMetaData(InKey);
+				if (MetaData != InPinName.ToString())
+				{
+					// Only return if the value matches the pin name as we don't want general function metadata
+					MetaData.Empty();
+				}
+			}
 		}
 	}
 
diff --git a/Engine/Source/Editor/GraphEditor/Private/KismetPins/SGraphPinObject.cpp b/Engine/Source/Editor/GraphEditor/Private/KismetPins/SGraphPinObject.cpp
--- a/Engine/Source/Editor/GraphEditor/Private/KismetPins/SGraphPinObject.cpp
+++ b/Engine/Source/Editor/GraphEditor/Private/KismetPins/SGraphPinObject.cpp
@@ -14,6 +14,7 @@
 #include "ScopedTransaction.h"
 #include "Engine/Selection.h"
 #include "AssetRegistryModule.h"
+#include "EdGraphSchema_K2.h"
 
 #define LOCTEXT_NAMESPACE "SGraphPinObject"
 
@@ -52,18 +53,15 @@ TSharedRef<SWidget>	SGraphPinObject::GetDefaultValueWidget()
 		return SNullWidget::NullWidget;
 	}
 
-	if( AllowSelfPinWidget() )
+	if(ShouldDisplayAsSelfPin())
 	{
-		if(Schema->IsSelfPin(*GraphPinObj))
-		{
-			return SNew(SEditableTextBox)
-				.Style( FEditorStyle::Get(), "Graph.EditableTextBox" )
-				.Text( this, &SGraphPinObject::GetValue )
-				.SelectAllTextWhenFocused(false)
-				.Visibility( this, &SGraphPinObject::GetDefaultValueVisibility )
-				.IsReadOnly(true)
-				.ForegroundColor( FSlateColor::UseForeground() );
-		}
+		return SNew(SEditableTextBox)
+			.Style( FEditorStyle::Get(), "Graph.EditableTextBox" )
+			.Text( this, &SGraphPinObject::GetValue )
+			.SelectAllTextWhenFocused(false)
+			.Visibility( this, &SGraphPinObject::GetDefaultValueVisibility )
+			.IsReadOnly( true )
+			.ForegroundColor( FSlateColor::UseForeground() );
 	}
 	// Don't show literal buttons for component type objects
 	if (Schema->ShouldShowAssetPickerForPin(GraphPinObj))
@@ -279,6 +277,27 @@ void SGraphPinObject::OnAssetEnterPressedInPicker(const TArray<FAssetData>& InSe
 	}
 }
 
+bool SGraphPinObject::ShouldDisplayAsSelfPin() const
+{
+	if (AllowSelfPinWidget() && GraphPinObj)
+	{
+		if (GraphPinObj->GetSchema()->IsSelfPin(*GraphPinObj))
+		{
+			return true;
+		}
+
+		// Also check function/node metadata
+		const FString DefaultToSelfString = GraphPinObj->GetOwningNode()->GetPinMetaData(GraphPinObj->PinName, FBlueprintMetadata::MD_DefaultToSelf);
+		if (DefaultToSelfString == GraphPinObj->PinName.ToString())
+		{
+			return true;
+		}
+	}
+
+	return false;
+}
+
+
 FText SGraphPinObject::GetValue() const
 {
 	const FAssetData& CurrentAssetData = GetAssetData(true);
@@ -293,6 +312,10 @@ FText SGraphPinObject::GetValue() const
 		{
 			Value =  FText::FromName(GraphPinObj->PinName);
 		}
+		else if (ShouldDisplayAsSelfPin())
+		{
+			Value = FText::FromName(UEdGraphSchema_K2::PN_Self);
+		}
 		else
 		{
 			Value = FText::GetEmpty();
diff --git a/Engine/Source/Editor/GraphEditor/Public/KismetPins/SGraphPinObject.h b/Engine/Source/Editor/GraphEditor/Public/KismetPins/SGraphPinObject.h
--- a/Engine/Source/Editor/GraphEditor/Public/KismetPins/SGraphPinObject.h
+++ b/Engine/Source/Editor/GraphEditor/Public/KismetPins/SGraphPinObject.h
@@ -47,6 +47,8 @@ protected:
 	virtual FText GetDefaultComboText() const;
 	/** Allow self pin widget */
 	virtual bool AllowSelfPinWidget() const { return true; }
+	/** True if this specific pin should be treated as a self pin */
+	virtual bool ShouldDisplayAsSelfPin() const;
 	/** Generate asset picker window */
 	virtual TSharedRef<SWidget> GenerateAssetPicker();
 	/** Called to validate selection from picker window */
