Control Rig: Fix crash in drag connection

diff --git a/Engine/Source/Editor/GraphEditor/Private/DragConnection.cpp b/Engine/Source/Editor/GraphEditor/Private/DragConnection.cpp
--- a/Engine/Source/Editor/GraphEditor/Private/DragConnection.cpp
+++ b/Engine/Source/Editor/GraphEditor/Private/DragConnection.cpp
@@ -261,11 +261,13 @@ FReply FDragConnection::DroppedOnNode(FVector2D ScreenPosition, FVector2D GraphP
 		{
 			for (UEdGraphPin* SourcePin : ValidSourcePins)
 			{
-				SourcePin->GetSchema()->SetPinBeingDroppedOnNode(SourcePin);
+				// copy it here since the pin might no longer be valid
+				const UEdGraphSchema* SourcePinSchema = SourcePin->GetSchema();
+				SourcePinSchema->SetPinBeingDroppedOnNode(SourcePin);
 
 				// Check for pin drop support
 				FText ResponseText;
-				if (SourcePin->GetOwningNode() != NodeOver && SourcePin->GetSchema()->SupportsDropPinOnNode(NodeOver, SourcePin->PinType, SourcePin->Direction, ResponseText))
+				if (SourcePin->GetOwningNode() != NodeOver && SourcePinSchema->SupportsDropPinOnNode(NodeOver, SourcePin->PinType, SourcePin->Direction, ResponseText))
 				{
 					bHandledPinDropOnNode = true;
 
@@ -281,7 +283,7 @@ FReply FDragConnection::DroppedOnNode(FVector2D ScreenPosition, FVector2D GraphP
 					{
 						SourcePin->Modify();
 						EdGraphPin->Modify();
-						SourcePin->GetSchema()->TryCreateConnection(SourcePin, EdGraphPin);
+						SourcePinSchema->TryCreateConnection(SourcePin, EdGraphPin);
 					}
 				}
 
@@ -291,7 +293,7 @@ FReply FDragConnection::DroppedOnNode(FVector2D ScreenPosition, FVector2D GraphP
 					bHandledPinDropOnNode = true;
 				}
 
-				SourcePin->GetSchema()->SetPinBeingDroppedOnNode(nullptr);
+				SourcePinSchema->SetPinBeingDroppedOnNode(nullptr);
 			}
 		}
 	}
