Fix for crash when zooming out of a graph

diff --git a/Engine/Source/Editor/GraphEditor/Private/SGraphNode.cpp b/Engine/Source/Editor/GraphEditor/Private/SGraphNode.cpp
--- a/Engine/Source/Editor/GraphEditor/Private/SGraphNode.cpp
+++ b/Engine/Source/Editor/GraphEditor/Private/SGraphNode.cpp
@@ -1487,14 +1487,15 @@ void SGraphNode::NotifyDisallowedPinConnection(const UEdGraphPin* PinA, const UE
 
 bool SGraphNode::UseLowDetailNodeTitles() const
 {
-	if (const SGraphPanel* MyOwnerPanel = GetOwnerPanel().Get())
+	if (InlineEditableText.IsValid())
 	{
-		return (MyOwnerPanel->GetCurrentLOD() <= EGraphRenderingLOD::LowestDetail) && !InlineEditableText->IsInEditMode();
-	}
-	else
-	{
-		return false;
+		if (const SGraphPanel* MyOwnerPanel = GetOwnerPanel().Get())
+		{
+			return (MyOwnerPanel->GetCurrentLOD() <= EGraphRenderingLOD::LowestDetail) && !InlineEditableText->IsInEditMode();
+		}
 	}
+
+	return false;
 }
 
 TSharedRef<SWidget> SGraphNode::AddPinButtonContent(FText PinText, FText PinTooltipText, bool bRightSide, FString DocumentationExcerpt, TSharedPtr<SToolTip> CustomTooltip)
