Temporary fix to unblock the floor until we figure out why the material editor needs the GraphPin.
#rnx

diff --git a/Engine/Source/Editor/GraphEditor/Private/SGraphPin.cpp b/Engine/Source/Editor/GraphEditor/Private/SGraphPin.cpp
--- a/Engine/Source/Editor/GraphEditor/Private/SGraphPin.cpp
+++ b/Engine/Source/Editor/GraphEditor/Private/SGraphPin.cpp
@@ -1107,9 +1107,9 @@ const FSlateBrush* SGraphPin::GetPinStatusIcon() const
 
 		if (UEdGraphNode* GraphNode = WatchedPin->GetOwningNodeUnchecked())
 		{
-			UBlueprint* Blueprint = FBlueprintEditorUtils::FindBlueprintForNodeChecked(GraphNode);
+			UBlueprint* Blueprint = FBlueprintEditorUtils::FindBlueprintForNode(GraphNode);
 
-			if (FKismetDebugUtilities::IsPinBeingWatched(Blueprint, WatchedPin))
+			if (Blueprint && FKismetDebugUtilities::IsPinBeingWatched(Blueprint, WatchedPin))
 			{
 				return FEditorStyle::GetBrush(TEXT("Graph.WatchedPinIcon_Pinned"));
 			}
