Revert temporary CL 15988014 that made the material editor crash with a reference to SGraphPin. Fixed by CL 16007736
#rnx

diff --git a/Engine/Source/Editor/GraphEditor/Private/SGraphPin.cpp b/Engine/Source/Editor/GraphEditor/Private/SGraphPin.cpp
--- a/Engine/Source/Editor/GraphEditor/Private/SGraphPin.cpp
+++ b/Engine/Source/Editor/GraphEditor/Private/SGraphPin.cpp
@@ -1107,9 +1107,9 @@ const FSlateBrush* SGraphPin::GetPinStatusIcon() const
 
 		if (UEdGraphNode* GraphNode = WatchedPin->GetOwningNodeUnchecked())
 		{
-			UBlueprint* Blueprint = FBlueprintEditorUtils::FindBlueprintForNode(GraphNode);
+			UBlueprint* Blueprint = FBlueprintEditorUtils::FindBlueprintForNodeChecked(GraphNode);
 
-			if (Blueprint && FKismetDebugUtilities::IsPinBeingWatched(Blueprint, WatchedPin))
+			if (FKismetDebugUtilities::IsPinBeingWatched(Blueprint, WatchedPin))
 			{
 				return FEditorStyle::GetBrush(TEXT("Graph.WatchedPinIcon_Pinned"));
 			}
