Expose FindReferences via the GRAPHEDTIOR_API so it can be used in the subobject editor

#rnx

diff --git a/Engine/Source/Editor/GraphEditor/Public/GraphEditorActions.h b/Engine/Source/Editor/GraphEditor/Public/GraphEditorActions.h
--- a/Engine/Source/Editor/GraphEditor/Public/GraphEditorActions.h
+++ b/Engine/Source/Editor/GraphEditor/Public/GraphEditorActions.h
@@ -95,6 +95,7 @@ public:
 
 	// Find references
 	TSharedPtr< FUICommandInfo > FindReferences;
+	GRAPHEDITOR_API TSharedPtr< FUICommandInfo > GetFindReferences() const { return FindReferences; }
 	TSharedPtr< FUICommandInfo > FindAndReplaceReferences;
 
 	// Jumps to the definition of the selected node (or otherwise focuses something interesting about that node, e.g., the inner graph for a collapsed graph)
