SNodePanel: Lazily repopulate visible children

diff --git a/Engine/Source/Editor/GraphEditor/Private/SNodePanel.cpp b/Engine/Source/Editor/GraphEditor/Private/SNodePanel.cpp
--- a/Engine/Source/Editor/GraphEditor/Private/SNodePanel.cpp
+++ b/Engine/Source/Editor/GraphEditor/Private/SNodePanel.cpp
@@ -1172,17 +1172,31 @@ void SNodePanel::RemoveAllNodes()
 
 void SNodePanel::PopulateVisibleChildren(const FGeometry& AllottedGeometry)
 {
-	VisibleChildren.Empty();
+	bool bRequiresSort = false;
 	for (int32 ChildIndex = 0; ChildIndex < Children.Num(); ++ChildIndex)
 	{
 		const TSharedRef<SNode>& SomeChild = Children[ChildIndex];
 		if ( !IsNodeCulled(SomeChild, AllottedGeometry) )
 		{
-			VisibleChildren.Add(SomeChild);
+			if(VisibleChildren.Find(SomeChild) == INDEX_NONE)
+			{
+				VisibleChildren.Add(SomeChild);
+				bRequiresSort = true;
+			}
+		}
+		else
+		{
+			if(VisibleChildren.Find(SomeChild) != INDEX_NONE)
+			{
+				VisibleChildren.Remove(SomeChild);
+				bRequiresSort = true;
+			}
+			
 		}
 	}
+	
 	// Depth Sort Nodes
-	if( VisibleChildren.Num() > 0 )
+	if( bRequiresSort && (VisibleChildren.Num() > 0) )
 	{
 		struct SNodeLessThanSort
 		{
