Edigrate 16359780.

Don't create transactions except on release when moving nodes.

diff --git a/Engine/Plugins/Enterprise/DataprepEditor/Source/DataprepEditor/Private/Widgets/DataprepGraph/SDataprepGraphTrackNode.cpp b/Engine/Plugins/Enterprise/DataprepEditor/Source/DataprepEditor/Private/Widgets/DataprepGraph/SDataprepGraphTrackNode.cpp
--- a/Engine/Plugins/Enterprise/DataprepEditor/Source/DataprepEditor/Private/Widgets/DataprepGraph/SDataprepGraphTrackNode.cpp
+++ b/Engine/Plugins/Enterprise/DataprepEditor/Source/DataprepEditor/Private/Widgets/DataprepGraph/SDataprepGraphTrackNode.cpp
@@ -563,10 +563,10 @@ bool SDataprepGraphTrackNode::RefreshLayout()
 	return false;
 }
 
-void SDataprepGraphTrackNode::MoveTo(const FVector2D& NewPosition, FNodeSet& NodeFilter)
+void SDataprepGraphTrackNode::MoveTo(const FVector2D& NewPosition, FNodeSet& NodeFilter, bool bMarkDirty)
 {
 	// Block track node to specific position
-	SGraphNode::MoveTo( TrackAnchor, NodeFilter);
+	SGraphNode::MoveTo( TrackAnchor, NodeFilter, bMarkDirty);
 }
 
 const FSlateBrush* SDataprepGraphTrackNode::GetShadowBrush(bool bSelected) const
diff --git a/Engine/Plugins/Enterprise/DataprepEditor/Source/DataprepEditor/Private/Widgets/DataprepGraph/SDataprepGraphTrackNode.h b/Engine/Plugins/Enterprise/DataprepEditor/Source/DataprepEditor/Private/Widgets/DataprepGraph/SDataprepGraphTrackNode.h
--- a/Engine/Plugins/Enterprise/DataprepEditor/Source/DataprepEditor/Private/Widgets/DataprepGraph/SDataprepGraphTrackNode.h
+++ b/Engine/Plugins/Enterprise/DataprepEditor/Source/DataprepEditor/Private/Widgets/DataprepGraph/SDataprepGraphTrackNode.h
@@ -78,7 +78,7 @@ public:
 		return false;
 	}
 	virtual void SetOwner(const TSharedRef<SGraphPanel>& OwnerPanel) override;
-	virtual void MoveTo( const FVector2D& NewPosition, FNodeSet& NodeFilter ) override;
+	virtual void MoveTo( const FVector2D& NewPosition, FNodeSet& NodeFilter, bool bMarkDirty = true ) override;
 	virtual const FSlateBrush* GetShadowBrush(bool bSelected) const override;
 	virtual FReply OnDrop( const FGeometry& MyGeometry, const FDragDropEvent& DragDropEvent ) override;
 	virtual bool ShouldAllowCulling() const override { return false; }
diff --git a/Engine/Plugins/Experimental/CommonConversation/Source/CommonConversationGraph/Private/SConversationGraphNode.cpp b/Engine/Plugins/Experimental/CommonConversation/Source/CommonConversationGraph/Private/SConversationGraphNode.cpp
--- a/Engine/Plugins/Experimental/CommonConversation/Source/CommonConversationGraph/Private/SConversationGraphNode.cpp
+++ b/Engine/Plugins/Experimental/CommonConversation/Source/CommonConversationGraph/Private/SConversationGraphNode.cpp
@@ -1061,9 +1061,9 @@ TSharedRef<SGraphNode> SConversationGraphNode::GetNodeUnderMouse(const FGeometry
 	return SubNode.IsValid() ? SubNode.ToSharedRef() : StaticCastSharedRef<SGraphNode>(AsShared());
 }
 
-void SConversationGraphNode::MoveTo(const FVector2D& NewPosition, FNodeSet& NodeFilter)
+void SConversationGraphNode::MoveTo(const FVector2D& NewPosition, FNodeSet& NodeFilter, bool bMarkDirty)
 {
-	SGraphNodeAI::MoveTo(NewPosition, NodeFilter);
+	SGraphNodeAI::MoveTo(NewPosition, NodeFilter, bMarkDirty);
 
 	UConversationGraphNode* ConversationGraphNode = CastChecked<UConversationGraphNode>(GraphNode);
 	if (!ConversationGraphNode->IsSubNode())
diff --git a/Engine/Plugins/Experimental/CommonConversation/Source/CommonConversationGraph/Private/SConversationGraphNode.h b/Engine/Plugins/Experimental/CommonConversation/Source/CommonConversationGraph/Private/SConversationGraphNode.h
--- a/Engine/Plugins/Experimental/CommonConversation/Source/CommonConversationGraph/Private/SConversationGraphNode.h
+++ b/Engine/Plugins/Experimental/CommonConversation/Source/CommonConversationGraph/Private/SConversationGraphNode.h
@@ -52,7 +52,7 @@ public:
 	virtual void GetOverlayBrushes(bool bSelected, const FVector2D WidgetSize, TArray<FOverlayBrushInfo>& Brushes) const override;
 	virtual TArray<FOverlayWidgetInfo> GetOverlayWidgets(bool bSelected, const FVector2D& WidgetSize) const override;
 	virtual TSharedRef<SGraphNode> GetNodeUnderMouse(const FGeometry& MyGeometry, const FPointerEvent& MouseEvent) override;
-	virtual void MoveTo(const FVector2D& NewPosition, FNodeSet& NodeFilter) override;
+	virtual void MoveTo(const FVector2D& NewPosition, FNodeSet& NodeFilter, bool bMarkDirty = true) override;
 	// End of SGraphNode interface
 
 	/** handle double click */
diff --git a/Engine/Plugins/Experimental/ControlRig/Source/ControlRigEditor/Private/Graph/SControlRigGraphNode.cpp b/Engine/Plugins/Experimental/ControlRig/Source/ControlRigEditor/Private/Graph/SControlRigGraphNode.cpp
--- a/Engine/Plugins/Experimental/ControlRig/Source/ControlRigEditor/Private/Graph/SControlRigGraphNode.cpp
+++ b/Engine/Plugins/Experimental/ControlRig/Source/ControlRigEditor/Private/Graph/SControlRigGraphNode.cpp
@@ -563,7 +563,7 @@ void SControlRigGraphNode::EndUserInteraction() const
 	SGraphNode::EndUserInteraction();
 }
 
-void SControlRigGraphNode::MoveTo(const FVector2D& NewPosition, FNodeSet& NodeFilter)
+void SControlRigGraphNode::MoveTo(const FVector2D& NewPosition, FNodeSet& NodeFilter, bool bMarkDirty)
 {
 	if (!NodeFilter.Find(SharedThis(this)))
 	{
diff --git a/Engine/Plugins/Experimental/ControlRig/Source/ControlRigEditor/Private/Graph/SControlRigGraphNode.h b/Engine/Plugins/Experimental/ControlRig/Source/ControlRigEditor/Private/Graph/SControlRigGraphNode.h
--- a/Engine/Plugins/Experimental/ControlRig/Source/ControlRigEditor/Private/Graph/SControlRigGraphNode.h
+++ b/Engine/Plugins/Experimental/ControlRig/Source/ControlRigEditor/Private/Graph/SControlRigGraphNode.h
@@ -32,7 +32,7 @@ public:
 	// SGraphNode interface
 	virtual TSharedRef<SWidget> CreateTitleWidget(TSharedPtr<SNodeTitle> InNodeTitle) override;
 	virtual void EndUserInteraction() const override;
-	virtual void MoveTo(const FVector2D& NewPosition, FNodeSet& NodeFilter) override;
+	virtual void MoveTo(const FVector2D& NewPosition, FNodeSet& NodeFilter, bool bMarkDirty = true) override;
 	virtual void AddPin( const TSharedRef<SGraphPin>& PinToAdd ) override;
 	virtual void SetDefaultTitleAreaWidget(TSharedRef<SOverlay> DefaultTitleAreaWidget) override
 	{
diff --git a/Engine/Plugins/Runtime/Metasound/Source/MetasoundEditor/Private/SMetasoundGraphNode.cpp b/Engine/Plugins/Runtime/Metasound/Source/MetasoundEditor/Private/SMetasoundGraphNode.cpp
--- a/Engine/Plugins/Runtime/Metasound/Source/MetasoundEditor/Private/SMetasoundGraphNode.cpp
+++ b/Engine/Plugins/Runtime/Metasound/Source/MetasoundEditor/Private/SMetasoundGraphNode.cpp
@@ -322,9 +322,9 @@ void SMetasoundGraphNode::SetDefaultTitleAreaWidget(TSharedRef<SOverlay> Default
 	}
 }
 
-void SMetasoundGraphNode::MoveTo(const FVector2D& NewPosition, FNodeSet& NodeFilter)
+void SMetasoundGraphNode::MoveTo(const FVector2D& NewPosition, FNodeSet& NodeFilter, bool bMarkDirty)
 {
-	SGraphNode::MoveTo(NewPosition, NodeFilter);
+	SGraphNode::MoveTo(NewPosition, NodeFilter, bMarkDirty);
 
 	GetMetasoundNode().UpdatePosition();
 }
diff --git a/Engine/Plugins/Runtime/Metasound/Source/MetasoundEditor/Private/SMetasoundGraphNode.h b/Engine/Plugins/Runtime/Metasound/Source/MetasoundEditor/Private/SMetasoundGraphNode.h
--- a/Engine/Plugins/Runtime/Metasound/Source/MetasoundEditor/Private/SMetasoundGraphNode.h
+++ b/Engine/Plugins/Runtime/Metasound/Source/MetasoundEditor/Private/SMetasoundGraphNode.h
@@ -38,7 +38,7 @@ protected:
 	virtual TSharedRef<SWidget> CreateTitleRightWidget() override;
 	virtual EVisibility IsAddPinButtonVisible() const override;
 	virtual FReply OnAddPin() override;
-	virtual void MoveTo(const FVector2D& NewPosition, FNodeSet& NodeFilter) override;
+	virtual void MoveTo(const FVector2D& NewPosition, FNodeSet& NodeFilter, bool bMarkDirty = true) override;
 	virtual void SetDefaultTitleAreaWidget(TSharedRef<SOverlay> DefaultTitleAreaWidget) override;
 
 	FName GetLiteralDataType() const;
diff --git a/Engine/Plugins/Runtime/nDisplay/Source/DisplayClusterConfigurator/Private/Views/OutputMapping/GraphNodes/SDisplayClusterConfiguratorBaseNode.cpp b/Engine/Plugins/Runtime/nDisplay/Source/DisplayClusterConfigurator/Private/Views/OutputMapping/GraphNodes/SDisplayClusterConfiguratorBaseNode.cpp
--- a/Engine/Plugins/Runtime/nDisplay/Source/DisplayClusterConfigurator/Private/Views/OutputMapping/GraphNodes/SDisplayClusterConfiguratorBaseNode.cpp
+++ b/Engine/Plugins/Runtime/nDisplay/Source/DisplayClusterConfigurator/Private/Views/OutputMapping/GraphNodes/SDisplayClusterConfiguratorBaseNode.cpp
@@ -153,7 +153,7 @@ FVector2D SDisplayClusterConfiguratorBaseNode::ComputeDesiredSize(float) const
 	return GetSize();
 }
 
-void SDisplayClusterConfiguratorBaseNode::MoveTo(const FVector2D& NewPosition, FNodeSet& NodeFilter)
+void SDisplayClusterConfiguratorBaseNode::MoveTo(const FVector2D& NewPosition, FNodeSet& NodeFilter, bool bMarkDirty)
 {
 	TSharedPtr<FDisplayClusterConfiguratorBlueprintEditor> Toolkit = ToolkitPtr.Pin();
 	check(Toolkit.IsValid());
@@ -215,7 +215,7 @@ void SDisplayClusterConfiguratorBaseNode::MoveTo(const FVector2D& NewPosition, F
 		YAlignmentTarget.TargetNode.Reset();
 	}
 
-	SGraphNode::MoveTo(CurrentPosition + BestOffset + AlignmentOffset, NodeFilter);
+	SGraphNode::MoveTo(CurrentPosition + BestOffset + AlignmentOffset, NodeFilter, bMarkDirty);
 
 	if (!bIsNodeFiltered)
 	{
diff --git a/Engine/Plugins/Runtime/nDisplay/Source/DisplayClusterConfigurator/Private/Views/OutputMapping/GraphNodes/SDisplayClusterConfiguratorBaseNode.h b/Engine/Plugins/Runtime/nDisplay/Source/DisplayClusterConfigurator/Private/Views/OutputMapping/GraphNodes/SDisplayClusterConfiguratorBaseNode.h
--- a/Engine/Plugins/Runtime/nDisplay/Source/DisplayClusterConfigurator/Private/Views/OutputMapping/GraphNodes/SDisplayClusterConfiguratorBaseNode.h
+++ b/Engine/Plugins/Runtime/nDisplay/Source/DisplayClusterConfigurator/Private/Views/OutputMapping/GraphNodes/SDisplayClusterConfiguratorBaseNode.h
@@ -72,7 +72,7 @@ public:
 	//~ Begin SGraphNode interface
 	virtual void UpdateGraphNode() override;
 	virtual FVector2D ComputeDesiredSize(float) const override;
-	virtual void MoveTo(const FVector2D& NewPosition, FNodeSet& NodeFilter) override;
+	virtual void MoveTo(const FVector2D& NewPosition, FNodeSet& NodeFilter, bool bMarkDirty = true) override;
 	virtual void EndUserInteraction() const override;
 	virtual const FSlateBrush* GetShadowBrush(bool bSelected) const override;
 	virtual bool CanBeSelected(const FVector2D& MousePositionInNode) const override;
diff --git a/Engine/Plugins/Runtime/nDisplay/Source/DisplayClusterConfigurator/Private/Views/OutputMapping/GraphNodes/SDisplayClusterConfiguratorCanvasNode.cpp b/Engine/Plugins/Runtime/nDisplay/Source/DisplayClusterConfigurator/Private/Views/OutputMapping/GraphNodes/SDisplayClusterConfiguratorCanvasNode.cpp
--- a/Engine/Plugins/Runtime/nDisplay/Source/DisplayClusterConfigurator/Private/Views/OutputMapping/GraphNodes/SDisplayClusterConfiguratorCanvasNode.cpp
+++ b/Engine/Plugins/Runtime/nDisplay/Source/DisplayClusterConfigurator/Private/Views/OutputMapping/GraphNodes/SDisplayClusterConfiguratorCanvasNode.cpp
@@ -75,12 +75,12 @@ void SDisplayClusterConfiguratorCanvasNode::UpdateGraphNode()
 	];
 }
 
-void SDisplayClusterConfiguratorCanvasNode::MoveTo(const FVector2D& NewPosition, FNodeSet& NodeFilter)
+void SDisplayClusterConfiguratorCanvasNode::MoveTo(const FVector2D& NewPosition, FNodeSet& NodeFilter, bool bMarkDirty)
 {
 	// Canvas node is not allowed to be moved in general, so add it to the node filter
 	NodeFilter.Add(SharedThis(this));
 
-	SGraphNode::MoveTo(NewPosition, NodeFilter);
+	SGraphNode::MoveTo(NewPosition, NodeFilter, bMarkDirty);
 }
 
 FVector2D SDisplayClusterConfiguratorCanvasNode::ComputeDesiredSize(float) const
diff --git a/Engine/Plugins/Runtime/nDisplay/Source/DisplayClusterConfigurator/Private/Views/OutputMapping/GraphNodes/SDisplayClusterConfiguratorCanvasNode.h b/Engine/Plugins/Runtime/nDisplay/Source/DisplayClusterConfigurator/Private/Views/OutputMapping/GraphNodes/SDisplayClusterConfiguratorCanvasNode.h
--- a/Engine/Plugins/Runtime/nDisplay/Source/DisplayClusterConfigurator/Private/Views/OutputMapping/GraphNodes/SDisplayClusterConfiguratorCanvasNode.h
+++ b/Engine/Plugins/Runtime/nDisplay/Source/DisplayClusterConfigurator/Private/Views/OutputMapping/GraphNodes/SDisplayClusterConfiguratorCanvasNode.h
@@ -25,7 +25,7 @@ public:
 
 	//~ SGraphNode interface
 	virtual void UpdateGraphNode() override;
-	virtual void MoveTo(const FVector2D& NewPosition, FNodeSet& NodeFilter) override;
+	virtual void MoveTo(const FVector2D& NewPosition, FNodeSet& NodeFilter, bool bMarkDirty = true) override;
 	virtual FVector2D ComputeDesiredSize(float) const override;
 	virtual FVector2D GetPosition() const override;
 	virtual TArray<FOverlayWidgetInfo> GetOverlayWidgets(bool bSelected, const FVector2D& WidgetSize) const override;
diff --git a/Engine/Plugins/Runtime/nDisplay/Source/DisplayClusterConfigurator/Private/Views/OutputMapping/GraphNodes/SDisplayClusterConfiguratorHostNode.cpp b/Engine/Plugins/Runtime/nDisplay/Source/DisplayClusterConfigurator/Private/Views/OutputMapping/GraphNodes/SDisplayClusterConfiguratorHostNode.cpp
--- a/Engine/Plugins/Runtime/nDisplay/Source/DisplayClusterConfigurator/Private/Views/OutputMapping/GraphNodes/SDisplayClusterConfiguratorHostNode.cpp
+++ b/Engine/Plugins/Runtime/nDisplay/Source/DisplayClusterConfigurator/Private/Views/OutputMapping/GraphNodes/SDisplayClusterConfiguratorHostNode.cpp
@@ -379,7 +379,7 @@ void SDisplayClusterConfiguratorHostNode::UpdateGraphNode()
 	];
 }
 
-void SDisplayClusterConfiguratorHostNode::MoveTo(const FVector2D& NewPosition, FNodeSet& NodeFilter)
+void SDisplayClusterConfiguratorHostNode::MoveTo(const FVector2D& NewPosition, FNodeSet& NodeFilter, bool bMarkDirty)
 {
 	UDisplayClusterConfiguratorHostNode* HostEdNode = GetGraphNodeChecked<UDisplayClusterConfiguratorHostNode>();
 	if (!HostEdNode->CanUserMoveNode())
@@ -387,7 +387,7 @@ void SDisplayClusterConfiguratorHostNode::MoveTo(const FVector2D& NewPosition, F
 		NodeFilter.Add(SharedThis(this));
 	}
 
-	SDisplayClusterConfiguratorBaseNode::MoveTo(NewPosition, NodeFilter);
+	SDisplayClusterConfiguratorBaseNode::MoveTo(NewPosition, NodeFilter, bMarkDirty);
 }
 
 FVector2D SDisplayClusterConfiguratorHostNode::ComputeDesiredSize(float) const
diff --git a/Engine/Plugins/Runtime/nDisplay/Source/DisplayClusterConfigurator/Private/Views/OutputMapping/GraphNodes/SDisplayClusterConfiguratorHostNode.h b/Engine/Plugins/Runtime/nDisplay/Source/DisplayClusterConfigurator/Private/Views/OutputMapping/GraphNodes/SDisplayClusterConfiguratorHostNode.h
--- a/Engine/Plugins/Runtime/nDisplay/Source/DisplayClusterConfigurator/Private/Views/OutputMapping/GraphNodes/SDisplayClusterConfiguratorHostNode.h
+++ b/Engine/Plugins/Runtime/nDisplay/Source/DisplayClusterConfigurator/Private/Views/OutputMapping/GraphNodes/SDisplayClusterConfiguratorHostNode.h
@@ -23,7 +23,7 @@ public:
 
 	//~ Begin SGraphNode interface
 	virtual void UpdateGraphNode() override;
-	virtual void MoveTo(const FVector2D& NewPosition, FNodeSet& NodeFilter) override;
+	virtual void MoveTo(const FVector2D& NewPosition, FNodeSet& NodeFilter, bool bMarkDirty = true) override;
 	virtual FVector2D ComputeDesiredSize(float) const override;
 	virtual FVector2D GetPosition() const override;
 	//~ End SGraphNode interface
diff --git a/Engine/Plugins/Runtime/nDisplay/Source/DisplayClusterConfigurator/Private/Views/OutputMapping/GraphNodes/SDisplayClusterConfiguratorViewportNode.cpp b/Engine/Plugins/Runtime/nDisplay/Source/DisplayClusterConfigurator/Private/Views/OutputMapping/GraphNodes/SDisplayClusterConfiguratorViewportNode.cpp
--- a/Engine/Plugins/Runtime/nDisplay/Source/DisplayClusterConfigurator/Private/Views/OutputMapping/GraphNodes/SDisplayClusterConfiguratorViewportNode.cpp
+++ b/Engine/Plugins/Runtime/nDisplay/Source/DisplayClusterConfigurator/Private/Views/OutputMapping/GraphNodes/SDisplayClusterConfiguratorViewportNode.cpp
@@ -152,14 +152,14 @@ void SDisplayClusterConfiguratorViewportNode::Tick(const FGeometry& AllottedGeom
 	UpdatePreviewTexture();
 }
 
-void SDisplayClusterConfiguratorViewportNode::MoveTo(const FVector2D& NewPosition, FNodeSet& NodeFilter)
+void SDisplayClusterConfiguratorViewportNode::MoveTo(const FVector2D& NewPosition, FNodeSet& NodeFilter, bool bMarkDirty)
 {
 	if (IsViewportLocked())
 	{
 		NodeFilter.Add(SharedThis(this));
 	}
 
-	SDisplayClusterConfiguratorBaseNode::MoveTo(NewPosition, NodeFilter);
+	SDisplayClusterConfiguratorBaseNode::MoveTo(NewPosition, NodeFilter, bMarkDirty);
 }
 
 bool SDisplayClusterConfiguratorViewportNode::IsNodeVisible() const
diff --git a/Engine/Plugins/Runtime/nDisplay/Source/DisplayClusterConfigurator/Private/Views/OutputMapping/GraphNodes/SDisplayClusterConfiguratorViewportNode.h b/Engine/Plugins/Runtime/nDisplay/Source/DisplayClusterConfigurator/Private/Views/OutputMapping/GraphNodes/SDisplayClusterConfiguratorViewportNode.h
--- a/Engine/Plugins/Runtime/nDisplay/Source/DisplayClusterConfigurator/Private/Views/OutputMapping/GraphNodes/SDisplayClusterConfiguratorViewportNode.h
+++ b/Engine/Plugins/Runtime/nDisplay/Source/DisplayClusterConfigurator/Private/Views/OutputMapping/GraphNodes/SDisplayClusterConfiguratorViewportNode.h
@@ -30,7 +30,7 @@ public:
 	//~ Begin SGraphNode interface
 	virtual void UpdateGraphNode() override;
 	virtual void Tick(const FGeometry& AllottedGeometry, const double InCurrentTime, const float InDeltaTime) override;
-	virtual void MoveTo(const FVector2D& NewPosition, FNodeSet& NodeFilter) override;
+	virtual void MoveTo(const FVector2D& NewPosition, FNodeSet& NodeFilter, bool bMarkDirty = true) override;
 	//~ End of SGraphNode interface
 
 	//~ Begin SDisplayClusterConfiguratorBaseNode interface
diff --git a/Engine/Plugins/Runtime/nDisplay/Source/DisplayClusterConfigurator/Private/Views/OutputMapping/GraphNodes/SDisplayClusterConfiguratorWindowNode.cpp b/Engine/Plugins/Runtime/nDisplay/Source/DisplayClusterConfigurator/Private/Views/OutputMapping/GraphNodes/SDisplayClusterConfiguratorWindowNode.cpp
--- a/Engine/Plugins/Runtime/nDisplay/Source/DisplayClusterConfigurator/Private/Views/OutputMapping/GraphNodes/SDisplayClusterConfiguratorWindowNode.cpp
+++ b/Engine/Plugins/Runtime/nDisplay/Source/DisplayClusterConfigurator/Private/Views/OutputMapping/GraphNodes/SDisplayClusterConfiguratorWindowNode.cpp
@@ -345,14 +345,14 @@ void SDisplayClusterConfiguratorWindowNode::UpdateGraphNode()
 	];
 }
 
-void SDisplayClusterConfiguratorWindowNode::MoveTo(const FVector2D& NewPosition, FNodeSet& NodeFilter)
+void SDisplayClusterConfiguratorWindowNode::MoveTo(const FVector2D& NewPosition, FNodeSet& NodeFilter, bool bMarkDirty)
 {
 	if (IsClusterNodeLocked())
 	{
 		NodeFilter.Add(SharedThis(this));
 	}
 
-	SDisplayClusterConfiguratorBaseNode::MoveTo(NewPosition, NodeFilter);
+	SDisplayClusterConfiguratorBaseNode::MoveTo(NewPosition, NodeFilter, bMarkDirty);
 }
 
 bool SDisplayClusterConfiguratorWindowNode::CanBeSelected(const FVector2D& MousePositionInNode) const
diff --git a/Engine/Plugins/Runtime/nDisplay/Source/DisplayClusterConfigurator/Private/Views/OutputMapping/GraphNodes/SDisplayClusterConfiguratorWindowNode.h b/Engine/Plugins/Runtime/nDisplay/Source/DisplayClusterConfigurator/Private/Views/OutputMapping/GraphNodes/SDisplayClusterConfiguratorWindowNode.h
--- a/Engine/Plugins/Runtime/nDisplay/Source/DisplayClusterConfigurator/Private/Views/OutputMapping/GraphNodes/SDisplayClusterConfiguratorWindowNode.h
+++ b/Engine/Plugins/Runtime/nDisplay/Source/DisplayClusterConfigurator/Private/Views/OutputMapping/GraphNodes/SDisplayClusterConfiguratorWindowNode.h
@@ -33,7 +33,7 @@ public:
 
 	//~ Begin SGraphNode interface
 	virtual void UpdateGraphNode() override;
-	virtual void MoveTo(const FVector2D& NewPosition, FNodeSet& NodeFilter) override;
+	virtual void MoveTo(const FVector2D& NewPosition, FNodeSet& NodeFilter, bool bMarkDirty = true) override;
 	virtual bool CanBeSelected(const FVector2D& MousePositionInNode) const override;
 	virtual FVector2D ComputeDesiredSize(float) const override;
 	virtual FVector2D GetPosition() const override;
diff --git a/Engine/Source/Editor/AnimationBlueprintEditor/Private/AnimationStateNodes/SGraphNodeAnimTransition.cpp b/Engine/Source/Editor/AnimationBlueprintEditor/Private/AnimationStateNodes/SGraphNodeAnimTransition.cpp
--- a/Engine/Source/Editor/AnimationBlueprintEditor/Private/AnimationStateNodes/SGraphNodeAnimTransition.cpp
+++ b/Engine/Source/Editor/AnimationBlueprintEditor/Private/AnimationStateNodes/SGraphNodeAnimTransition.cpp
@@ -37,7 +37,7 @@ void SGraphNodeAnimTransition::GetNodeInfoPopups(FNodeInfoContext* Context, TArr
 {
 }
 
-void SGraphNodeAnimTransition::MoveTo(const FVector2D& NewPosition, FNodeSet& NodeFilter)
+void SGraphNodeAnimTransition::MoveTo(const FVector2D& NewPosition, FNodeSet& NodeFilter, bool bMarkDirty)
 {
 	// Ignored; position is set by the location of the attached state nodes
 }
diff --git a/Engine/Source/Editor/AnimationBlueprintEditor/Private/AnimationStateNodes/SGraphNodeAnimTransition.h b/Engine/Source/Editor/AnimationBlueprintEditor/Private/AnimationStateNodes/SGraphNodeAnimTransition.h
--- a/Engine/Source/Editor/AnimationBlueprintEditor/Private/AnimationStateNodes/SGraphNodeAnimTransition.h
+++ b/Engine/Source/Editor/AnimationBlueprintEditor/Private/AnimationStateNodes/SGraphNodeAnimTransition.h
@@ -22,7 +22,7 @@ public:
 
 	// SNodePanel::SNode interface
 	virtual void GetNodeInfoPopups(FNodeInfoContext* Context, TArray<FGraphInformationPopupInfo>& Popups) const override;
-	virtual void MoveTo(const FVector2D& NewPosition, FNodeSet& NodeFilter) override;
+	virtual void MoveTo(const FVector2D& NewPosition, FNodeSet& NodeFilter, bool bMarkDirty = true) override;
 	virtual bool RequiresSecondPassLayout() const override;
 	virtual void PerformSecondPassLayout(const TMap< UObject*, TSharedRef<SNode> >& NodeToWidgetLookup) const override;
 	// End of SNodePanel::SNode interface
diff --git a/Engine/Source/Editor/BehaviorTreeEditor/Private/SGraphNode_BehaviorTree.cpp b/Engine/Source/Editor/BehaviorTreeEditor/Private/SGraphNode_BehaviorTree.cpp
--- a/Engine/Source/Editor/BehaviorTreeEditor/Private/SGraphNode_BehaviorTree.cpp
+++ b/Engine/Source/Editor/BehaviorTreeEditor/Private/SGraphNode_BehaviorTree.cpp
@@ -1097,9 +1097,9 @@ TSharedRef<SGraphNode> SGraphNode_BehaviorTree::GetNodeUnderMouse(const FGeometr
 	return SubNode.IsValid() ? SubNode.ToSharedRef() : StaticCastSharedRef<SGraphNode>(AsShared());
 }
 
-void SGraphNode_BehaviorTree::MoveTo(const FVector2D& NewPosition, FNodeSet& NodeFilter)
+void SGraphNode_BehaviorTree::MoveTo(const FVector2D& NewPosition, FNodeSet& NodeFilter, bool bMarkDirty)
 {
-	SGraphNodeAI::MoveTo(NewPosition, NodeFilter);
+	SGraphNodeAI::MoveTo(NewPosition, NodeFilter, bMarkDirty);
 
 	// keep node order (defined by linked pins) up to date with actual positions
 	// this function will keep spamming on every mouse move update
diff --git a/Engine/Source/Editor/BehaviorTreeEditor/Private/SGraphNode_BehaviorTree.h b/Engine/Source/Editor/BehaviorTreeEditor/Private/SGraphNode_BehaviorTree.h
--- a/Engine/Source/Editor/BehaviorTreeEditor/Private/SGraphNode_BehaviorTree.h
+++ b/Engine/Source/Editor/BehaviorTreeEditor/Private/SGraphNode_BehaviorTree.h
@@ -34,7 +34,7 @@ public:
 	virtual void GetOverlayBrushes(bool bSelected, const FVector2D WidgetSize, TArray<FOverlayBrushInfo>& Brushes) const override;
 	virtual TArray<FOverlayWidgetInfo> GetOverlayWidgets(bool bSelected, const FVector2D& WidgetSize) const override;
 	virtual TSharedRef<SGraphNode> GetNodeUnderMouse(const FGeometry& MyGeometry, const FPointerEvent& MouseEvent) override;
-	virtual void MoveTo(const FVector2D& NewPosition, FNodeSet& NodeFilter) override;
+	virtual void MoveTo(const FVector2D& NewPosition, FNodeSet& NodeFilter, bool bMarkDirty = true) override;
 	// End of SGraphNode interface
 
 	/** handle double click */
diff --git a/Engine/Source/Editor/GraphEditor/Private/MaterialNodes/SGraphNodeMaterialBase.cpp b/Engine/Source/Editor/GraphEditor/Private/MaterialNodes/SGraphNodeMaterialBase.cpp
--- a/Engine/Source/Editor/GraphEditor/Private/MaterialNodes/SGraphNodeMaterialBase.cpp
+++ b/Engine/Source/Editor/GraphEditor/Private/MaterialNodes/SGraphNodeMaterialBase.cpp
@@ -296,9 +296,9 @@ void SGraphNodeMaterialBase::CreatePinWidgets()
 	}
 }
 
-void SGraphNodeMaterialBase::MoveTo(const FVector2D& NewPosition, FNodeSet& NodeFilter)
+void SGraphNodeMaterialBase::MoveTo(const FVector2D& NewPosition, FNodeSet& NodeFilter, bool bMarkDirty)
 {
-	SGraphNode::MoveTo(NewPosition, NodeFilter);
+	SGraphNode::MoveTo(NewPosition, NodeFilter, bMarkDirty);
 
 	MaterialNode->MaterialExpression->MaterialExpressionEditorX = MaterialNode->NodePosX;
 	MaterialNode->MaterialExpression->MaterialExpressionEditorY = MaterialNode->NodePosY;
diff --git a/Engine/Source/Editor/GraphEditor/Private/MaterialNodes/SGraphNodeMaterialBase.h b/Engine/Source/Editor/GraphEditor/Private/MaterialNodes/SGraphNodeMaterialBase.h
--- a/Engine/Source/Editor/GraphEditor/Private/MaterialNodes/SGraphNodeMaterialBase.h
+++ b/Engine/Source/Editor/GraphEditor/Private/MaterialNodes/SGraphNodeMaterialBase.h
@@ -92,7 +92,7 @@ public:
 	// End of SGraphNode interface
 
 	// SNodePanel::SNode interface
-	virtual void MoveTo(const FVector2D& NewPosition, FNodeSet& NodeFilter) override;
+	virtual void MoveTo(const FVector2D& NewPosition, FNodeSet& NodeFilter, bool bMarkDirty = true) override;
 	// End of SNodePanel::SNode interface
 
 	UMaterialGraphNode* GetMaterialGraphNode() const {return MaterialNode;}
diff --git a/Engine/Source/Editor/GraphEditor/Private/MaterialNodes/SGraphNodeMaterialComment.cpp b/Engine/Source/Editor/GraphEditor/Private/MaterialNodes/SGraphNodeMaterialComment.cpp
--- a/Engine/Source/Editor/GraphEditor/Private/MaterialNodes/SGraphNodeMaterialComment.cpp
+++ b/Engine/Source/Editor/GraphEditor/Private/MaterialNodes/SGraphNodeMaterialComment.cpp
@@ -11,11 +11,11 @@ void SGraphNodeMaterialComment::Construct(const FArguments& InArgs, class UMater
 	this->CommentNode = InNode;
 }
 
-void SGraphNodeMaterialComment::MoveTo(const FVector2D& NewPosition, FNodeSet& NodeFilter)
+void SGraphNodeMaterialComment::MoveTo(const FVector2D& NewPosition, FNodeSet& NodeFilter, bool bMarkDirty)
 {
 	if( !NodeFilter.Find( SharedThis( this ) ))
 	{
-		SGraphNodeComment::MoveTo(NewPosition, NodeFilter);
+		SGraphNodeComment::MoveTo(NewPosition, NodeFilter, bMarkDirty);
 		if (CommentNode && CommentNode->MaterialExpressionComment)
 		{
 			CommentNode->MaterialExpressionComment->MaterialExpressionEditorX = CommentNode->NodePosX;
diff --git a/Engine/Source/Editor/GraphEditor/Private/MaterialNodes/SGraphNodeMaterialComment.h b/Engine/Source/Editor/GraphEditor/Private/MaterialNodes/SGraphNodeMaterialComment.h
--- a/Engine/Source/Editor/GraphEditor/Private/MaterialNodes/SGraphNodeMaterialComment.h
+++ b/Engine/Source/Editor/GraphEditor/Private/MaterialNodes/SGraphNodeMaterialComment.h
@@ -17,7 +17,7 @@ public:
 	void Construct(const FArguments& InArgs, class UMaterialGraphNode_Comment* InNode);
 
 	// SNodePanel::SNode interface
-	virtual void MoveTo(const FVector2D& NewPosition, FNodeSet& NodeFilter) override;
+	virtual void MoveTo(const FVector2D& NewPosition, FNodeSet& NodeFilter, bool bMarkDirty = true) override;
 	// End of SNodePanel::SNode interface
 
 private:
diff --git a/Engine/Source/Editor/GraphEditor/Private/MaterialNodes/SGraphNodeMaterialComposite.cpp b/Engine/Source/Editor/GraphEditor/Private/MaterialNodes/SGraphNodeMaterialComposite.cpp
--- a/Engine/Source/Editor/GraphEditor/Private/MaterialNodes/SGraphNodeMaterialComposite.cpp
+++ b/Engine/Source/Editor/GraphEditor/Private/MaterialNodes/SGraphNodeMaterialComposite.cpp
@@ -31,9 +31,9 @@ void SGraphNodeMaterialComposite::Construct(const FArguments& InArgs, class UMat
 	this->UpdateGraphNode();
 }
 
-void SGraphNodeMaterialComposite::MoveTo(const FVector2D& NewPosition, FNodeSet& NodeFilter)
+void SGraphNodeMaterialComposite::MoveTo(const FVector2D& NewPosition, FNodeSet& NodeFilter, bool bMarkDirty)
 {
-	SGraphNode::MoveTo(NewPosition, NodeFilter);
+	SGraphNode::MoveTo(NewPosition, NodeFilter, bMarkDirty);
 
 	CompositeNode->MaterialExpression->MaterialExpressionEditorX = CompositeNode->NodePosX;
 	CompositeNode->MaterialExpression->MaterialExpressionEditorY = CompositeNode->NodePosY;
diff --git a/Engine/Source/Editor/GraphEditor/Private/MaterialNodes/SGraphNodeMaterialComposite.h b/Engine/Source/Editor/GraphEditor/Private/MaterialNodes/SGraphNodeMaterialComposite.h
--- a/Engine/Source/Editor/GraphEditor/Private/MaterialNodes/SGraphNodeMaterialComposite.h
+++ b/Engine/Source/Editor/GraphEditor/Private/MaterialNodes/SGraphNodeMaterialComposite.h
@@ -26,7 +26,7 @@ public:
 	// End of SGraphNode interface
 
 	// SNodePanel::SNode interface
-	virtual void MoveTo(const FVector2D& NewPosition, FNodeSet& NodeFilter) override;
+	virtual void MoveTo(const FVector2D& NewPosition, FNodeSet& NodeFilter, bool bMarkDirty = true) override;
 	// End of SNodePanel::SNode interface
 protected:
 	virtual UEdGraph* GetInnerGraph() const;
diff --git a/Engine/Source/Editor/GraphEditor/Private/MaterialNodes/SGraphNodeMaterialResult.cpp b/Engine/Source/Editor/GraphEditor/Private/MaterialNodes/SGraphNodeMaterialResult.cpp
--- a/Engine/Source/Editor/GraphEditor/Private/MaterialNodes/SGraphNodeMaterialResult.cpp
+++ b/Engine/Source/Editor/GraphEditor/Private/MaterialNodes/SGraphNodeMaterialResult.cpp
@@ -64,9 +64,9 @@ void SGraphNodeMaterialResult::CreatePinWidgets()
 	}
 }
 
-void SGraphNodeMaterialResult::MoveTo(const FVector2D& NewPosition, FNodeSet& NodeFilter)
+void SGraphNodeMaterialResult::MoveTo(const FVector2D& NewPosition, FNodeSet& NodeFilter, bool bMarkDirty)
 {
-	SGraphNode::MoveTo(NewPosition, NodeFilter);
+	SGraphNode::MoveTo(NewPosition, NodeFilter, bMarkDirty);
 
 	RootNode->Material->EditorX = RootNode->NodePosX;
 	RootNode->Material->EditorY = RootNode->NodePosY;
diff --git a/Engine/Source/Editor/GraphEditor/Private/MaterialNodes/SGraphNodeMaterialResult.h b/Engine/Source/Editor/GraphEditor/Private/MaterialNodes/SGraphNodeMaterialResult.h
--- a/Engine/Source/Editor/GraphEditor/Private/MaterialNodes/SGraphNodeMaterialResult.h
+++ b/Engine/Source/Editor/GraphEditor/Private/MaterialNodes/SGraphNodeMaterialResult.h
@@ -23,7 +23,7 @@ public:
 	// End of SGraphNode interface
 
 	// SNodePanel::SNode interface
-	virtual void MoveTo(const FVector2D& NewPosition, FNodeSet& NodeFilter) override;
+	virtual void MoveTo(const FVector2D& NewPosition, FNodeSet& NodeFilter, bool bMarkDirty = true) override;
 	// End of SNodePanel::SNode interface
 
 private:
diff --git a/Engine/Source/Editor/GraphEditor/Private/SGraphNode.cpp b/Engine/Source/Editor/GraphEditor/Private/SGraphNode.cpp
--- a/Engine/Source/Editor/GraphEditor/Private/SGraphNode.cpp
+++ b/Engine/Source/Editor/GraphEditor/Private/SGraphNode.cpp
@@ -512,14 +512,14 @@ void SGraphNode::SetOwner( const TSharedRef<SGraphPanel>& OwnerPanel )
 }
 
 /** @param NewPosition  The Node should be relocated to this position in the graph panel */
-void SGraphNode::MoveTo( const FVector2D& NewPosition, FNodeSet& NodeFilter )
+void SGraphNode::MoveTo( const FVector2D& NewPosition, FNodeSet& NodeFilter, bool bMarkDirty )
 {
 	if ( !NodeFilter.Find( SharedThis( this )))
 	{
 		if (GraphNode && !RequiresSecondPassLayout())
 		{
 			NodeFilter.Add( SharedThis( this ) );
-			GraphNode->Modify();
+			GraphNode->Modify(bMarkDirty);
 			GraphNode->NodePosX = NewPosition.X;
 			GraphNode->NodePosY = NewPosition.Y;
 		}
diff --git a/Engine/Source/Editor/GraphEditor/Private/SGraphNodeComment.cpp b/Engine/Source/Editor/GraphEditor/Private/SGraphNodeComment.cpp
--- a/Engine/Source/Editor/GraphEditor/Private/SGraphNodeComment.cpp
+++ b/Engine/Source/Editor/GraphEditor/Private/SGraphNodeComment.cpp
@@ -357,10 +357,10 @@ void SGraphNodeComment::GetOverlayBrushes(bool bSelected, const FVector2D Widget
 	return SGraphNode::GetOverlayBrushes(bSelected, WidgetSize, Brushes);
 }
 
-void SGraphNodeComment::MoveTo( const FVector2D& NewPosition, FNodeSet& NodeFilter ) 
+void SGraphNodeComment::MoveTo( const FVector2D& NewPosition, FNodeSet& NodeFilter, bool bMarkDirty)
 {
 	FVector2D PositionDelta = NewPosition - GetPosition();
-	SGraphNode::MoveTo(NewPosition, NodeFilter);
+	SGraphNode::MoveTo(NewPosition, NodeFilter, bMarkDirty);
 	// Don't drag note content if either of the shift keys are down.
 	FModifierKeysState KeysState = FSlateApplication::Get().GetModifierKeys();
 	if(!KeysState.IsShiftDown())
@@ -379,7 +379,7 @@ void SGraphNodeComment::MoveTo( const FVector2D& NewPosition, FNodeSet& NodeFilt
 					if (!Panel->SelectionManager.IsNodeSelected(Node) && !NodeFilter.Find(Node->DEPRECATED_NodeWidget.Pin()))
 					{
 						NodeFilter.Add(Node->DEPRECATED_NodeWidget.Pin());
-						Node->Modify();
+						Node->Modify(bMarkDirty);
 						Node->NodePosX += PositionDelta.X;
 						Node->NodePosY += PositionDelta.Y;
 					}
diff --git a/Engine/Source/Editor/GraphEditor/Private/SNodePanel.cpp b/Engine/Source/Editor/GraphEditor/Private/SNodePanel.cpp
--- a/Engine/Source/Editor/GraphEditor/Private/SNodePanel.cpp
+++ b/Engine/Source/Editor/GraphEditor/Private/SNodePanel.cpp
@@ -776,7 +776,8 @@ FReply SNodePanel::OnMouseMove(const FGeometry& MyGeometry, const FPointerEvent&
 							const FVector2D AnchorNodeOldPos = NodeBeingDragged->GetPosition();
 							const FVector2D DeltaPos = AnchorNodeNewPos - AnchorNodeOldPos;
 
-							// Perform movement in 2 passes:
+							// Perform movement in 3 passes:
+							
 							// 1. Gather all selected nodes positions and calculate new positions
 							struct FDefferedNodePosition 
 							{ 
@@ -784,37 +785,30 @@ FReply SNodePanel::OnMouseMove(const FGeometry& MyGeometry, const FPointerEvent&
 								FVector2D	NewPosition; 
 							};
 							TArray<FDefferedNodePosition> DefferedNodesToMove;
-							
+
+							// 2. Deffer actual move transactions to mouse release or focus lost
+							bool bStoreOriginalNodePositions = OriginalNodePositions.Num() == 0;
 							for (FGraphPanelSelectionSet::TIterator NodeIt(SelectionManager.SelectedNodes); NodeIt; ++NodeIt)
 							{
-								TSharedRef<SNode>* pWidget = NodeToWidgetLookup.Find(*NodeIt);
-								if (pWidget != nullptr)
+								if (TSharedRef<SNode>* pWidget = NodeToWidgetLookup.Find(*NodeIt))
 								{
 									SNode& Widget = pWidget->Get();
 									FDefferedNodePosition NodePosition = { &Widget, Widget.GetPosition() + DeltaPos };
 									DefferedNodesToMove.Add(NodePosition);
-								}
-							}
 
-							// Create a new transaction record
-							if(!ScopedTransactionPtr.IsValid())
-							{
-								if(DefferedNodesToMove.Num() > 1)
-								{
-									ScopedTransactionPtr = MakeShareable(new FScopedTransaction(NSLOCTEXT("GraphEditor", "MoveNodesAction", "Move Nodes")));
-								}
-								else if(DefferedNodesToMove.Num() > 0)
-								{
-									ScopedTransactionPtr = MakeShareable(new FScopedTransaction(NSLOCTEXT("GraphEditor", "MoveNodeAction", "Move Node")));
+									if (bStoreOriginalNodePositions)
+									{
+										OriginalNodePositions.FindOrAdd(*pWidget) = Widget.GetPosition();
+									}
 								}
 							}
 
-							// 2. Move selected nodes to new positions
+							// 3. Move selected nodes to new positions
 							SNode::FNodeSet NodeFilter;
 
 							for (int32 NodeIdx = 0; NodeIdx < DefferedNodesToMove.Num(); ++NodeIdx)
 							{
-								DefferedNodesToMove[NodeIdx].Node->MoveTo( DefferedNodesToMove[NodeIdx].NewPosition, NodeFilter );
+								DefferedNodesToMove[NodeIdx].Node->MoveTo(DefferedNodesToMove[NodeIdx].NewPosition, NodeFilter, false);
 							}
 						}
 					}
@@ -909,6 +903,7 @@ FReply SNodePanel::OnMouseButtonUp( const FGeometry& MyGeometry, const FPointerE
 		{
 			OnEndNodeInteraction(NodeUnderMousePtr.Pin().ToSharedRef());
 
+			FinalizeNodeMovements();
 			ScopedTransactionPtr.Reset();
 		}
 				
@@ -983,6 +978,16 @@ FReply SNodePanel::OnMouseButtonUp( const FGeometry& MyGeometry, const FPointerE
 	return ReplyState;	
 }
 
+void SNodePanel::OnMouseCaptureLost(const FCaptureLostEvent& CaptureLostEvent)
+{
+	if (OriginalNodePositions.Num() > 0)
+	{
+		FinalizeNodeMovements();
+	}
+
+	SPanel::OnMouseCaptureLost(CaptureLostEvent);
+}
+
 FReply SNodePanel::OnMouseWheel(const FGeometry& MyGeometry, const FPointerEvent& MouseEvent)
 {
 	// We want to zoom into this point; i.e. keep it the same fraction offset into the panel
@@ -1711,6 +1716,55 @@ bool SNodePanel::HasDeferredObjectFocus() const
 	return DeferredMovementTargetObject != nullptr;
 }
 
+void SNodePanel::FinalizeNodeMovements()
+{
+	// Process moved nodes on focus lost
+	if (OriginalNodePositions.Num() > 0)
+	{
+		// Build up all the current positions
+		TMap<SNode*, FVector2D> CurrentNodePositions;
+
+		for (FGraphPanelSelectionSet::TIterator NodeIt(SelectionManager.SelectedNodes); NodeIt; ++NodeIt)
+		{
+			TSharedRef<SNode>* pWidget = NodeToWidgetLookup.Find(*NodeIt);
+			if (pWidget != nullptr)
+			{
+				SNode& Widget = pWidget->Get();
+				CurrentNodePositions.FindOrAdd(&Widget) = Widget.GetPosition();
+			}
+		}
+
+		// Move all the nodes back to their original position before we start the transaction
+		SNode::FNodeSet OriginalNodeFilter;
+
+		for (TPair<TWeakPtr<SNode>, FVector2D>& OriginalNodePosition : OriginalNodePositions)
+		{
+			if (OriginalNodePosition.Key.IsValid())
+			{
+				OriginalNodePosition.Key.Pin()->MoveTo(OriginalNodePosition.Value, OriginalNodeFilter, false);
+			}
+		}
+
+		OriginalNodePositions.Reset();
+
+		if (CurrentNodePositions.Num() > 0)
+		{
+			FScopedTransaction NodeMoveTransaction(
+				CurrentNodePositions.Num() > 1 
+				? NSLOCTEXT("GraphEditor", "MoveNodesAction", "Move Nodes")
+				: NSLOCTEXT("GraphEditor", "MoveNodeAction", "Move Node"));
+
+			// Move all the nodes back to their current position but on on the undo stack
+			SNode::FNodeSet CurrentNodeFilter;
+
+			for (TPair<SNode*, FVector2D>& CurrentNodePosition : CurrentNodePositions)
+			{
+				CurrentNodePosition.Key->MoveTo(CurrentNodePosition.Value, CurrentNodeFilter, true);
+			}
+		}
+	}
+}
+
 void SNodePanel::PostChangedZoom()
 {
 	CurrentLOD = ZoomLevels->GetLOD(ZoomLevel);
diff --git a/Engine/Source/Editor/GraphEditor/Public/SGraphNode.h b/Engine/Source/Editor/GraphEditor/Public/SGraphNode.h
--- a/Engine/Source/Editor/GraphEditor/Public/SGraphNode.h
+++ b/Engine/Source/Editor/GraphEditor/Public/SGraphNode.h
@@ -108,7 +108,7 @@ public:
 	// End of SWidget interface
 
 	// SNodePanel::SNode interface
-	virtual void MoveTo( const FVector2D& NewPosition, FNodeSet& NodeFilter ) override;
+	virtual void MoveTo( const FVector2D& NewPosition, FNodeSet& NodeFilter, bool bMarkDirty = true ) override;
 	virtual FVector2D GetPosition() const override;
 	virtual FString GetNodeComment() const override;
 	virtual UObject* GetObjectBeingDisplayed() const override;
diff --git a/Engine/Source/Editor/GraphEditor/Public/SGraphNodeComment.h b/Engine/Source/Editor/GraphEditor/Public/SGraphNodeComment.h
--- a/Engine/Source/Editor/GraphEditor/Public/SGraphNodeComment.h
+++ b/Engine/Source/Editor/GraphEditor/Public/SGraphNodeComment.h
@@ -72,7 +72,7 @@ protected:
 	virtual bool IsNodeUnderComment(UEdGraphNode_Comment* InCommentNode, const TSharedRef<SGraphNode> InNodeWidget) const;
 
 	/** called when user is moving the comment node */
-	virtual void MoveTo(const FVector2D& NewPosition, FNodeSet& NodeFilter ) override;
+	virtual void MoveTo(const FVector2D& NewPosition, FNodeSet& NodeFilter, bool bMarkDirty = true) override;
 
 	//~ Begin SGraphNodeResizable Interface
 	virtual float GetTitleBarHeight() const override;
diff --git a/Engine/Source/Editor/GraphEditor/Public/SNodePanel.h b/Engine/Source/Editor/GraphEditor/Public/SNodePanel.h
--- a/Engine/Source/Editor/GraphEditor/Public/SNodePanel.h
+++ b/Engine/Source/Editor/GraphEditor/Public/SNodePanel.h
@@ -470,8 +470,9 @@ public:
 		/**
 		* @param NewPosition	The Node should be relocated to this position in the graph panel
 		* @param NodeFilter		Set of nodes to prevent movement on, after moving successfully a node is added to this set.
+		* @param bMarkDirty		If we should mark nodes as dirty on move
 		*/
-		virtual void MoveTo(const FVector2D& NewPosition, FNodeSet& NodeFilter )
+		virtual void MoveTo(const FVector2D& NewPosition, FNodeSet& NodeFilter, bool bMarkDirty = true)
 		{
 		}
 
@@ -623,6 +624,7 @@ public:
 	virtual void Tick( const FGeometry& AllottedGeometry, const double InCurrentTime, const float InDeltaTime ) override;
 	virtual FReply OnMouseButtonDown(const FGeometry& MyGeometry, const FPointerEvent& MouseEvent) override;
 	virtual FReply OnMouseButtonUp(const FGeometry& MyGeometry, const FPointerEvent& MouseEvent) override;
+	virtual void OnMouseCaptureLost(const FCaptureLostEvent& CaptureLostEvent) override;
 	virtual FReply OnMouseMove(const FGeometry& MyGeometry, const FPointerEvent& MouseEvent) override;
 	virtual FReply OnMouseWheel(const FGeometry& MyGeometry, const FPointerEvent& MouseEvent) override;
 	virtual FCursorReply OnCursorQuery( const FGeometry& MyGeometry, const FPointerEvent& CursorEvent ) const override;
@@ -683,6 +685,9 @@ public:
 	/** If it is focusing on a particular object */
 	bool HasDeferredObjectFocus() const;
 
+	/** Commit transactions for any node movements */
+	void FinalizeNodeMovements();
+
 	/** Returns the current LOD level of this panel, based on the zoom factor */
 	EGraphRenderingLOD::Type GetCurrentLOD() const { return CurrentLOD; }
 
@@ -965,6 +970,9 @@ protected:
 	/** A flag to detect when a visual update is pending to prevent deferred
 	    commands like zoom to fit from running when there are no widgets */
 	bool bVisualUpdatePending;
+	
+	/** Node positions pre-drag, used to limit transaction creation on drag */
+	TMap<TWeakPtr<SNode>, FVector2D> OriginalNodePositions;
 
 private:
 	/** Active timer that handles deferred zooming until the target zoom is reached */
