Slate: Fix missing slot initialization.

diff --git a/Engine/Source/Editor/GraphEditor/Private/SNodePanel.cpp b/Engine/Source/Editor/GraphEditor/Private/SNodePanel.cpp
--- a/Engine/Source/Editor/GraphEditor/Private/SNodePanel.cpp
+++ b/Engine/Source/Editor/GraphEditor/Private/SNodePanel.cpp
@@ -204,6 +204,30 @@ namespace NodePanelDefs
 	static const float MouseZoomScaling = 0.04f;
 };
 
+
+void SNodePanel::SNode::FNodeSlot::Construct(const FChildren& SlotOwner, FSlotArguments&& InArgs)
+{
+	TSlotBase<FNodeSlot>::Construct(SlotOwner, MoveTemp(InArgs));
+	TAlignmentWidgetSlotMixin<FNodeSlot>::ConstructMixin(SlotOwner, MoveTemp(InArgs));
+
+	if (InArgs._Padding.IsSet())
+	{
+		SlotPadding = MoveTemp(InArgs._Padding);
+	}
+	if (InArgs._SlotOffset.IsSet())
+	{
+		Offset = MoveTemp(InArgs._SlotOffset);
+	}
+	if (InArgs._SlotSize.IsSet())
+	{
+		Size = MoveTemp(InArgs._SlotSize);
+	}
+	if (InArgs._AllowScaling.IsSet())
+	{
+		AllowScale = MoveTemp(InArgs._AllowScaling);
+	}
+}
+
 SNodePanel::SNodePanel()
 	: Children(this)
 	, VisibleChildren(this)
diff --git a/Engine/Source/Editor/GraphEditor/Public/SNodePanel.h b/Engine/Source/Editor/GraphEditor/Public/SNodePanel.h
--- a/Engine/Source/Editor/GraphEditor/Public/SNodePanel.h
+++ b/Engine/Source/Editor/GraphEditor/Public/SNodePanel.h
@@ -275,7 +275,7 @@ public:
 	public:
 
 		/** A slot that support alignment of content and padding and z-order */
-		class FNodeSlot : public TSlotBase<FNodeSlot>, public TAlignmentWidgetSlotMixin<FNodeSlot>
+		class GRAPHEDITOR_API FNodeSlot : public TSlotBase<FNodeSlot>, public TAlignmentWidgetSlotMixin<FNodeSlot>
 		{
 		public:
 			friend SNode;
@@ -300,11 +300,7 @@ public:
 				SLATE_ATTRIBUTE(bool, AllowScaling)
 			SLATE_SLOT_END_ARGS()
 
-			void Construct(const FChildren& SlotOwner, FSlotArguments&& InArgs)
-			{
-				TSlotBase<FNodeSlot>::Construct(SlotOwner, MoveTemp(InArgs));
-				TAlignmentWidgetSlotMixin<FNodeSlot>::ConstructMixin(SlotOwner, MoveTemp(InArgs));
-			}
+			void Construct(const FChildren& SlotOwner, FSlotArguments&& InArgs);
 
 		public:
 			UE_DEPRECATED(5.0, "Padding is now deprecated. Use the FSlotArgument or the SetPadding function.")
