[NO NON-PII COMMIT MESSAGE]

diff --git a/Engine/Source/Editor/GraphEditor/Private/SGraphNodeComment.cpp b/Engine/Source/Editor/GraphEditor/Private/SGraphNodeComment.cpp
--- a/Engine/Source/Editor/GraphEditor/Private/SGraphNodeComment.cpp
+++ b/Engine/Source/Editor/GraphEditor/Private/SGraphNodeComment.cpp
@@ -145,7 +145,6 @@ void SGraphNodeComment::UpdateGraphNode()
 	CommentStyle.EditableTextBoxStyle.Font.Size = CachedFontSize;
 	CommentStyle.TextStyle.Font.Size = CachedFontSize;
 
-	bool bIsSet = GraphNode->IsA(UEdGraphNode_Comment::StaticClass());
 	this->ContentScale.Bind( this, &SGraphNode::GetContentScale );
 	this->GetOrAddSlot( ENodeZone::Center )
 		.HAlign(HAlign_Fill)
@@ -199,7 +198,7 @@ void SGraphNodeComment::UpdateGraphNode()
 					SNew(SBorder)
 					.BorderImage( FEditorStyle::GetBrush("NoBorder") )
 				]
-			]			
+			]
 		];
 
 	// Create comment bubble
@@ -265,7 +264,7 @@ FReply SGraphNodeComment::OnMouseButtonUp( const FGeometry& MyGeometry, const FP
 	{
 		bUserIsDragging = false;
 
-		// Resize the node	
+		// Resize the node
 		UserSize.X = FMath::RoundToFloat(UserSize.X);
 		UserSize.Y = FMath::RoundToFloat(UserSize.Y);
 
