Fix widget for material exec pins

diff --git a/Engine/Source/Editor/GraphEditor/Private/NodeFactory.cpp b/Engine/Source/Editor/GraphEditor/Private/NodeFactory.cpp
--- a/Engine/Source/Editor/GraphEditor/Private/NodeFactory.cpp
+++ b/Engine/Source/Editor/GraphEditor/Private/NodeFactory.cpp
@@ -251,7 +251,11 @@ TSharedPtr<SGraphPin> FNodeFactory::CreatePinWidget(UEdGraphPin* InPin)
 
 	if (const UMaterialGraphSchema* MaterialGraphSchema = Cast<const UMaterialGraphSchema>(InPin->GetSchema()))
 	{
-		if (InPin->PinType.PinCategory == MaterialGraphSchema->PC_MaterialInput)
+		if (InPin->PinType.PinCategory == MaterialGraphSchema->PC_Exec)
+		{
+			return SNew(SGraphPinExec, InPin);
+		}
+		else if (InPin->PinType.PinCategory == MaterialGraphSchema->PC_MaterialInput)
 		{
 			return SNew(SGraphPinMaterialInput, InPin);
 		}
