Added "Break This Link" command to graph editor, which is displayed instead of the "Break Link..." submenu when there's only one link to display.

Fixed various Title Case errors in graph editor command names.

diff --git a/Engine/Source/Editor/GraphEditor/Private/GraphEditorActions.cpp b/Engine/Source/Editor/GraphEditor/Private/GraphEditorActions.cpp
--- a/Engine/Source/Editor/GraphEditor/Private/GraphEditorActions.cpp
+++ b/Engine/Source/Editor/GraphEditor/Private/GraphEditorActions.cpp
@@ -9,36 +9,36 @@ void FGraphEditorCommandsImpl::RegisterCommands()
 	UI_COMMAND( ReconstructNodes, "Refresh Nodes", "Refreshes nodes", EUserInterfaceActionType::Button, FInputChord() )
 	UI_COMMAND( BreakNodeLinks, "Break Node Link(s)", "Breaks the selected node from all connected pins", EUserInterfaceActionType::Button, FInputChord() )
 	
-	UI_COMMAND( AddExecutionPin, "Add execution pin", "Adds another execution output pin to an execution sequence or switch node", EUserInterfaceActionType::Button, FInputChord() )
-	UI_COMMAND( InsertExecutionPinBefore, "Insert execution pin before", "Adds another execution output pin before this one, to an execution sequence node", EUserInterfaceActionType::Button, FInputChord() )
-	UI_COMMAND( InsertExecutionPinAfter, "Insert execution pin after", "Adds another execution output pin after this one, to an execution sequence node", EUserInterfaceActionType::Button, FInputChord() )
-	UI_COMMAND( RemoveExecutionPin, "Remove execution pin", "Removes an execution output pin from an execution sequence or switch node", EUserInterfaceActionType::Button, FInputChord() )
+	UI_COMMAND( AddExecutionPin, "Add Execution Pin", "Adds another execution output pin to an execution sequence or switch node", EUserInterfaceActionType::Button, FInputChord() )
+	UI_COMMAND( InsertExecutionPinBefore, "Insert Execution Pin Before", "Adds another execution output pin before this one, to an execution sequence node", EUserInterfaceActionType::Button, FInputChord() )
+	UI_COMMAND( InsertExecutionPinAfter, "Insert Execution Pin After", "Adds another execution output pin after this one, to an execution sequence node", EUserInterfaceActionType::Button, FInputChord() )
+	UI_COMMAND( RemoveExecutionPin, "Remove Execution Pin", "Removes an execution output pin from an execution sequence or switch node", EUserInterfaceActionType::Button, FInputChord() )
 
-	UI_COMMAND( RemoveThisStructVarPin, "Remove this struct variable pin", "Removes the selected input pin", EUserInterfaceActionType::Button, FInputChord() )
-	UI_COMMAND( RemoveOtherStructVarPins, "Remove all other pins", "Removes all variable input pins, except for the selected one", EUserInterfaceActionType::Button, FInputChord() )
+	UI_COMMAND( RemoveThisStructVarPin, "Remove This Struct Variable Pin", "Removes the selected input pin", EUserInterfaceActionType::Button, FInputChord() )
+	UI_COMMAND( RemoveOtherStructVarPins, "Remove All Other Pins", "Removes all variable input pins, except for the selected one", EUserInterfaceActionType::Button, FInputChord() )
 
-	UI_COMMAND( RestoreAllStructVarPins, "Restore all structure pins", "Restore all structure pins", EUserInterfaceActionType::Button, FInputChord() )
+	UI_COMMAND( RestoreAllStructVarPins, "Restore All Structure Pins", "Restore all structure pins", EUserInterfaceActionType::Button, FInputChord() )
 
 	UI_COMMAND( AddOptionPin, "Add Option Pin", "Adds another option input pin to the node", EUserInterfaceActionType::Button, FInputChord() )
 	UI_COMMAND( RemoveOptionPin, "Remove Option Pin", "Removes the last option input pin from the node", EUserInterfaceActionType::Button, FInputChord() )
 
 	UI_COMMAND( ChangePinType, "Change Pin Type", "Changes the type of this pin (boolean, int, etc.)", EUserInterfaceActionType::Button, FInputChord() )
 
-	UI_COMMAND( DeleteAndReconnectNodes, "Delete and reconnect exec pins", "Deletes the currently selected nodes and makes connections between all input pins to their output pins.", EUserInterfaceActionType::Button, FInputChord(EModifierKey::Shift, EKeys::Delete), FInputChord(EModifierKey::Shift, EKeys::BackSpace))
+	UI_COMMAND( DeleteAndReconnectNodes, "Delete and Reconnect Exec Pins", "Deletes the currently selected nodes and makes connections between all input pins to their output pins.", EUserInterfaceActionType::Button, FInputChord(EModifierKey::Shift, EKeys::Delete), FInputChord(EModifierKey::Shift, EKeys::BackSpace))
 
 	UI_COMMAND( ShowAllPins, "Show All Pins", "Shows all pins", EUserInterfaceActionType::RadioButton, FInputChord() )
 	UI_COMMAND( HideNoConnectionPins, "Hide Unconnected Pins", "Hides all pins with no connections", EUserInterfaceActionType::RadioButton, FInputChord() )
 	UI_COMMAND( HideNoConnectionNoDefaultPins, "Hide Unused Pins", "Hides all pins with no connections and no default value", EUserInterfaceActionType::RadioButton, FInputChord() )
 
-	UI_COMMAND( AddParentNode, "Add call to parent function", "Adds a node that calls this function's parent", EUserInterfaceActionType::Button, FInputChord() )
+	UI_COMMAND( AddParentNode, "Add Call to Parent Function", "Adds a node that calls this function's parent", EUserInterfaceActionType::Button, FInputChord() )
 	
 	UI_COMMAND( CreateMatchingFunction, "Create Matching Function", "Adds a function to the blueprint with a matching signature", EUserInterfaceActionType::Button, FInputChord() )
 
-	UI_COMMAND( ToggleBreakpoint, "Toggle breakpoint", "Adds or removes a breakpoint on each selected node", EUserInterfaceActionType::Button, FInputChord(EKeys::F9) )
-	UI_COMMAND( AddBreakpoint, "Add breakpoint", "Adds a breakpoint to each selected node", EUserInterfaceActionType::Button, FInputChord() )
-	UI_COMMAND( RemoveBreakpoint, "Remove breakpoint", "Removes any breakpoints on each selected node", EUserInterfaceActionType::Button, FInputChord() )
-	UI_COMMAND( EnableBreakpoint, "Enable breakpoint", "Enables any breakpoints on each selected node", EUserInterfaceActionType::Button, FInputChord() )
-	UI_COMMAND( DisableBreakpoint, "Disable breakpoint", "Disables any breakpoints on each selected node", EUserInterfaceActionType::Button, FInputChord() )
+	UI_COMMAND( ToggleBreakpoint, "Toggle Breakpoint", "Adds or removes a breakpoint on each selected node", EUserInterfaceActionType::Button, FInputChord(EKeys::F9) )
+	UI_COMMAND( AddBreakpoint, "Add Breakpoint", "Adds a breakpoint to each selected node", EUserInterfaceActionType::Button, FInputChord() )
+	UI_COMMAND( RemoveBreakpoint, "Remove Breakpoint", "Removes any breakpoints on each selected node", EUserInterfaceActionType::Button, FInputChord() )
+	UI_COMMAND( EnableBreakpoint, "Enable Breakpoint", "Enables any breakpoints on each selected node", EUserInterfaceActionType::Button, FInputChord() )
+	UI_COMMAND( DisableBreakpoint, "Disable Breakpoint", "Disables any breakpoints on each selected node", EUserInterfaceActionType::Button, FInputChord() )
 
 	UI_COMMAND( CollapseNodes, "Collapse Nodes", "Collapses selected nodes into a single node", EUserInterfaceActionType::Button, FInputChord() )
 	UI_COMMAND( PromoteSelectionToFunction, "Promote to Function", "Promotes selected collapsed graphs to functions.", EUserInterfaceActionType::Button, FInputChord() )
@@ -46,8 +46,8 @@ void FGraphEditorCommandsImpl::RegisterCommands()
 	UI_COMMAND( ExpandNodes, "Expand Node", "Expands the node's internal graph into the current graph and removes this node.", EUserInterfaceActionType::Button, FInputChord() )
 	UI_COMMAND( CollapseSelectionToFunction, "Collapse to Function", "Collapses selected nodes into a single function node.", EUserInterfaceActionType::Button, FInputChord() )
 	UI_COMMAND( CollapseSelectionToMacro, "Collapse to Macro", "Collapses selected nodes into a single macro node.", EUserInterfaceActionType::Button, FInputChord() )
-	UI_COMMAND( ConvertFunctionToEvent, "Convert function to event", "Converts selected function to an event and removes the function defintion.", EUserInterfaceActionType::Button, FInputChord() )
-	UI_COMMAND( ConvertEventToFunction, "Convert event to function", "Converts the selected event to a function graph and removes this event node.", EUserInterfaceActionType::Button, FInputChord() )
+	UI_COMMAND( ConvertFunctionToEvent, "Convert Function to Event", "Converts selected function to an event and removes the function defintion.", EUserInterfaceActionType::Button, FInputChord() )
+	UI_COMMAND( ConvertEventToFunction, "Convert Event to Function", "Converts the selected event to a function graph and removes this event node.", EUserInterfaceActionType::Button, FInputChord() )
 
 	UI_COMMAND( AlignNodesTop, "Align Top", "Aligns the top edges of the selected nodes", EUserInterfaceActionType::Button, FInputChord(EModifierKey::Shift, EKeys::W) )
 	UI_COMMAND( AlignNodesMiddle, "Align Middle", "Aligns the vertical middles of the selected nodes", EUserInterfaceActionType::Button, FInputChord(EModifierKey::Shift|EModifierKey::Alt, EKeys::W) )
@@ -67,19 +67,20 @@ void FGraphEditorCommandsImpl::RegisterCommands()
 	UI_COMMAND( EnableNodes_DevelopmentOnly, "Enable Nodes (Development Only)", "Selected node(s) will be enabled in development mode only.", EUserInterfaceActionType::RadioButton, FInputChord() )
 
 	UI_COMMAND( SelectReferenceInLevel, "Find Actor in Level", "Select the actor referenced by this node in the level", EUserInterfaceActionType::Button, FInputChord() )
-	UI_COMMAND( AssignReferencedActor, "Assign selected Actor", "Assign the selected actor to be this node's referenced object", EUserInterfaceActionType::Button, FInputChord() )
+	UI_COMMAND( AssignReferencedActor, "Assign Selected Actor", "Assign the selected actor to be this node's referenced object", EUserInterfaceActionType::Button, FInputChord() )
 	UI_COMMAND( FindReferences, "Find References", "Find references of this item", EUserInterfaceActionType::Button, FInputChord(EKeys::F, EModifierKey::Shift | EModifierKey::Alt) )
 	UI_COMMAND( FindAndReplaceReferences, "Replace References", "Brings up a window to help find and replace all instances of this item", EUserInterfaceActionType::Button, FInputChord() )
 	
 	UI_COMMAND( GoToDefinition, "Goto Definition", "Jumps to the defintion of the selected node if available, e.g., C++ code for a native function or the graph for a Blueprint function.", EUserInterfaceActionType::Button, FInputChord(EKeys::G, EModifierKey::Alt) )
 
-	UI_COMMAND( BreakPinLinks, "Break All Pin Link(s)", "Breaks all of the selected pin's links", EUserInterfaceActionType::Button, FInputChord() )
+	UI_COMMAND( BreakThisLink, "Break This Link", "Breaks the selected pin's only link.", EUserInterfaceActionType::Button, FInputChord() )
+	UI_COMMAND( BreakPinLinks, "Break All Pin Links", "Breaks all of the selected pin's links", EUserInterfaceActionType::Button, FInputChord() )
 	UI_COMMAND( PromoteToVariable, "Promote to Variable", "Promotes something to a variable", EUserInterfaceActionType::Button, FInputChord() )
 	UI_COMMAND( PromoteToLocalVariable, "Promote to Local Variable", "Promotes something to a local variable of the current function", EUserInterfaceActionType::Button, FInputChord() )
 	UI_COMMAND( SplitStructPin, "Split Struct Pin", "Breaks a struct pin in to a separate pin per element", EUserInterfaceActionType::Button, FInputChord() )
 	UI_COMMAND( RecombineStructPin, "Recombine Struct Pin", "Takes struct pins that have been broken in to composite elements and combines them back to a single struct pin", EUserInterfaceActionType::Button, FInputChord() )
-	UI_COMMAND( StartWatchingPin, "Watch this value", "Adds this pin or variable to the watch list", EUserInterfaceActionType::Button, FInputChord() )
-	UI_COMMAND( StopWatchingPin, "Stop watching this value", "Removes this pin or variable from the watch list ", EUserInterfaceActionType::Button, FInputChord() )
+	UI_COMMAND( StartWatchingPin, "Watch This Value", "Adds this pin or variable to the watch list", EUserInterfaceActionType::Button, FInputChord() )
+	UI_COMMAND( StopWatchingPin, "Stop Watching This Value", "Removes this pin or variable from the watch list ", EUserInterfaceActionType::Button, FInputChord() )
 	UI_COMMAND( ResetPinToDefaultValue, "Reset to Default Value", "Reset value of this pin to the default", EUserInterfaceActionType::Button, FInputChord() )
 	UI_COMMAND( SelectAllInputNodes, "Select All Input Nodes", "Adds all input Nodes linked to this Pin to selection", EUserInterfaceActionType::Button, FInputChord() )
 	UI_COMMAND( SelectAllOutputNodes, "Select All Output Nodes", "Adds all output Nodes linked to this Pin to selection", EUserInterfaceActionType::Button, FInputChord() )
diff --git a/Engine/Source/Editor/GraphEditor/Private/SGraphEditorImpl.cpp b/Engine/Source/Editor/GraphEditor/Private/SGraphEditorImpl.cpp
--- a/Engine/Source/Editor/GraphEditor/Private/SGraphEditorImpl.cpp
+++ b/Engine/Source/Editor/GraphEditor/Private/SGraphEditorImpl.cpp
@@ -269,8 +269,8 @@ void SGraphEditorImpl::GetPinContextMenuActionsForSchema(UToolMenu* InMenu) cons
 				UGraphNodeContextMenuContext* NodeContext = InToolMenu->FindContext<UGraphNodeContextMenuContext>();
 				if (NodeContext && NodeContext->Pin)
 				{
-					FText SingleDescFormat = LOCTEXT("BreakDesc", "Break link to {NodeTitle}");
-					FText MultiDescFormat = LOCTEXT("BreakDescMulti", "Break link to {NodeTitle} ({NumberOfNodes})");
+					FText SingleDescFormat = LOCTEXT("BreakDesc", "Break Link to {NodeTitle}");
+					FText MultiDescFormat = LOCTEXT("BreakDescMulti", "Break Link to {NodeTitle} ({NumberOfNodes})");
 
 					TMap< FString, uint32 > LinkTitleCount;
 					for (UEdGraphPin* TargetPin : NodeContext->Pin->LinkedTo)
@@ -289,6 +289,22 @@ void SGraphEditorImpl::GetPinContextMenuActionsForSchema(UToolMenu* InMenu) cons
 			}),
 			BreakLinksMenuVisibility,
 			EUserInterfaceActionType::Button);
+
+		// Break This Link
+		{
+			FToolUIAction BreakThisLinkAction;
+			BreakThisLinkAction.ExecuteAction = FToolMenuExecuteAction::CreateSP(this, &SGraphEditorImpl::ExecuteBreakPinLinks);
+			BreakThisLinkAction.IsActionVisibleDelegate = FToolMenuIsActionButtonVisible::CreateSP(this, &SGraphEditorImpl::IsBreakThisLinkVisible);
+
+			TSharedPtr<FUICommandInfo> BreakThisLinkCommand = FGraphEditorCommands::Get().BreakThisLink;
+			Section.AddMenuEntry(
+				BreakThisLinkCommand->GetCommandName(),
+				BreakThisLinkCommand->GetLabel(),
+				BreakThisLinkCommand->GetDescription(),
+				BreakThisLinkCommand->GetIcon(),
+				BreakThisLinkAction
+			);
+		}
 	}
 
 	FToolUIAction PinActionSubMenuVisibiliity;
@@ -297,7 +313,7 @@ void SGraphEditorImpl::GetPinContextMenuActionsForSchema(UToolMenu* InMenu) cons
 	// Jump to specific connections
 	{
 		Section.AddSubMenu("JumpToConnection",
-			LOCTEXT("JumpToConnection", "Jump To Connection..."),
+			LOCTEXT("JumpToConnection", "Jump to Connection..."),
 			LOCTEXT("JumpToSpecificConnection", "Jump to specific connection..."),
 			FNewToolMenuDelegate::CreateLambda([this, GetMenuEntryForPin](UToolMenu* InToolMenu)
 			{
@@ -356,7 +372,7 @@ void SGraphEditorImpl::GetPinContextMenuActionsForSchema(UToolMenu* InMenu) cons
 
 				// Add individual pin connections
 				FText SingleDescFormat = LOCTEXT("StraightenDesc", "Straighten Connection to {NodeTitle}");
-				FText MultiDescFormat = LOCTEXT("StraightenDescMulti", "Straigten Connection to {NodeTitle} ({NumberOfNodes})");
+				FText MultiDescFormat = LOCTEXT("StraightenDescMulti", "Straighten Connection to {NodeTitle} ({NumberOfNodes})");
 				TMap< FString, uint32 > LinkTitleCount;
 				for (UEdGraphPin* TargetPin : NodeContext->Pin->LinkedTo)
 				{
@@ -400,7 +416,18 @@ bool SGraphEditorImpl::IsBreakPinLinksVisible(const FToolMenuContext& InContext)
 	UGraphNodeContextMenuContext* NodeContext = InContext.FindContext<UGraphNodeContextMenuContext>();
 	if (NodeContext && NodeContext->Pin)
 	{
-		return !NodeContext->bIsDebugging && (NodeContext->Pin->LinkedTo.Num() > 0);
+		return !NodeContext->bIsDebugging && (NodeContext->Pin->LinkedTo.Num() > 1);
+	}
+
+	return false;
+}
+
+bool SGraphEditorImpl::IsBreakThisLinkVisible(const FToolMenuContext& InContext) const
+{
+	UGraphNodeContextMenuContext* NodeContext = InContext.FindContext<UGraphNodeContextMenuContext>();
+	if (NodeContext && NodeContext->Pin)
+	{
+		return !NodeContext->bIsDebugging && (NodeContext->Pin->LinkedTo.Num() == 1);
 	}
 
 	return false;
diff --git a/Engine/Source/Editor/GraphEditor/Private/SGraphEditorImpl.h b/Engine/Source/Editor/GraphEditor/Private/SGraphEditorImpl.h
--- a/Engine/Source/Editor/GraphEditor/Private/SGraphEditorImpl.h
+++ b/Engine/Source/Editor/GraphEditor/Private/SGraphEditorImpl.h
@@ -376,6 +376,7 @@ private:
 	void GetPinContextMenuActionsForSchema(UToolMenu* InMenu) const;
 	void ExecuteBreakPinLinks(const FToolMenuContext& InContext) const;
 	bool IsBreakPinLinksVisible(const FToolMenuContext& InContext) const;
+	bool IsBreakThisLinkVisible(const FToolMenuContext& InContext) const;
 	bool HasAnyLinkedPins(const FToolMenuContext& InContext) const;
 
 	void ExecuteSelectConnectedNodesFromPin(const FToolMenuContext& InContext) const;
diff --git a/Engine/Source/Editor/GraphEditor/Public/GraphEditorActions.h b/Engine/Source/Editor/GraphEditor/Public/GraphEditorActions.h
--- a/Engine/Source/Editor/GraphEditor/Public/GraphEditorActions.h
+++ b/Engine/Source/Editor/GraphEditor/Public/GraphEditorActions.h
@@ -102,6 +102,7 @@ public:
 	TSharedPtr< FUICommandInfo > GoToDefinition;
 
 	// Pin-specific actions
+	TSharedPtr< FUICommandInfo > BreakThisLink;
 	TSharedPtr< FUICommandInfo > BreakPinLinks;
 	TSharedPtr< FUICommandInfo > PromoteToVariable;
 	TSharedPtr< FUICommandInfo > PromoteToLocalVariable;
