[GraphEditor] fixed visibility test for overlay widgets that are not in the child hierarchy.

diff --git a/Engine/Source/Editor/GraphEditor/Private/SGraphPanel.cpp b/Engine/Source/Editor/GraphEditor/Private/SGraphPanel.cpp
--- a/Engine/Source/Editor/GraphEditor/Private/SGraphPanel.cpp
+++ b/Engine/Source/Editor/GraphEditor/Private/SGraphPanel.cpp
@@ -7,6 +7,7 @@
 #include "Layout/WidgetPath.h"
 #include "Framework/Application/MenuStack.h"
 #include "Framework/Application/SlateApplication.h"
+#include "Types/SlateAttributeMetaData.h"
 #include "EdGraphNode_Comment.h"
 #include "Settings/EditorExperimentalSettings.h"
 #include "Editor.h"
@@ -293,14 +294,18 @@ int32 SGraphPanel::OnPaint( const FPaintArgs& Args, const FGeometry& AllottedGeo
 						for (int32 WidgetIndex = 0; WidgetIndex < OverlayWidgets.Num(); ++WidgetIndex)
 						{
 							FOverlayWidgetInfo& OverlayInfo = OverlayWidgets[WidgetIndex];
-							if(OverlayInfo.Widget->GetVisibility() == EVisibility::Visible)
+							if (SWidget* Widget = OverlayInfo.Widget.Get())
 							{
-								// call SlatePrepass as these widgets are not in the 'normal' child hierarchy
-								OverlayInfo.Widget->SlatePrepass(AllottedGeometry.GetAccumulatedLayoutTransform().GetScale());
+								FSlateAttributeMetaData::UpdateOnlyVisibilityAttributes(*Widget, FSlateAttributeMetaData::EInvalidationPermission::AllowInvalidationIfConstructed);
+								if (Widget->GetVisibility() == EVisibility::Visible)
+								{
+									// call SlatePrepass as these widgets are not in the 'normal' child hierarchy
+									Widget->SlatePrepass(AllottedGeometry.GetAccumulatedLayoutTransform().GetScale());
 
-								const FGeometry WidgetGeometry = CurWidget.Geometry.MakeChild(OverlayInfo.OverlayOffset, OverlayInfo.Widget->GetDesiredSize());
+									const FGeometry WidgetGeometry = CurWidget.Geometry.MakeChild(OverlayInfo.OverlayOffset, Widget->GetDesiredSize());
 
-								OverlayInfo.Widget->Paint(NewArgs, WidgetGeometry, MyCullingRect, OutDrawElements, CurWidgetsMaxLayerId, InWidgetStyle, bParentEnabled);
+									Widget->Paint(NewArgs, WidgetGeometry, MyCullingRect, OutDrawElements, CurWidgetsMaxLayerId, InWidgetStyle, bParentEnabled);
+								}
 							}
 						}
 					}
