Add pin property watches info to watched node info popups

diff --git a/Engine/Source/Editor/GraphEditor/Private/KismetNodes/SGraphNodeK2Base.cpp b/Engine/Source/Editor/GraphEditor/Private/KismetNodes/SGraphNodeK2Base.cpp
--- a/Engine/Source/Editor/GraphEditor/Private/KismetNodes/SGraphNodeK2Base.cpp
+++ b/Engine/Source/Editor/GraphEditor/Private/KismetNodes/SGraphNodeK2Base.cpp
@@ -23,6 +23,7 @@
 #include "IDocumentation.h"
 #include "TutorialMetaData.h"
 #include "Widgets/Layout/SBox.h"
+#include "BlueprintEditorSettings.h"
 
 #define LOCTEXT_NAMESPACE "SGraphNodeK2Base"
 
@@ -538,55 +539,112 @@ void SGraphNodeK2Base::GetNodeInfoPopups(FNodeInfoContext* Context, TArray<FGrap
 		// Display pinned watches
 		if (K2Context->WatchedNodeSet.Contains(GraphNode))
 		{
+			const UBlueprintEditorSettings* BlueprintEditorSettings = GetDefault<UBlueprintEditorSettings>();
 			UBlueprint* Blueprint = K2Context->SourceBlueprint;
 			const UEdGraphSchema* Schema = GraphNode->GetSchema();
 
-			FString PinnedWatchText;
-			int32 ValidWatchCount = 0;
-			for (int32 PinIndex = 0; PinIndex < GraphNode->Pins.Num(); ++PinIndex)
+			const FPerBlueprintSettings* FoundSettings = BlueprintEditorSettings->PerBlueprintSettings.Find(Blueprint->GetPathName());
+			if (FoundSettings)
 			{
-				UEdGraphPin* WatchPin = GraphNode->Pins[PinIndex];
-				if (K2Context->WatchedPinSet.Contains(WatchPin))
+				FString PinnedWatchText;
+				int32 ValidWatchCount = 0;
+				TMap<const UEdGraphPin*, TSharedPtr<FPropertyInstanceInfo>> CachedPinInfo;
+				for (const FBlueprintWatchedPin& WatchedPin : FoundSettings->WatchedPins)
 				{
-					if (ValidWatchCount > 0)
+					const UEdGraphPin* WatchPin = WatchedPin.Get();
+					if (WatchPin && WatchPin->GetOwningNode() == GraphNode)
 					{
-						PinnedWatchText += TEXT("\n");
+						if (ValidWatchCount > 0)
+						{
+							PinnedWatchText += TEXT("\n");
+						}
+						TSharedPtr<FPropertyInstanceInfo> PinInfo;
+						if (CachedPinInfo.Find(WatchPin))
+						{
+							PinInfo = CachedPinInfo[WatchPin];
+						}
+						else
+						{
+							const FKismetDebugUtilities::EWatchTextResult WatchStatus = FKismetDebugUtilities::GetDebugInfo(CachedPinInfo.Add(WatchPin), Blueprint, ActiveObject, WatchPin);
+
+							if (WatchStatus == FKismetDebugUtilities::EWTR_Valid)
+							{
+								PinInfo = CachedPinInfo[WatchPin];
+							}
+							else
+							{
+								FString PinName = UEdGraphSchema_K2::TypeToText(WatchPin->PinType).ToString();
+								PinName += TEXT(" ");
+								PinName += Schema->GetPinDisplayName(WatchPin).ToString();
+
+								switch (WatchStatus)
+								{
+								case FKismetDebugUtilities::EWTR_Valid:
+									break;
+
+								case FKismetDebugUtilities::EWTR_NotInScope:
+									PinnedWatchText += FText::Format(LOCTEXT("WatchingWhenNotInScopeFmt", "Watching {0}\n\t(not in scope)"), FText::FromString(PinName)).ToString();
+									break;
+
+								case FKismetDebugUtilities::EWTR_NoProperty:
+									PinnedWatchText += FText::Format(LOCTEXT("WatchingUnknownPropertyFmt", "Watching {0}\n\t(no debug data)"), FText::FromString(PinName)).ToString();
+									break;
+
+								default:
+								case FKismetDebugUtilities::EWTR_NoDebugObject:
+									PinnedWatchText += FText::Format(LOCTEXT("WatchingNoDebugObjectFmt", "Watching {0}"), FText::FromString(PinName)).ToString();
+									break;
+								}
+							}
+						}
+
+						if (PinInfo.IsValid())
+						{
+							FString WatchName;
+							FString WatchText;
+							if (WatchedPin.GetPathToProperty().IsEmpty())
+							{
+								WatchName = UEdGraphSchema_K2::TypeToText(WatchPin->PinType).ToString();
+								WatchName += TEXT(" ");
+								WatchName += Schema->GetPinDisplayName(WatchPin).ToString();
+
+								WatchText = PinInfo->Value.ToString();
+							}
+							else
+							{
+								TSharedPtr<FPropertyInstanceInfo> PropWatch = PinInfo->ResolvePathToProperty(WatchedPin.GetPathToProperty());
+								if (PropWatch.IsValid())
+								{
+									WatchName = UEdGraphSchema_K2::TypeToText(PropWatch->Property.Get()).ToString();
+									WatchName += TEXT(" ");
+
+									WatchText = PropWatch->Value.ToString();
+								}
+								else
+								{
+									WatchText = LOCTEXT("NoDebugData", "(no debug data)").ToString();
+								}
+								
+								WatchName += Schema->GetPinDisplayName(WatchPin).ToString();
+
+								for (const FName& PathName : WatchedPin.GetPathToProperty())
+								{
+									WatchName += TEXT("/");
+									WatchName += PathName.ToString();
+								}
+							}
+
+							PinnedWatchText += FText::Format(LOCTEXT("WatchingAndValidFmt", "Watching {0}\n\t{1}"), FText::FromString(WatchName), FText::FromString(WatchText)).ToString();
+						}
+
+						ValidWatchCount++;
 					}
-
-					FString PinName = UEdGraphSchema_K2::TypeToText(WatchPin->PinType).ToString();
-					PinName += TEXT(" ");
-					PinName += Schema->GetPinDisplayName(WatchPin).ToString();
-
-					FString WatchText;
-					const FKismetDebugUtilities::EWatchTextResult WatchStatus = FKismetDebugUtilities::GetWatchText(/*inout*/ WatchText, Blueprint, ActiveObject, WatchPin);
-
-					switch (WatchStatus)
-					{
-					case FKismetDebugUtilities::EWTR_Valid:
-						PinnedWatchText += FText::Format(LOCTEXT("WatchingAndValidFmt", "Watching {0}\n\t{1}"), FText::FromString(PinName), FText::FromString(WatchText)).ToString();//@TODO: Print out object being debugged name?
-						break;
-
-					case FKismetDebugUtilities::EWTR_NotInScope:
-						PinnedWatchText += FText::Format(LOCTEXT("WatchingWhenNotInScopeFmt", "Watching {0}\n\t(not in scope)"), FText::FromString(PinName)).ToString();
-						break;
-
-					case FKismetDebugUtilities::EWTR_NoProperty:
-						PinnedWatchText += FText::Format(LOCTEXT("WatchingUnknownPropertyFmt", "Watching {0}\n\t(no debug data)"), FText::FromString(PinName)).ToString();
-						break;
-
-					default:
-					case FKismetDebugUtilities::EWTR_NoDebugObject:
-						PinnedWatchText += FText::Format(LOCTEXT("WatchingNoDebugObjectFmt", "Watching {0}"), FText::FromString(PinName)).ToString();
-						break;
-					}
-
-					ValidWatchCount++;
 				}
-			}
 
-			if (ValidWatchCount)
-			{
-				new (Popups) FGraphInformationPopupInfo(NULL, PinnedWatchColor, PinnedWatchText);
+				if (ValidWatchCount)
+				{
+					new (Popups) FGraphInformationPopupInfo(nullptr, PinnedWatchColor, PinnedWatchText);
+				}
 			}
 		}
 	}
diff --git a/Engine/Source/Editor/Kismet/Private/Debugging/SKismetDebugTreeView.cpp b/Engine/Source/Editor/Kismet/Private/Debugging/SKismetDebugTreeView.cpp
--- a/Engine/Source/Editor/Kismet/Private/Debugging/SKismetDebugTreeView.cpp
+++ b/Engine/Source/Editor/Kismet/Private/Debugging/SKismetDebugTreeView.cpp
@@ -1368,9 +1368,6 @@ protected:
 
 	void OnNavigateToWatchLocation();
 
-	// Finds the DebugInfo for the property pointed to by PathToProperty from the Pin's debuginfo
-	TSharedPtr<FPropertyInstanceInfo> GetWatchedPropertyDebugInfo(const TSharedPtr<FPropertyInstanceInfo>& InPinInfo) const;
-
 private:
 	void AddWatch() const
 	{
@@ -1811,52 +1808,6 @@ void FWatchLineItem::OnNavigateToWatchLocation()
 	}
 }
 
-TSharedPtr<FPropertyInstanceInfo> FWatchLineItem::GetWatchedPropertyDebugInfo(const TSharedPtr<FPropertyInstanceInfo>& InPinInfo) const
-{
-	// Drill down to the actual property being watched if necessary
-	TSharedPtr<FPropertyInstanceInfo> ThisDebugInfo = InPinInfo;
-	for (const FName& ChildName : PathToProperty)
-	{
-		if (ThisDebugInfo.IsValid())
-		{
-			const FString ChildNameStr = ChildName.ToString();
-			FProperty* Property = ThisDebugInfo->Property.Get();
-
-			TSharedPtr<FPropertyInstanceInfo>* FoundChild;
-			
-			if (Property->IsA<FSetProperty>() || Property->IsA<FArrayProperty>() || Property->IsA<FMapProperty>())
-			{
-				FoundChild = ThisDebugInfo->Children.FindByPredicate(
-					[&ChildNameStr](const TSharedPtr<FPropertyInstanceInfo>& Child)
-					{
-						return Child->DisplayName.ToString() == ChildNameStr;
-					}
-				);
-			}
-			else
-			{
-				FoundChild = ThisDebugInfo->Children.FindByPredicate(
-					[&ChildNameStr](const TSharedPtr<FPropertyInstanceInfo>& Child)
-					{
-						return Child->Property->GetAuthoredName() == ChildNameStr;
-					}
-				);
-			}
-
-			if (FoundChild)
-			{
-				ThisDebugInfo = *FoundChild;
-			}
-			else
-			{
-				ThisDebugInfo.Reset();
-			}
-		}
-	}
-
-	return ThisDebugInfo;
-}
-
 TSharedPtr<FPropertyInstanceInfo> FWatchLineItem::GetPropertyInfo() const
 {
 	if (CachedPropertyInfo.IsValid())
@@ -1880,7 +1831,14 @@ TSharedPtr<FPropertyInstanceInfo> FWatchLineItem::GetPropertyInfo() const
 			case FKismetDebugUtilities::EWTR_Valid:
 			{
 					check(DebugInfo);
-					CachedPropertyInfo = GetWatchedPropertyDebugInfo(DebugInfo);
+					if (PathToProperty.IsEmpty())
+					{
+						CachedPropertyInfo = DebugInfo;
+					}
+					else
+					{
+						CachedPropertyInfo = DebugInfo->ResolvePathToProperty(PathToProperty);
+					}
 					return CachedPropertyInfo;
 			}
 
diff --git a/Engine/Source/Editor/UnrealEd/Private/Kismet2/KismetDebugUtilities.cpp b/Engine/Source/Editor/UnrealEd/Private/Kismet2/KismetDebugUtilities.cpp
--- a/Engine/Source/Editor/UnrealEd/Private/Kismet2/KismetDebugUtilities.cpp
+++ b/Engine/Source/Editor/UnrealEd/Private/Kismet2/KismetDebugUtilities.cpp
@@ -2207,6 +2207,59 @@ void FPropertyInstanceInfo::PopulateChildren(FPropertyInstance PropertyInstance,
 	}
 }
 
+
+TSharedPtr<FPropertyInstanceInfo> FPropertyInstanceInfo::ResolvePathToProperty(const TArray<FName>& InPathToProperty)
+{
+	const auto FindChildInPropertyInfo = [](FPropertyInstanceInfo* InPropertyInfo, const FName& InChildName) -> TSharedPtr<FPropertyInstanceInfo>*
+	{
+		if (InPropertyInfo)
+		{
+			const FString ChildNameStr = InChildName.ToString();
+			FProperty* Property = InPropertyInfo->Property.Get();
+
+			if (Property->IsA<FSetProperty>() || Property->IsA<FArrayProperty>() || Property->IsA<FMapProperty>())
+			{
+				return InPropertyInfo->Children.FindByPredicate(
+					[&ChildNameStr](const TSharedPtr<FPropertyInstanceInfo>& Child)
+					{
+						// Display name for Container Element Properties is just their index
+						return Child->DisplayName.ToString() == ChildNameStr;
+					}
+				);
+			}
+			else
+			{
+				return InPropertyInfo->Children.FindByPredicate(
+					[&ChildNameStr](const TSharedPtr<FPropertyInstanceInfo>& Child)
+					{
+						return Child->Property->GetAuthoredName() == ChildNameStr;
+					}
+				);
+			}
+		}
+		return nullptr;
+	};
+
+	TSharedPtr<FPropertyInstanceInfo> ThisDebugInfo = nullptr;
+	for (const FName& ChildName : InPathToProperty)
+	{
+		TSharedPtr<FPropertyInstanceInfo>* FoundChild = ThisDebugInfo.IsValid() ? FindChildInPropertyInfo(ThisDebugInfo.Get(), ChildName) : FindChildInPropertyInfo(this, ChildName);
+
+		if (FoundChild)
+		{
+			ThisDebugInfo = *FoundChild;
+		}
+		else
+		{
+			return nullptr;
+		}
+	}
+
+	return ThisDebugInfo;
+}
+
+
+
 FText FKismetDebugUtilities::GetAndClearLastExceptionMessage()
 {
 	FKismetDebugUtilitiesData& Data = FKismetDebugUtilitiesData::Get();
diff --git a/Engine/Source/Editor/UnrealEd/Public/Kismet2/KismetDebugUtilities.h b/Engine/Source/Editor/UnrealEd/Public/Kismet2/KismetDebugUtilities.h
--- a/Engine/Source/Editor/UnrealEd/Public/Kismet2/KismetDebugUtilities.h
+++ b/Engine/Source/Editor/UnrealEd/Public/Kismet2/KismetDebugUtilities.h
@@ -77,6 +77,9 @@ struct FPropertyInstanceInfo
 	*/
 	void PopulateChildren(FPropertyInstance PropertyInstance,
 		TMap<FPropertyInstance, TSharedPtr<FPropertyInstanceInfo>> &VisitedNodes);
+
+	/** Resolves the PathToProperty treating this PropertyInstance as the head of the path */
+	UNREALED_API TSharedPtr<FPropertyInstanceInfo> ResolvePathToProperty(const TArray<FName>& InPathToProperty);
 	
 	FText Name;
 	FText DisplayName;
