Fix to nullptr acess on accessing and invalidated pin. This change will still trigger the ensure for invalidated pin but will prevent a crash acessing the ptr.

diff --git a/Engine/Source/Editor/GraphEditor/Private/SGraphPanel.cpp b/Engine/Source/Editor/GraphEditor/Private/SGraphPanel.cpp
--- a/Engine/Source/Editor/GraphEditor/Private/SGraphPanel.cpp
+++ b/Engine/Source/Editor/GraphEditor/Private/SGraphPanel.cpp
@@ -781,11 +781,17 @@ bool SGraphPanel::OnHandleLeftMouseRelease(const FGeometry& MyGeometry, const FP
 			if( PinWidgetGeometry.Geometry.IsUnderLocation( MouseEvent.GetScreenSpacePosition() ) )
 			{
 				SGraphPin& TargetPin = static_cast<SGraphPin&>( PinWidgetGeometry.Widget.Get() );
-
+				
 				if (PreviewConnectionPin->TryHandlePinConnection(TargetPin))
 				{
-					NodeList.Add(TargetPin.GetPinObj()->GetOwningNode());
-					NodeList.Add(PreviewConnectionPin->GetPinObj()->GetOwningNode());
+					// We have to do a second check on PinObjs here since TryHandlePinConnection, may invalidate them.
+					UEdGraphPin* PreviewConnectionPinObj = PreviewConnectionPin->GetPinObj();
+					UEdGraphPin* TargetPinObj = TargetPin.GetPinObj();
+					if (TargetPinObj && PreviewConnectionPinObj)
+					{
+						NodeList.Add(TargetPinObj->GetOwningNode());
+						NodeList.Add(PreviewConnectionPinObj->GetOwningNode());
+					}
 				}
 				bHandledDrop = true;
 			}
