Blueprints: Add new promotable operator icons to make it clearer that you can connect any type to it instead of just the one that is currently connected.

#rnx

diff --git a/Engine/Source/Editor/EditorStyle/Private/StarshipStyle.cpp b/Engine/Source/Editor/EditorStyle/Private/StarshipStyle.cpp
--- a/Engine/Source/Editor/EditorStyle/Private/StarshipStyle.cpp
+++ b/Engine/Source/Editor/EditorStyle/Private/StarshipStyle.cpp
@@ -5263,6 +5263,9 @@ void FStarshipEditorStyle::FStyle::SetupPersonaStyle()
 		Set("Kismet.VariableList.MapValueTypeIcon", new IMAGE_BRUSH_SVG( "Starship/Blueprints/pillmapvalue", Icon16x16));
 		Set("Kismet.VariableList.MapKeyTypeIcon", new IMAGE_BRUSH_SVG( "Starship/Blueprints/pillmapkey", Icon16x16));
 
+		Set("Kismet.VariableList.PromotableTypeOuterIcon", new IMAGE_BRUSH_SVG( "Starship/Blueprints/promotable_type_outer_icon", Icon14x14));
+		Set("Kismet.VariableList.PromotableTypeInnerIcon", new IMAGE_BRUSH_SVG( "Starship/Blueprints/promotable_type_inner_icon", Icon14x14));
+
 		Set("Kismet.VariableList.ExposeForInstance", new CORE_IMAGE_BRUSH_SVG("Starship/Common/visible", Icon16x16));
 		Set("Kismet.VariableList.HideForInstance", new CORE_IMAGE_BRUSH_SVG("Starship/Common/hidden", Icon16x16));
 
diff --git a/Engine/Source/Editor/GraphEditor/Private/NodeFactory.cpp b/Engine/Source/Editor/GraphEditor/Private/NodeFactory.cpp
--- a/Engine/Source/Editor/GraphEditor/Private/NodeFactory.cpp
+++ b/Engine/Source/Editor/GraphEditor/Private/NodeFactory.cpp
@@ -33,13 +33,14 @@
 #include "K2Node_SpawnActorFromClass.h"
 #include "K2Node_Switch.h"
 #include "K2Node_Timeline.h"
-
+#include "K2Node_PromotableOperator.h"
 
 #include "SGraphNodeDefault.h"
 #include "SGraphNodeComment.h"
 #include "SGraphNodeDocumentation.h"
 #include "EdGraph/EdGraphNode_Documentation.h"
 #include "SGraphNodeKnot.h"
+#include "SGraphNodePromotableOperator.h"
 
 #include "KismetNodes/SGraphNodeK2Default.h"
 #include "KismetNodes/SGraphNodeK2Var.h"
@@ -151,6 +152,10 @@ TSharedPtr<SGraphNode> FNodeFactory::CreateNodeWidget(UEdGraphNode* InNode)
 		{
 			return SNew(SGraphNodeSwitchStatement, SwitchNode);
 		}
+		else if(UK2Node_PromotableOperator* PromotableOperator = Cast<UK2Node_PromotableOperator>(InNode))
+		{
+			return SNew(SGraphNodePromotableOperator, PromotableOperator);
+		}
 		else if (InNode->GetClass()->ImplementsInterface(UK2Node_AddPinInterface::StaticClass()))
 		{
 			return SNew(SGraphNodeK2Sequence, CastChecked<UK2Node>(InNode));
diff --git a/Engine/Source/Editor/GraphEditor/Private/SGraphNodePromotableOperator.cpp b/Engine/Source/Editor/GraphEditor/Private/SGraphNodePromotableOperator.cpp
new file mode 100644
--- /dev/null
+++ b/Engine/Source/Editor/GraphEditor/Private/SGraphNodePromotableOperator.cpp
@@ -0,0 +1,74 @@
+// Copyright Epic Games, Inc. All Rights Reserved.
+
+#include "SGraphNodePromotableOperator.h"
+#include "EdGraphSchema_K2.h"
+#include "GraphEditorSettings.h"
+#include "SGraphPin.h"
+#include "SPinTypeSelector.h"
+#include "Widgets/Images/SLayeredImage.h"
+#include "GraphEditorSettings.h"
+
+#define LOCTEXT_NAMESPACE "SGraphNodePromotableOperator"
+
+//////////////////////////////////////////////////////////////////////////
+// SGraphNodePromotableOperator
+
+void SGraphNodePromotableOperator::Construct(const FArguments& InArgs, UK2Node_PromotableOperator* InNode)
+{
+	SGraphNodeK2Sequence::Construct(SGraphNodeK2Sequence::FArguments(), InNode);
+
+	LoadCachedIcons();
+}
+
+void SGraphNodePromotableOperator::CreatePinWidgets()
+{
+	SGraphNodeK2Sequence::CreatePinWidgets();
+
+	TSet<TSharedRef<SWidget>> AllPins;
+	GetPins(AllPins);
+
+	LoadCachedIcons();
+
+	for(TSharedRef<SWidget>& Widget : AllPins)
+	{
+		
+		if(TSharedPtr<SGraphPin> Pin = StaticCastSharedRef<SGraphPin>(Widget))
+		{
+			UEdGraphPin* SourcePin = Pin->GetPinObj();
+			
+			// Split pins should be drawn as normal pins, the inner properties are not promotable
+			if(!SourcePin || SourcePin->ParentPin)
+			{
+				continue;
+			}
+			
+			if(TSharedPtr<SLayeredImage> PinImage = StaticCastSharedPtr<SLayeredImage>(Pin->GetPinImageWidget()))
+			{
+				// Set the image to use the outer icon, which will be the connect pin type color
+				PinImage->SetLayerBrush(0, CachedOuterIcon);
+
+				// Set the inner image to be wildcard color, which is grey by default
+				PinImage->AddLayer(CachedInnerIcon, GetDefault<UGraphEditorSettings>()->WildcardPinTypeColor);
+			}
+		}
+	}
+}
+
+void SGraphNodePromotableOperator::LoadCachedIcons()
+{
+	static const FName PromotableTypeOuterName("Kismet.VariableList.PromotableTypeOuterIcon");
+	static const FName PromotableTypeInnerName("Kismet.VariableList.PromotableTypeInnerIcon");
+
+	// Outer ring icons
+	if(!CachedOuterIcon)
+	{
+		CachedOuterIcon = FEditorStyle::GetBrush(PromotableTypeOuterName);
+	}
+
+	if(!CachedInnerIcon)
+	{
+		CachedInnerIcon = FEditorStyle::GetBrush(PromotableTypeInnerName);
+	}
+}
+
+#undef LOCTEXT_NAMESPACE
diff --git a/Engine/Source/Editor/GraphEditor/Private/SGraphPin.cpp b/Engine/Source/Editor/GraphEditor/Private/SGraphPin.cpp
--- a/Engine/Source/Editor/GraphEditor/Private/SGraphPin.cpp
+++ b/Engine/Source/Editor/GraphEditor/Private/SGraphPin.cpp
@@ -1337,6 +1337,16 @@ EVisibility SGraphPin::GetPinVisiblity() const
 	return EVisibility::Visible;
 }
 
+TSharedPtr<SWidget> SGraphPin::GetPinImageWidget() const
+{
+	return PinImage;
+}
+
+void SGraphPin::SetPinImageWidget(TSharedRef<SWidget> NewWidget)
+{
+	PinImage = NewWidget;
+}
+
 void SGraphPin::SetCustomPinIcon(const FSlateBrush* InConnectedBrush, const FSlateBrush* InDisconnectedBrush)
 {
 	Custom_Brush_Connected = InConnectedBrush;
diff --git a/Engine/Source/Editor/GraphEditor/Public/SGraphNodePromotableOperator.h b/Engine/Source/Editor/GraphEditor/Public/SGraphNodePromotableOperator.h
new file mode 100644
--- /dev/null
+++ b/Engine/Source/Editor/GraphEditor/Public/SGraphNodePromotableOperator.h
@@ -0,0 +1,35 @@
+// Copyright Epic Games, Inc. All Rights Reserved.
+
+#pragma once
+
+#include "CoreMinimal.h"
+#include "K2Node_PromotableOperator.h"
+#include "KismetNodes/SGraphNodeK2Sequence.h"
+
+/**
+ * A custom slate node for the promotable operator (K2Node_PromotableOperator)
+ * These nodes include common math operations like Add, Subtract, Multiply, etc.
+ * 
+ * This custom node will provide a neat little secondary icon to visually show that
+ * promotable operator pins can be converted to other types.
+ */
+class SGraphNodePromotableOperator : public SGraphNodeK2Sequence
+{
+public:
+
+	SLATE_BEGIN_ARGS(SGraphNodePromotableOperator){}
+	SLATE_END_ARGS()
+
+	void Construct(const FArguments& InArgs, UK2Node_PromotableOperator* InNode);
+	
+	// SGraphNode interface
+	virtual void CreatePinWidgets() override;
+	// End of SGraphNode interface
+
+protected:
+
+	void LoadCachedIcons();
+
+	const FSlateBrush* CachedOuterIcon;
+	const FSlateBrush* CachedInnerIcon;
+};
diff --git a/Engine/Source/Editor/GraphEditor/Public/SGraphPin.h b/Engine/Source/Editor/GraphEditor/Public/SGraphPin.h
--- a/Engine/Source/Editor/GraphEditor/Public/SGraphPin.h
+++ b/Engine/Source/Editor/GraphEditor/Public/SGraphPin.h
@@ -270,6 +270,14 @@ protected:
 	/** Determines the pin's visibility based on the LOD factor, when it is low LOD, no hit test will occur */
 	EVisibility GetPinVisiblity() const;
 
+public:
+
+	/** Returns the current pin image widget that is being used */
+	TSharedPtr<SWidget> GetPinImageWidget() const;
+
+	/** Sets the pin widget for this pin */
+	void SetPinImageWidget(TSharedRef<SWidget> NewWidget);
+
 protected:
 	/** The GraphNode that owns this pin */
 	TWeakPtr<SGraphNode> OwnerNodePtr;
