Control Rig: Cannot drag wires out from pins on collapsed nodes

This is due to the fact that we are changing the setting on the actual node - not the widget itself. Added a boolean to overdrive the behavior per widget

diff --git a/Engine/Plugins/Animation/ControlRig/Source/ControlRigEditor/Private/ControlRigGraphDetails.cpp b/Engine/Plugins/Animation/ControlRig/Source/ControlRigEditor/Private/ControlRigGraphDetails.cpp
--- a/Engine/Plugins/Animation/ControlRig/Source/ControlRigEditor/Private/ControlRigGraphDetails.cpp
+++ b/Engine/Plugins/Animation/ControlRig/Source/ControlRigEditor/Private/ControlRigGraphDetails.cpp
@@ -618,7 +618,7 @@ void FControlRigArgumentDefaultNode::GenerateChildContent(IDetailChildrenBuilder
 	for (TSharedRef<SWidget> Pin : Pins)
 	{
 		TSharedRef<SGraphPin> SPin = StaticCastSharedRef<SGraphPin>(Pin);
-		SPin->GetPinObj()->bNotConnectable = true;
+		SPin->EnableDragAndDrop(false);
 	}
 }
 
diff --git a/Engine/Source/Editor/GraphEditor/Private/SGraphPin.cpp b/Engine/Source/Editor/GraphEditor/Private/SGraphPin.cpp
--- a/Engine/Source/Editor/GraphEditor/Private/SGraphPin.cpp
+++ b/Engine/Source/Editor/GraphEditor/Private/SGraphPin.cpp
@@ -79,7 +79,7 @@ SGraphPin::SGraphPin()
 	, bOnlyShowDefaultValue(false)
 	, bIsMovingLinks(false)
 	, bUsePinColorForText(false)
-
+	, bDragAndDropEnabled(true)
 {
 	IsEditable = true;
 
@@ -452,7 +452,7 @@ FReply SGraphPin::OnPinMouseDown( const FGeometry& SenderGeometry, const FPointe
 				}
 			}
 
-			if (!MouseEvent.IsShiftDown() && !GraphPinObj->bNotConnectable)
+			if (!MouseEvent.IsShiftDown() && !GraphPinObj->bNotConnectable && bDragAndDropEnabled)
 			{
 				// Start a drag-drop on the pin
 				if (ensure(OwnerNodePinned.IsValid()))
diff --git a/Engine/Source/Editor/GraphEditor/Public/SGraphPin.h b/Engine/Source/Editor/GraphEditor/Public/SGraphPin.h
--- a/Engine/Source/Editor/GraphEditor/Public/SGraphPin.h
+++ b/Engine/Source/Editor/GraphEditor/Public/SGraphPin.h
@@ -211,6 +211,9 @@ public:
 	/** @returns true if we have a valid PinValueInspector tooltip */
 	bool HasInteractiveTooltip() const;
 
+	/** Enables or disables drag and drop on the pin */
+	void EnableDragAndDrop(bool bEnable) { bDragAndDropEnabled = bEnable; }
+
 protected:
 
 	/** If true the graph pin subclass is responsible for setting the IsEnabled delegates for the aspects it cares about. If false, the default value widget enabling is done by the base class */
@@ -347,4 +350,7 @@ protected:
 
 	/** TRUE if the pin should use the Pin's color for the text */
 	bool bUsePinColorForText;
+
+	/** TRUE if the pin should allow any drag and drop */
+	bool bDragAndDropEnabled;
 };
