Remove '/' character in debugger watch text if the next property is a container element

diff --git a/Engine/Source/Editor/GraphEditor/Private/KismetNodes/SGraphNodeK2Base.cpp b/Engine/Source/Editor/GraphEditor/Private/KismetNodes/SGraphNodeK2Base.cpp
--- a/Engine/Source/Editor/GraphEditor/Private/KismetNodes/SGraphNodeK2Base.cpp
+++ b/Engine/Source/Editor/GraphEditor/Private/KismetNodes/SGraphNodeK2Base.cpp
@@ -629,7 +629,11 @@ void SGraphNodeK2Base::GetNodeInfoPopups(FNodeInfoContext* Context, TArray<FGrap
 
 								for (const FName& PathName : WatchedPin.GetPathToProperty())
 								{
-									WatchName += TEXT("/");
+									if (!PathName.ToString().StartsWith("["))
+									{
+										WatchName += TEXT("/");
+									}
+
 									WatchName += PathName.ToString();
 								}
 							}
