Remove pin tooltips on pin destruction, preventing stale tooltip crashes.

diff --git a/Engine/Source/Editor/GraphEditor/Private/SGraphPin.cpp b/Engine/Source/Editor/GraphEditor/Private/SGraphPin.cpp
--- a/Engine/Source/Editor/GraphEditor/Private/SGraphPin.cpp
+++ b/Engine/Source/Editor/GraphEditor/Private/SGraphPin.cpp
@@ -147,6 +147,14 @@ SGraphPin::SGraphPin()
 	CachedImg_Pin_BackgroundHovered = FEditorStyle::GetBrush( NAME_Pin_BackgroundHovered );
 }
 
+SGraphPin::~SGraphPin()
+{
+	if (ValueInspectorTooltip.IsValid())
+	{
+		ValueInspectorTooltip.Pin()->TryDismissTooltip();
+	}
+}
+
 void SGraphPin::Construct(const FArguments& InArgs, UEdGraphPin* InPin)
 {
 	bUsePinColorForText = InArgs._UsePinColorForText;
diff --git a/Engine/Source/Editor/GraphEditor/Public/SGraphPin.h b/Engine/Source/Editor/GraphEditor/Public/SGraphPin.h
--- a/Engine/Source/Editor/GraphEditor/Public/SGraphPin.h
+++ b/Engine/Source/Editor/GraphEditor/Public/SGraphPin.h
@@ -93,6 +93,7 @@ public:
 
 public:
 	SGraphPin();
+	virtual ~SGraphPin();
 
 	/** Set attribute for determining if pin is editable */
 	void SetIsEditable(TAttribute<bool> InIsEditable);
