EdGraph: Keep soft handles to pins and nodes during link

diff --git a/Engine/Source/Editor/GraphEditor/Private/DragConnection.cpp b/Engine/Source/Editor/GraphEditor/Private/DragConnection.cpp
--- a/Engine/Source/Editor/GraphEditor/Private/DragConnection.cpp
+++ b/Engine/Source/Editor/GraphEditor/Private/DragConnection.cpp
@@ -8,6 +8,7 @@
 #include "EdGraph/EdGraph.h"
 #include "SGraphPanel.h"
 #include "ScopedTransaction.h"
+#include "EdGraphHandleTypes.h"
 
 TSharedRef<FDragConnection> FDragConnection::New(const TSharedRef<SGraphPanel>& GraphPanel, const FDraggedPinTable& DraggedPins)
 {
@@ -201,27 +202,36 @@ FReply FDragConnection::DroppedOnPin(FVector2D ScreenPosition, FVector2D GraphPo
 	TArray<UEdGraphPin*> ValidSourcePins;
 	ValidateGraphPinList(/*out*/ ValidSourcePins);
 
+	// store the pins as pin tuples since the structure of the
+	// graph may change during the creation of a connection
+	TArray<FEdGraphPinHandle> ValidSourcePinHandles;
+	for(const UEdGraphPin* ValidSourcePin : ValidSourcePins)
+	{
+		ValidSourcePinHandles.Add(ValidSourcePin);
+	}
+
 	const FScopedTransaction Transaction( NSLOCTEXT("UnrealEd", "GraphEd_CreateConnection", "Create Pin Link") );
 
-	UEdGraphPin* PinB = GetHoveredPin();
+	FEdGraphPinHandle PinB(GetHoveredPin());
 	bool bError = false;
-	TSet<UEdGraphNode*> NodeList;
+	TSet<FEdGraphNodeHandle> NodeList;
 
-	for (UEdGraphPin* PinA : ValidSourcePins)
+	for (const FEdGraphPinHandle& PinA : ValidSourcePinHandles)
 	{
-		if ((PinA != NULL) && (PinB != NULL))
+		if ((PinA.GetPin() != NULL) && (PinB.GetPin() != NULL))
 		{
-			UEdGraph* MyGraphObj = PinA->GetOwningNode()->GetGraph();
+			const UEdGraph* MyGraphObj = PinA.GetGraph();
 
-			if (MyGraphObj->GetSchema()->TryCreateConnection(PinA, PinB))
+			// the pin may change during the creation of the link
+			if (MyGraphObj->GetSchema()->TryCreateConnection(PinA.GetPin(), PinB.GetPin()))
 			{
-				if (!PinA->IsPendingKill())
+				if (PinA.GetPin() && !PinA.GetPin()->IsPendingKill())
 				{
-					NodeList.Add(PinA->GetOwningNode());
+					NodeList.Add(PinA.GetPin()->GetOwningNode());
 				}
-				if (!PinB->IsPendingKill())
+				if (PinB.GetPin() && !PinB.GetPin()->IsPendingKill())
 				{
-					NodeList.Add(PinB->GetOwningNode());
+					NodeList.Add(PinB.GetNode());
 				}
 			}
 		}
@@ -234,8 +244,10 @@ FReply FDragConnection::DroppedOnPin(FVector2D ScreenPosition, FVector2D GraphPo
 	// Send all nodes that received a new pin connection a notification
 	for (auto It = NodeList.CreateConstIterator(); It; ++It)
 	{
-		UEdGraphNode* Node = (*It);
-		Node->NodeConnectionListChanged();
+		if(UEdGraphNode* Node = It->GetNode())
+		{
+			Node->NodeConnectionListChanged();
+		}
 	}
 
 	if (bError)
diff --git a/Engine/Source/Editor/GraphEditor/Public/EdGraphHandleTypes.h b/Engine/Source/Editor/GraphEditor/Public/EdGraphHandleTypes.h
new file mode 100644
--- /dev/null
+++ b/Engine/Source/Editor/GraphEditor/Public/EdGraphHandleTypes.h
@@ -0,0 +1,103 @@
+// Copyright Epic Games, Inc. All Rights Reserved.
+
+#pragma once
+
+#include "CoreMinimal.h"
+#include "EdGraph/EdGraph.h"
+#include "EdGraph/EdGraphNode.h"
+#include "EdGraph/EdGraphPin.h"
+#include "UObject/SoftObjectPtr.h"
+
+struct GRAPHEDITOR_API FEdGraphNodeHandle
+{
+public:
+
+	FORCEINLINE FEdGraphNodeHandle(const UEdGraphNode* InNode)
+		: Graph(InNode->GetGraph())
+		, NodeName(InNode->GetFName())
+	{}
+
+	FORCEINLINE FEdGraphNodeHandle(const FEdGraphNodeHandle& InOther)
+		: Graph(InOther.Graph)
+		, NodeName(InOther.NodeName)
+	{}
+
+	friend FORCEINLINE uint32 GetTypeHash(const FEdGraphNodeHandle& InHandle)
+	{
+		return HashCombine(GetTypeHash(InHandle.Graph.ToSoftObjectPath()), GetTypeHash(InHandle.NodeName));
+	}
+
+	FORCEINLINE bool operator ==(const FEdGraphNodeHandle& InOther) const
+	{
+		return Graph.GetUniqueID() == InOther.Graph.GetUniqueID() &&
+			NodeName.IsEqual(InOther.NodeName, ENameCase::CaseSensitive, true);
+	}
+
+	FORCEINLINE const UEdGraph* GetGraph() const { return Graph.Get(); }
+
+	FORCEINLINE UEdGraphNode* GetNode() const
+	{
+		if(const UEdGraph* EdGraph = GetGraph())
+		{
+			const TObjectPtr<UEdGraphNode>* Node = EdGraph->Nodes.FindByPredicate([this](TObjectPtr<UEdGraphNode> Node) -> bool
+			{
+				return Node->GetFName() == NodeName;
+			});
+			if(Node)
+			{
+				return Node->Get();
+			}
+		}
+		return nullptr;
+	}
+
+private:
+	TSoftObjectPtr<UEdGraph> Graph;
+	FName NodeName;
+};
+
+struct GRAPHEDITOR_API FEdGraphPinHandle : FEdGraphNodeHandle
+{
+public:
+
+	FORCEINLINE FEdGraphPinHandle(const UEdGraphPin* InPin)
+		: FEdGraphNodeHandle(InPin->GetOwningNode())
+		, PinName(InPin->GetFName())
+	{}
+		
+	FORCEINLINE FEdGraphPinHandle(const FEdGraphPinHandle& InOther)
+		: FEdGraphNodeHandle(InOther)
+		, PinName(InOther.PinName)
+	{}
+
+	friend FORCEINLINE uint32 GetTypeHash(const FEdGraphPinHandle& InHandle)
+	{
+		return HashCombine(GetTypeHash((FEdGraphNodeHandle)InHandle), GetTypeHash(InHandle.PinName));
+	}
+
+	FORCEINLINE bool operator ==(const FEdGraphPinHandle& InOther) const
+	{
+		return FEdGraphNodeHandle::operator==(InOther) &&
+			PinName.IsEqual(InOther.PinName, ENameCase::CaseSensitive, true);
+	}
+	
+	FORCEINLINE UEdGraphPin* GetPin() const
+	{
+		if(UEdGraphNode* EdNode = GetNode())
+		{
+			const UEdGraphPin*const* Pin = EdNode->Pins.FindByPredicate([this](UEdGraphPin* Pin) -> bool
+			{
+				return Pin->GetFName() == PinName;
+			});
+			if(Pin)
+			{
+				return (UEdGraphPin*)*Pin;
+			}
+		}
+		return nullptr;
+	}
+
+private:
+	FName PinName;
+};
+
