Shift+LeftClick to duplicate drag a pin connection

diff --git a/Engine/Source/Editor/GraphEditor/Private/SGraphPin.cpp b/Engine/Source/Editor/GraphEditor/Private/SGraphPin.cpp
--- a/Engine/Source/Editor/GraphEditor/Private/SGraphPin.cpp
+++ b/Engine/Source/Editor/GraphEditor/Private/SGraphPin.cpp
@@ -363,7 +363,7 @@ FReply SGraphPin::OnPinMouseDown( const FGeometry& SenderGeometry, const FPointe
 			}
 
 			TSharedPtr<SGraphNode> OwnerNodePinned = OwnerNodePtr.Pin();
-			if (MouseEvent.IsControlDown() && (GraphPinObj->LinkedTo.Num() > 0))
+			if ((MouseEvent.IsControlDown() || MouseEvent.IsShiftDown()) && (GraphPinObj->LinkedTo.Num() > 0))
 			{
 				// Get a reference to the owning panel widget
 				check(OwnerNodePinned.IsValid());
@@ -446,8 +446,11 @@ FReply SGraphPin::OnPinMouseDown( const FGeometry& SenderGeometry, const FPointe
 				// Note: that for some nodes, this can cause reconstruction. In that case, pins we had previously linked to may now be destroyed. 
 				//       So the break MUST come after the SpawnPinDragEvent(), since that acquires handles from PinArray (the pins need to be
 				//       around for us to construct valid handles from).
-				const UEdGraphSchema* Schema = GraphPinObj->GetSchema();
-				Schema->BreakPinLinks(*GraphPinObj, true);
+				if (MouseEvent.IsControlDown())
+				{
+					const UEdGraphSchema* Schema = GraphPinObj->GetSchema();
+					Schema->BreakPinLinks(*GraphPinObj, true);
+				}
 
 				if (DragEvent.IsValid())
 				{
