Control Rig: Fix for crash during pin wiring

diff --git a/Engine/Source/Editor/GraphEditor/Public/EdGraphHandleTypes.h b/Engine/Source/Editor/GraphEditor/Public/EdGraphHandleTypes.h
--- a/Engine/Source/Editor/GraphEditor/Public/EdGraphHandleTypes.h
+++ b/Engine/Source/Editor/GraphEditor/Public/EdGraphHandleTypes.h
@@ -13,8 +13,8 @@ struct GRAPHEDITOR_API FEdGraphNodeHandle
 public:
 
 	FORCEINLINE FEdGraphNodeHandle(const UEdGraphNode* InNode)
-		: Graph(InNode->GetGraph())
-		, NodeName(InNode->GetFName())
+		: Graph(InNode ? InNode->GetGraph() : nullptr)
+		, NodeName(InNode ? InNode->GetFName() : NAME_None)
 	{}
 
 	FORCEINLINE FEdGraphNodeHandle(const FEdGraphNodeHandle& InOther)
@@ -61,8 +61,8 @@ struct GRAPHEDITOR_API FEdGraphPinHandle : FEdGraphNodeHandle
 public:
 
 	FORCEINLINE FEdGraphPinHandle(const UEdGraphPin* InPin)
-		: FEdGraphNodeHandle(InPin->GetOwningNode())
-		, PinName(InPin->GetFName())
+		: FEdGraphNodeHandle(InPin ? InPin->GetOwningNode() : nullptr)
+		, PinName(InPin ? InPin->GetFName() : NAME_None)
 	{}
 		
 	FORCEINLINE FEdGraphPinHandle(const FEdGraphPinHandle& InOther)
