Control Rig: Execute pins can't be connected in the proper order

diff --git a/Engine/Source/Editor/GraphEditor/Public/EdGraphHandleTypes.h b/Engine/Source/Editor/GraphEditor/Public/EdGraphHandleTypes.h
--- a/Engine/Source/Editor/GraphEditor/Public/EdGraphHandleTypes.h
+++ b/Engine/Source/Editor/GraphEditor/Public/EdGraphHandleTypes.h
@@ -63,22 +63,29 @@ public:
 	FORCEINLINE FEdGraphPinHandle(const UEdGraphPin* InPin)
 		: FEdGraphNodeHandle(InPin ? InPin->GetOwningNode() : nullptr)
 		, PinName(InPin ? InPin->GetFName() : NAME_None)
+		, PinDirection(InPin ? InPin->Direction.GetValue() : EEdGraphPinDirection::EGPD_Input)
 	{}
 		
 	FORCEINLINE FEdGraphPinHandle(const FEdGraphPinHandle& InOther)
 		: FEdGraphNodeHandle(InOther)
 		, PinName(InOther.PinName)
+		, PinDirection(InOther.PinDirection)
 	{}
 
 	friend FORCEINLINE uint32 GetTypeHash(const FEdGraphPinHandle& InHandle)
 	{
-		return HashCombine(GetTypeHash((FEdGraphNodeHandle)InHandle), GetTypeHash(InHandle.PinName));
+		return HashCombine(
+				HashCombine(
+					GetTypeHash((FEdGraphNodeHandle)InHandle),
+					GetTypeHash(InHandle.PinName)),
+				GetTypeHash(InHandle.PinDirection));
 	}
 
 	FORCEINLINE bool operator ==(const FEdGraphPinHandle& InOther) const
 	{
 		return FEdGraphNodeHandle::operator==(InOther) &&
-			PinName.IsEqual(InOther.PinName, ENameCase::CaseSensitive, true);
+			PinName.IsEqual(InOther.PinName, ENameCase::CaseSensitive, true) &&
+			PinDirection == InOther.PinDirection;
 	}
 	
 	FORCEINLINE UEdGraphPin* GetPin() const
@@ -87,7 +94,7 @@ public:
 		{
 			const UEdGraphPin*const* Pin = EdNode->Pins.FindByPredicate([this](UEdGraphPin* Pin) -> bool
 			{
-				return Pin->GetFName() == PinName;
+				return Pin->GetFName() == PinName && Pin->Direction == PinDirection;
 			});
 			if(Pin)
 			{
@@ -99,5 +106,6 @@ public:
 
 private:
 	FName PinName;
+	EEdGraphPinDirection PinDirection;
 };
 
