Fixed pin handles not matching pin via persistent guid first, which introduced issues with Niagara pin handling as Niagara has hidden pins of the same name "None" as part of a core paradigm

diff --git a/Engine/Source/Editor/GraphEditor/Public/EdGraphHandleTypes.h b/Engine/Source/Editor/GraphEditor/Public/EdGraphHandleTypes.h
--- a/Engine/Source/Editor/GraphEditor/Public/EdGraphHandleTypes.h
+++ b/Engine/Source/Editor/GraphEditor/Public/EdGraphHandleTypes.h
@@ -64,17 +64,25 @@ public:
 		: FEdGraphNodeHandle(InPin ? InPin->GetOwningNode() : nullptr)
 		, PinName(InPin ? InPin->GetFName() : NAME_None)
 		, PinDirection(InPin ? InPin->Direction.GetValue() : EEdGraphPinDirection::EGPD_Input)
+		, PersistentPinGuid(InPin ? InPin->PersistentGuid : FGuid())
 	{}
 		
 	FORCEINLINE FEdGraphPinHandle(const FEdGraphPinHandle& InOther)
 		: FEdGraphNodeHandle(InOther)
 		, PinName(InOther.PinName)
 		, PinDirection(InOther.PinDirection)
+		, PersistentPinGuid(InOther.PersistentPinGuid)
 	{}
 
 	friend FORCEINLINE uint32 GetTypeHash(const FEdGraphPinHandle& InHandle)
 	{
-		return HashCombine(
+		return InHandle.PersistentPinGuid.IsValid() ?
+			HashCombine(
+			   HashCombine(
+				   GetTypeHash((FEdGraphNodeHandle)InHandle),
+				   GetTypeHash(InHandle.PersistentPinGuid)),
+			   GetTypeHash(InHandle.PinDirection))
+			: HashCombine(
 				HashCombine(
 					GetTypeHash((FEdGraphNodeHandle)InHandle),
 					GetTypeHash(InHandle.PinName)),
@@ -84,7 +92,7 @@ public:
 	FORCEINLINE bool operator ==(const FEdGraphPinHandle& InOther) const
 	{
 		return FEdGraphNodeHandle::operator==(InOther) &&
-			PinName.IsEqual(InOther.PinName, ENameCase::CaseSensitive, true) &&
+			((PersistentPinGuid == InOther.PersistentPinGuid) || (PinName.IsEqual(InOther.PinName, ENameCase::CaseSensitive, true))) &&
 			PinDirection == InOther.PinDirection;
 	}
 	
@@ -94,8 +102,17 @@ public:
 		{
 			const UEdGraphPin*const* Pin = EdNode->Pins.FindByPredicate([this](UEdGraphPin* Pin) -> bool
 			{
-				return Pin->GetFName() == PinName && Pin->Direction == PinDirection;
+				return Pin->PersistentGuid == PersistentPinGuid && Pin->Direction == PinDirection;
 			});
+
+			if(Pin == nullptr)
+			{
+				Pin = EdNode->Pins.FindByPredicate([this](UEdGraphPin* Pin) -> bool
+				{
+					return Pin->GetFName() == PinName && Pin->Direction == PinDirection;
+				});
+			}
+			
 			if(Pin)
 			{
 				return (UEdGraphPin*)*Pin;
@@ -107,5 +124,6 @@ public:
 private:
 	FName PinName;
 	EEdGraphPinDirection PinDirection;
+	FGuid PersistentPinGuid;
 };
 
