Fixed invalid guids being candidates for pin matching during OnDropped

diff --git a/Engine/Source/Editor/GraphEditor/Public/EdGraphHandleTypes.h b/Engine/Source/Editor/GraphEditor/Public/EdGraphHandleTypes.h
--- a/Engine/Source/Editor/GraphEditor/Public/EdGraphHandleTypes.h
+++ b/Engine/Source/Editor/GraphEditor/Public/EdGraphHandleTypes.h
@@ -102,7 +102,7 @@ public:
 		{
 			const UEdGraphPin*const* Pin = EdNode->Pins.FindByPredicate([this](UEdGraphPin* Pin) -> bool
 			{
-				return Pin->PersistentGuid == PersistentPinGuid && Pin->Direction == PinDirection;
+				return Pin->PersistentGuid.IsValid() && Pin->PersistentGuid == PersistentPinGuid && Pin->Direction == PinDirection;
 			});
 
 			if(Pin == nullptr)
