Improved behavior on links in graphs that support multiple links per pin.
Fixes issues with highlighting and alt-click delete on a link.

diff --git a/Engine/Source/Editor/GraphEditor/Private/ConnectionDrawingPolicy.cpp b/Engine/Source/Editor/GraphEditor/Private/ConnectionDrawingPolicy.cpp
--- a/Engine/Source/Editor/GraphEditor/Private/ConnectionDrawingPolicy.cpp
+++ b/Engine/Source/Editor/GraphEditor/Private/ConnectionDrawingPolicy.cpp
@@ -150,6 +150,9 @@ void FConnectionDrawingPolicy::SetHoveredPins(const TSet< FEdGraphPinReference >
 		}
 	}
 
+	// When we have only a single pin selected, we'll extend selection to apply the hover effect on the links
+	const bool bMakeConnectedPinsHovered = (InHoveredPins.Num() == 1);
+
 	// Convert the widget pointer for hovered pins to be EdGraphPin pointers for their connected nets (both ends of any connection)
 	for (auto PinIt = InHoveredPins.CreateConstIterator(); PinIt; ++PinIt)
 	{
@@ -159,9 +162,12 @@ void FConnectionDrawingPolicy::SetHoveredPins(const TSet< FEdGraphPinReference >
 			{
 				HoveredPins.Add(Pin);
 
-				for (auto LinkIt = Pin->LinkedTo.CreateConstIterator(); LinkIt; ++LinkIt)
+				if (bMakeConnectedPinsHovered)
 				{
-					HoveredPins.Add(*LinkIt);
+					for (auto LinkIt = Pin->LinkedTo.CreateConstIterator(); LinkIt; ++LinkIt)
+					{
+						HoveredPins.Add(*LinkIt);
+					}
 				}
 			}
 		}
@@ -599,3 +605,14 @@ bool FGraphSplineOverlapResult::GetPins(const class SGraphPanel& InGraphPanel, U
 	return (OutPin1 != nullptr) && (OutPin2 != nullptr);
 }
 
+void FGraphSplineOverlapResult::GetPinWidgets(const class SGraphPanel& InGraphPanel, TSharedPtr<class SGraphPin>& OutPin1, TSharedPtr<class SGraphPin>& OutPin2) const
+{
+	OutPin1 = nullptr;
+	OutPin2 = nullptr;
+
+	if (IsValid())
+	{
+		OutPin1 = Pin1Handle.FindInGraphPanel(InGraphPanel);
+		OutPin2 = Pin2Handle.FindInGraphPanel(InGraphPanel);
+	}
+}
diff --git a/Engine/Source/Editor/GraphEditor/Private/SGraphPanel.cpp b/Engine/Source/Editor/GraphEditor/Private/SGraphPanel.cpp
--- a/Engine/Source/Editor/GraphEditor/Private/SGraphPanel.cpp
+++ b/Engine/Source/Editor/GraphEditor/Private/SGraphPanel.cpp
@@ -505,28 +505,48 @@ int32 SGraphPanel::OnPaint( const FPaintArgs& Args, const FGeometry& AllottedGeo
 
 void SGraphPanel::OnSplineHoverStateChanged(const FGraphSplineOverlapResult& NewSplineHoverState)
 {
-	TSharedPtr<SGraphPin> OldPinWidget = PreviousFrameSplineOverlap.GetBestPinWidget(*this);
+	TSharedPtr<SGraphPin> OldPin1Widget;
+	TSharedPtr<SGraphPin> OldPin2Widget;
+	PreviousFrameSplineOverlap.GetPinWidgets(*this, OldPin1Widget, OldPin2Widget);
+
 	PreviousFrameSplineOverlap = NewSplineHoverState;
-	TSharedPtr<SGraphPin> NewPinWidget = PreviousFrameSplineOverlap.GetBestPinWidget(*this);
+
+	TSharedPtr<SGraphPin> NewPin1Widget;
+	TSharedPtr<SGraphPin> NewPin2Widget;
+	PreviousFrameSplineOverlap.GetPinWidgets(*this, NewPin1Widget, NewPin2Widget);
 
 	PreviousFrameSavedMousePosForSplineOverlap = SavedMousePosForOnPaintEventLocalSpace;
 
-	// Handle mouse enter/leaves on the associated pin
-	if (OldPinWidget != NewPinWidget)
+	// Handle exiting hovering on the pins
+	if (OldPin1Widget.IsValid() && OldPin1Widget != NewPin1Widget && OldPin1Widget != NewPin2Widget)
 	{
-		if (OldPinWidget.IsValid())
-		{
-			OldPinWidget->OnMouseLeave(LastPointerEvent);
-		}
+		OldPin1Widget->OnMouseLeave(LastPointerEvent);
+	}
 
-		if (NewPinWidget.IsValid())
-		{
-			NewPinWidget->OnMouseEnter(LastPointerGeometry, LastPointerEvent);
+	if (OldPin2Widget.IsValid() && OldPin2Widget != NewPin2Widget && OldPin2Widget != NewPin2Widget)
+	{
+		OldPin2Widget->OnMouseLeave(LastPointerEvent);
+	}
 
-			// Get the pin/wire glowing quicker, since it's a direct selection (this time was already set to 'now' as part of entering the pin)
-			//@TODO: Source this parameter from the graph rendering settings once it is there (see code in ApplyHoverDeemphasis)
-			TimeWhenMouseEnteredPin -= 0.75f;
-		}
+	// Handle enter hovering on the pins
+	bool bChangedHover = false;
+	if (NewPin1Widget.IsValid() && NewPin1Widget != OldPin1Widget && NewPin1Widget != OldPin2Widget)
+	{
+		NewPin1Widget->OnMouseEnter(LastPointerGeometry, LastPointerEvent);
+		bChangedHover = true;
+	}
+
+	if (NewPin2Widget.IsValid() && NewPin2Widget != OldPin1Widget && NewPin2Widget != OldPin2Widget)
+	{
+		NewPin2Widget->OnMouseEnter(LastPointerGeometry, LastPointerEvent);
+		bChangedHover = true;
+	}
+
+	if (bChangedHover)
+	{
+		// Get the pin/wire glowing quicker, since it's a direct selection (this time was already set to 'now' as part of entering the pin)
+		//@TODO: Source this parameter from the graph rendering settings once it is there (see code in ApplyHoverDeemphasis)
+		TimeWhenMouseEnteredPin -= 0.75f;
 	}
 }
 
@@ -642,7 +662,15 @@ FReply SGraphPanel::OnMouseButtonDown(const FGeometry& MyGeometry, const FPointe
 {
 	if ((MouseEvent.GetEffectingButton() == EKeys::LeftMouseButton) && (MouseEvent.IsAltDown() || MouseEvent.IsControlDown()))
 	{
-		if (SGraphPin* BestPinFromHoveredSpline = GetBestPinFromHoveredSpline())
+		// Intercept alt-left clicking on the hovered spline for targeted break link
+		UEdGraphPin* Pin1;
+		UEdGraphPin* Pin2;
+		if (MouseEvent.IsAltDown() && PreviousFrameSplineOverlap.GetPins(*this, Pin1, Pin2))
+		{
+			const UEdGraphSchema* Schema = GraphObj->GetSchema();
+			Schema->BreakSinglePinLink(Pin1, Pin2);
+		}
+		else if (SGraphPin* BestPinFromHoveredSpline = GetBestPinFromHoveredSpline())
 		{
 			return BestPinFromHoveredSpline->OnPinMouseDown(MyGeometry, MouseEvent);
 		}
diff --git a/Engine/Source/Editor/GraphEditor/Public/GraphSplineOverlapResult.h b/Engine/Source/Editor/GraphEditor/Public/GraphSplineOverlapResult.h
--- a/Engine/Source/Editor/GraphEditor/Public/GraphSplineOverlapResult.h
+++ b/Engine/Source/Editor/GraphEditor/Public/GraphSplineOverlapResult.h
@@ -70,6 +70,7 @@ public:
 	}
 
 	bool GetPins(const class SGraphPanel& InGraphPanel, UEdGraphPin*& OutPin1, UEdGraphPin*& OutPin2) const;
+	void GetPinWidgets(const class SGraphPanel& InGraphPanel, TSharedPtr<class SGraphPin>& OutPin1, TSharedPtr<class SGraphPin>& OutPin2) const;
 
 protected:
 	FGraphPinHandle Pin1Handle;
