Fixed typo found by static analysis

#rnx
#ue UE-150237

diff --git a/Engine/Source/Editor/GraphEditor/Private/SGraphPanel.cpp b/Engine/Source/Editor/GraphEditor/Private/SGraphPanel.cpp
--- a/Engine/Source/Editor/GraphEditor/Private/SGraphPanel.cpp
+++ b/Engine/Source/Editor/GraphEditor/Private/SGraphPanel.cpp
@@ -523,7 +523,7 @@ void SGraphPanel::OnSplineHoverStateChanged(const FGraphSplineOverlapResult& New
 		OldPin1Widget->OnMouseLeave(LastPointerEvent);
 	}
 
-	if (OldPin2Widget.IsValid() && OldPin2Widget != NewPin2Widget && OldPin2Widget != NewPin2Widget)
+	if (OldPin2Widget.IsValid() && OldPin2Widget != NewPin1Widget && OldPin2Widget != NewPin2Widget)
 	{
 		OldPin2Widget->OnMouseLeave(LastPointerEvent);
 	}
