Move SGraphPinVector4.cpp from NiagaraPins folder to KismetPins folder

diff --git a/Engine/Source/Editor/GraphEditor/Private/KismetPins/SGraphPinVector4.cpp b/Engine/Source/Editor/GraphEditor/Private/KismetPins/SGraphPinVector4.cpp
new file mode 100644
--- /dev/null
+++ b/Engine/Source/Editor/GraphEditor/Private/KismetPins/SGraphPinVector4.cpp
@@ -0,0 +1,321 @@
+// Copyright Epic Games, Inc. All Rights Reserved.
+
+
+#include "KismetPins/SGraphPinVector4.h"
+#include "Widgets/SBoxPanel.h"
+#include "Widgets/Input/SNumericEntryBox.h"
+#include "ScopedTransaction.h"
+
+
+#define LOCTEXT_NAMESPACE "VectorTextBox"
+
+
+//Class implementation to create 3 editable text boxes to represent vector/rotator graph pin
+class SVector4TextBox : public SCompoundWidget
+{
+public:
+
+	SLATE_BEGIN_ARGS(SVector4TextBox) {}
+	SLATE_ATTRIBUTE(FString, VisibleText_0)
+		SLATE_ATTRIBUTE(FString, VisibleText_1)
+		SLATE_ATTRIBUTE(FString, VisibleText_2)
+		SLATE_ATTRIBUTE(FString, VisibleText_3)
+		SLATE_EVENT(FOnFloatValueCommitted, OnFloatCommitted_Box_0)
+		SLATE_EVENT(FOnFloatValueCommitted, OnFloatCommitted_Box_1)
+		SLATE_EVENT(FOnFloatValueCommitted, OnFloatCommitted_Box_2)
+		SLATE_EVENT(FOnFloatValueCommitted, OnFloatCommitted_Box_3)
+		SLATE_END_ARGS()
+
+		//Construct editable text boxes with the appropriate getter & setter functions along with tool tip text
+		void Construct(const FArguments& InArgs)
+	{
+			VisibleText_0 = InArgs._VisibleText_0;
+			VisibleText_1 = InArgs._VisibleText_1;
+			VisibleText_2 = InArgs._VisibleText_2;
+			VisibleText_3 = InArgs._VisibleText_3;
+			const FLinearColor LabelClr = FLinearColor(1.f, 1.f, 1.f, 0.4f);
+
+			this->ChildSlot
+				[
+					SNew(SVerticalBox)
+					+ SVerticalBox::Slot()
+					.AutoHeight()
+					.Padding(0)
+					[
+						SNew(SHorizontalBox)
+						+ SHorizontalBox::Slot()
+						.AutoWidth().Padding(2).HAlign(HAlign_Fill)
+						[
+							//Create Text box 0 
+							SNew(SNumericEntryBox<float>)
+							.LabelVAlign(VAlign_Center)
+							.Label()
+							[
+								SNew(STextBlock)
+								.Font(FAppStyle::GetFontStyle("Graph.VectorEditableTextBox"))
+								.Text(LOCTEXT("VectorNodeXAxisValueLabel", "X"))
+								.ColorAndOpacity(LabelClr)
+							]
+							.Value(this, &SVector4TextBox::GetTypeInValue_0)
+								.OnValueCommitted(InArgs._OnFloatCommitted_Box_0)
+								.Font(FAppStyle::GetFontStyle("Graph.VectorEditableTextBox"))
+								.UndeterminedString(LOCTEXT("MultipleValues", "Multiple Values"))
+								.ToolTipText(LOCTEXT("VectorNodeXAxisValueLabel_ToolTip", "X value"))
+								.EditableTextBoxStyle(&FAppStyle::GetWidgetStyle<FEditableTextBoxStyle>("Graph.VectorEditableTextBox"))
+								.BorderForegroundColor(FLinearColor::White)
+								.BorderBackgroundColor(FLinearColor::White)
+						]
+						+ SHorizontalBox::Slot()
+							.AutoWidth().Padding(2).HAlign(HAlign_Fill)
+							[
+								//Create Text box 1
+								SNew(SNumericEntryBox<float>)
+								.LabelVAlign(VAlign_Center)
+								.Label()
+								[
+									SNew(STextBlock)
+									.Font(FAppStyle::GetFontStyle("Graph.VectorEditableTextBox"))
+									.Text(LOCTEXT("VectorNodeYAxisValueLabel", "Y"))
+									.ColorAndOpacity(LabelClr)
+								]
+								.Value(this, &SVector4TextBox::GetTypeInValue_1)
+									.OnValueCommitted(InArgs._OnFloatCommitted_Box_1)
+									.Font(FAppStyle::GetFontStyle("Graph.VectorEditableTextBox"))
+									.UndeterminedString(LOCTEXT("MultipleValues", "Multiple Values"))
+									.ToolTipText(LOCTEXT("VectorNodeYAxisValueLabel_ToolTip", "Y value"))
+									.EditableTextBoxStyle(&FAppStyle::GetWidgetStyle<FEditableTextBoxStyle>("Graph.VectorEditableTextBox"))
+									.BorderForegroundColor(FLinearColor::White)
+									.BorderBackgroundColor(FLinearColor::White)
+							]
+						+ SHorizontalBox::Slot()
+							.AutoWidth().Padding(2).HAlign(HAlign_Fill)
+							[
+								//Create Text box 2
+								SNew(SNumericEntryBox<float>)
+								.LabelVAlign(VAlign_Center)
+								.Label()
+								[
+									SNew(STextBlock)
+									.Font(FAppStyle::GetFontStyle("Graph.VectorEditableTextBox"))
+									.Text(LOCTEXT("VectorNodeZAxisValueLabel", "Z"))
+									.ColorAndOpacity(LabelClr)
+								]
+								.Value(this, &SVector4TextBox::GetTypeInValue_2)
+									.OnValueCommitted(InArgs._OnFloatCommitted_Box_2)
+									.Font(FAppStyle::GetFontStyle("Graph.VectorEditableTextBox"))
+									.UndeterminedString(LOCTEXT("MultipleValues", "Multiple Values"))
+									.ToolTipText(LOCTEXT("VectorNodeZAxisValueLabel_ToolTip", "Z value"))
+									.EditableTextBoxStyle(&FAppStyle::GetWidgetStyle<FEditableTextBoxStyle>("Graph.VectorEditableTextBox"))
+									.BorderForegroundColor(FLinearColor::White)
+									.BorderBackgroundColor(FLinearColor::White)
+							]
+						+ SHorizontalBox::Slot()
+							.AutoWidth().Padding(2).HAlign(HAlign_Fill)
+							[
+								//Create Text box 3
+								SNew(SNumericEntryBox<float>)
+								.LabelVAlign(VAlign_Center)
+								.Label()
+								[
+									SNew(STextBlock)
+									.Font(FAppStyle::GetFontStyle("Graph.VectorEditableTextBox"))
+									.Text(LOCTEXT("VectorNodeWAxisValueLabel", "W"))
+									.ColorAndOpacity(LabelClr)
+								]
+								.Value(this, &SVector4TextBox::GetTypeInValue_3)
+									.OnValueCommitted(InArgs._OnFloatCommitted_Box_3)
+									.Font(FAppStyle::GetFontStyle("Graph.VectorEditableTextBox"))
+									.UndeterminedString(LOCTEXT("MultipleValues", "Multiple Values"))
+									.ToolTipText(LOCTEXT("VectorNodeWAxisValueLabel_ToolTip", "W value"))
+									.EditableTextBoxStyle(&FAppStyle::GetWidgetStyle<FEditableTextBoxStyle>("Graph.VectorEditableTextBox"))
+									.BorderForegroundColor(FLinearColor::White)
+									.BorderBackgroundColor(FLinearColor::White)
+							]
+
+					]
+				];
+		}
+
+private:
+
+	//Get value for text box 0
+	TOptional<float> GetTypeInValue_0() const
+	{
+		return FCString::Atof(*(VisibleText_0.Get()));
+	}
+
+	//Get value for text box 1
+	TOptional<float> GetTypeInValue_1() const
+	{
+		return FCString::Atof(*(VisibleText_1.Get()));
+	}
+
+	//Get value for text box 2
+	TOptional<float> GetTypeInValue_2() const
+	{
+		return FCString::Atof(*(VisibleText_2.Get()));
+	}
+
+	//Get value for text box 3
+	TOptional<float> GetTypeInValue_3() const
+	{
+		return FCString::Atof(*(VisibleText_3.Get()));
+	}
+
+	TAttribute< FString >	VisibleText_0;
+	TAttribute< FString >	VisibleText_1;
+	TAttribute< FString >	VisibleText_2;
+	TAttribute< FString >	VisibleText_3;
+};
+
+
+
+
+/************************************************************************/
+/*                    SGraphPinVector4 class implementation              */
+/************************************************************************/
+
+void SGraphPinVector4::Construct(const FArguments& InArgs, UEdGraphPin* InGraphPinObj)
+{
+	SGraphPin::Construct(SGraphPin::FArguments(), InGraphPinObj);
+}
+
+TSharedRef<SWidget>	SGraphPinVector4::GetDefaultValueWidget()
+{
+	UScriptStruct* RotatorStruct = TBaseStructure<FRotator>::Get();
+
+	//Create widget
+	return	SNew(SVector4TextBox)
+		.VisibleText_0(this, &SGraphPinVector4::GetCurrentValue_0)
+		.VisibleText_1(this, &SGraphPinVector4::GetCurrentValue_1)
+		.VisibleText_2(this, &SGraphPinVector4::GetCurrentValue_2)
+		.VisibleText_3(this, &SGraphPinVector4::GetCurrentValue_3)
+		.Visibility(this, &SGraphPin::GetDefaultValueVisibility)
+		.IsEnabled(this, &SGraphPin::GetDefaultValueIsEditable)
+		.OnFloatCommitted_Box_0(this, &SGraphPinVector4::OnChangedValueTextBox_0)
+		.OnFloatCommitted_Box_1(this, &SGraphPinVector4::OnChangedValueTextBox_1)
+		.OnFloatCommitted_Box_2(this, &SGraphPinVector4::OnChangedValueTextBox_2)
+		.OnFloatCommitted_Box_3(this, &SGraphPinVector4::OnChangedValueTextBox_3);
+}
+
+//Rotator is represented as X->Roll, Y->Pitch, Z->Yaw
+
+FString SGraphPinVector4::GetCurrentValue_0() const
+{
+	//Text box 0: Rotator->Roll, Vector->X
+	return GetValue(TextBox_0);
+}
+
+FString SGraphPinVector4::GetCurrentValue_1() const
+{
+	//Text box 1: Rotator->Pitch, Vector->Y
+	return GetValue(TextBox_1);
+}
+
+FString SGraphPinVector4::GetCurrentValue_2() const
+{
+	//Text box 2: Rotator->Yaw, Vector->Z
+	return GetValue(TextBox_2);
+}
+
+FString SGraphPinVector4::GetCurrentValue_3() const
+{
+	//Text box 3: Vector->W
+	return GetValue(TextBox_3);
+}
+
+FString SGraphPinVector4::GetValue(ETextBoxIndex Index) const
+{
+	FString DefaultString = GraphPinObj->GetDefaultAsString();
+	TArray<FString> ResultString;
+
+	//Parse string to split its contents separated by ','
+	DefaultString.TrimStartAndEndInline();
+	DefaultString.ParseIntoArray(ResultString, TEXT(","), true);
+
+	if (Index < ResultString.Num())
+	{
+		return ResultString[Index];
+	}
+	else
+	{
+		return FString(TEXT("0"));
+	}
+}
+
+
+void SGraphPinVector4::OnChangedValueTextBox_0(float NewValue, ETextCommit::Type CommitInfo)
+{
+	const FString ValueStr = FString::Printf(TEXT("%f"), NewValue);
+
+	FString DefaultValue;
+	//Update X value
+	DefaultValue = ValueStr + FString(TEXT(",")) + GetValue(TextBox_1) + FString(TEXT(",")) + GetValue(TextBox_2) + FString(TEXT(",")) + GetValue(TextBox_3);
+
+	if(GraphPinObj->GetDefaultAsString() != DefaultValue)
+	{
+		const FScopedTransaction Transaction( NSLOCTEXT("GraphEditor", "ChangeVector4PinValue", "Change Vector4 Pin Value" ) );
+		GraphPinObj->Modify();
+
+		//Set new default value
+		GraphPinObj->GetSchema()->TrySetDefaultValue(*GraphPinObj, DefaultValue);
+	}
+}
+
+void SGraphPinVector4::OnChangedValueTextBox_1(float NewValue, ETextCommit::Type CommitInfo)
+{
+	const FString ValueStr = FString::Printf(TEXT("%f"), NewValue);
+
+	FString DefaultValue;
+	//Update Y value
+	DefaultValue = GetValue(TextBox_0) + FString(TEXT(",")) + ValueStr + FString(TEXT(",")) + GetValue(TextBox_2) + FString(TEXT(",")) + GetValue(TextBox_3);
+
+	if(GraphPinObj->GetDefaultAsString() != DefaultValue)
+	{
+		const FScopedTransaction Transaction( NSLOCTEXT("GraphEditor", "ChangeVector4PinValue", "Change Vector4 Pin Value" ) );
+		GraphPinObj->Modify();
+
+		//Set new default value
+		GraphPinObj->GetSchema()->TrySetDefaultValue(*GraphPinObj, DefaultValue);
+	}
+}
+
+void SGraphPinVector4::OnChangedValueTextBox_2(float NewValue, ETextCommit::Type CommitInfo)
+{
+	const FString ValueStr = FString::Printf(TEXT("%f"), NewValue);
+
+	FString DefaultValue;
+	//Update Z value
+	DefaultValue = GetValue(TextBox_0) + FString(TEXT(",")) + GetValue(TextBox_1) + FString(TEXT(",")) + ValueStr + FString(TEXT(",")) + GetValue(TextBox_3);
+
+	if(GraphPinObj->GetDefaultAsString() != DefaultValue)
+	{
+		const FScopedTransaction Transaction( NSLOCTEXT("GraphEditor", "ChangeVector4PinValue", "Change Vector4 Pin Value" ) );
+		GraphPinObj->Modify();
+
+		//Set new default value
+		GraphPinObj->GetSchema()->TrySetDefaultValue(*GraphPinObj, DefaultValue);
+	}
+}
+
+void SGraphPinVector4::OnChangedValueTextBox_3(float NewValue, ETextCommit::Type CommitInfo)
+{
+	const FString ValueStr = FString::Printf(TEXT("%f"), NewValue);
+
+	FString DefaultValue;
+	//Update W value
+	DefaultValue = GetValue(TextBox_0) + FString(TEXT(",")) + GetValue(TextBox_1) + FString(TEXT(",")) + GetValue(TextBox_2) + FString(TEXT(",")) + ValueStr;
+
+	if(GraphPinObj->GetDefaultAsString() != DefaultValue)
+	{
+		const FScopedTransaction Transaction( NSLOCTEXT("GraphEditor", "ChangeVector4PinValue", "Change Vector4 Pin Value" ) );
+		GraphPinObj->Modify();
+
+		//Set new default value
+		GraphPinObj->GetSchema()->TrySetDefaultValue(*GraphPinObj, DefaultValue);
+	}
+}
+
+
+
+#undef LOCTEXT_NAMESPACE
diff --git a/Engine/Source/Editor/GraphEditor/Private/NiagaraPins/SGraphPinVector4.cpp b/Engine/Source/Editor/GraphEditor/Private/NiagaraPins/SGraphPinVector4.cpp
deleted file mode 100644
