[Bug Fix]: Clicking an Add or Removed entry in the blueprint diff tool only moves one of the views.
- The SNodePanel deferrs some of it's panning to future frames which was causing locked graphs in the diff tool to become out of sync. To solve this, I wait until the SNodePanel has finished finding it's target before setting locked graphs on the same trajectory

diff --git a/Engine/Source/Editor/GraphEditor/Private/SGraphEditorImpl.cpp b/Engine/Source/Editor/GraphEditor/Private/SGraphEditorImpl.cpp
--- a/Engine/Source/Editor/GraphEditor/Private/SGraphEditorImpl.cpp
+++ b/Engine/Source/Editor/GraphEditor/Private/SGraphEditorImpl.cpp
@@ -816,6 +816,7 @@ EVisibility SGraphEditorImpl::PIENotification( ) const
 	
 SGraphEditorImpl::~SGraphEditorImpl()
 {
+	UnRegisterActiveTimer(FocusEditorTimer.Pin().ToSharedRef());
 }
 
 void SGraphEditorImpl::Tick( const FGeometry& AllottedGeometry, const double InCurrentTime, const float InDeltaTime )
@@ -1417,7 +1418,7 @@ void SGraphEditorImpl::AddNotification( FNotificationInfo& Info, bool bSuccess )
 	}
 }
 
-void SGraphEditorImpl::FocusLockedEditorHere()
+EActiveTimerReturnType SGraphEditorImpl::HandleFocusEditorDeferred(double InCurrentTime, float InDeltaTime)
 {
 	for( int i = 0; i < LockedGraphs.Num(); ++i )
 	{
@@ -1425,12 +1426,39 @@ void SGraphEditorImpl::FocusLockedEditorHere()
 		if (LockedGraph != TSharedPtr<SGraphEditor>())
 		{
 			LockedGraph->SetViewLocation(GraphPanel->GetViewOffset(), GraphPanel->GetZoomAmount());
+
+			// if GraphPanel is currently panning to a target, pan with it
+			FVector2D TopLeft, BottomRight;
+			if (GraphPanel->GetZoomTargetRect(TopLeft, BottomRight))
+			{
+				// if this graph is zooming/panning to a target location, have the locked graph do the same
+				LockedGraph->GetGraphPanel()->JumpToRect(TopLeft, BottomRight);
+			}
 		}
 		else
 		{
 			LockedGraphs.RemoveAtSwap(i--);
 		}
 	}
+	
+	// If GraphPanel is going to pan to a target but hasn't yet, wait until it does so we don't miss it
+	if (GraphPanel->HasDeferredObjectFocus())
+	{
+		return EActiveTimerReturnType::Continue;
+	}
+	return EActiveTimerReturnType::Stop;
+}
+
+void SGraphEditorImpl::FocusLockedEditorHere()
+{
+	if (!FocusEditorTimer.IsValid())
+	{
+		FocusEditorTimer = RegisterActiveTimer(
+		0.f,
+		FWidgetActiveTimerDelegate::CreateSP(this, &SGraphEditorImpl::HandleFocusEditorDeferred)
+		);
+	}
+	
 }
 
 void SGraphEditorImpl::SetPinVisibility( SGraphEditor::EPinVisibility InVisibility ) 
diff --git a/Engine/Source/Editor/GraphEditor/Private/SGraphEditorImpl.h b/Engine/Source/Editor/GraphEditor/Private/SGraphEditorImpl.h
--- a/Engine/Source/Editor/GraphEditor/Private/SGraphEditorImpl.h
+++ b/Engine/Source/Editor/GraphEditor/Private/SGraphEditorImpl.h
@@ -383,7 +383,9 @@ private:
 	void SelectAllNodesInDirection(const UEdGraphPin* InGraphPin) const;
 	bool IsSelectConnectedNodesFromPinVisible(const FToolMenuContext& InContext, EEdGraphPinDirection DirectionToSelect) const;
 
+	EActiveTimerReturnType HandleFocusEditorDeferred(double InCurrentTime, float InDeltaTime);
 private:
+	TWeakPtr<FActiveTimerHandle> FocusEditorTimer;
 	bool bIsActiveTimerRegistered;
 	uint32 NumNodesAddedSinceLastPointerPosition;
 };
diff --git a/Engine/Source/Editor/GraphEditor/Private/SNodePanel.cpp b/Engine/Source/Editor/GraphEditor/Private/SNodePanel.cpp
--- a/Engine/Source/Editor/GraphEditor/Private/SNodePanel.cpp
+++ b/Engine/Source/Editor/GraphEditor/Private/SNodePanel.cpp
@@ -309,6 +309,15 @@ FVector2D SNodePanel::GetViewOffset() const
 	return ViewOffset;
 }
 
+bool SNodePanel::GetZoomTargetRect(FVector2D& TopLeft, FVector2D& BottomRight) const
+{
+	TopLeft = ZoomTargetTopLeft;
+	BottomRight = ZoomTargetBottomRight;
+	
+	// if the zoom target rect is all zeroed out, then notify caller that there is no target
+	return !(ZoomTargetTopLeft == FVector2D::ZeroVector && ZoomTargetBottomRight == FVector2D::ZeroVector);
+}
+
 void SNodePanel::Construct()
 {
 	if (!ZoomLevels)
diff --git a/Engine/Source/Editor/GraphEditor/Public/SNodePanel.h b/Engine/Source/Editor/GraphEditor/Public/SNodePanel.h
--- a/Engine/Source/Editor/GraphEditor/Public/SNodePanel.h
+++ b/Engine/Source/Editor/GraphEditor/Public/SNodePanel.h
@@ -721,6 +721,14 @@ public:
 	/** @return the view offset in graph space */
 	FVector2D GetViewOffset() const;
 
+	/** 
+	 * when a panel is scrolling/zooming to a target, this can be called to get it's destination
+	 * @param TopLeft top left corner of the destination
+	 * @param BottomRight bottom right corner of the destination
+	 * @return true if there's a scrolling/zooming target and false if there is no destination
+	 */
+	bool GetZoomTargetRect(FVector2D& TopLeft, FVector2D& BottomRight) const;
+
 	/** @return the current view bookmark ID */
 	const FGuid& GetViewBookmarkId() const { return CurrentBookmarkGuid; }
 
