Fixes crash when closing BP graph editor windows that haven't been moved or manipulated in certain ways

#rnx

diff --git a/Engine/Source/Editor/GraphEditor/Private/SGraphEditorImpl.cpp b/Engine/Source/Editor/GraphEditor/Private/SGraphEditorImpl.cpp
--- a/Engine/Source/Editor/GraphEditor/Private/SGraphEditorImpl.cpp
+++ b/Engine/Source/Editor/GraphEditor/Private/SGraphEditorImpl.cpp
@@ -816,7 +816,10 @@ EVisibility SGraphEditorImpl::PIENotification( ) const
 	
 SGraphEditorImpl::~SGraphEditorImpl()
 {
-	UnRegisterActiveTimer(FocusEditorTimer.Pin().ToSharedRef());
+	if (FocusEditorTimer.IsValid())
+	{
+		UnRegisterActiveTimer(FocusEditorTimer.Pin().ToSharedRef());
+	}
 }
 
 void SGraphEditorImpl::Tick( const FGeometry& AllottedGeometry, const double InCurrentTime, const float InDeltaTime )
