GraphEditor: NodeHandle / PinHandle to rely on pin index as well

diff --git a/Engine/Source/Editor/GraphEditor/Public/EdGraphHandleTypes.h b/Engine/Source/Editor/GraphEditor/Public/EdGraphHandleTypes.h
--- a/Engine/Source/Editor/GraphEditor/Public/EdGraphHandleTypes.h
+++ b/Engine/Source/Editor/GraphEditor/Public/EdGraphHandleTypes.h
@@ -65,13 +65,20 @@ public:
 		, PinName(InPin ? InPin->GetFName() : NAME_None)
 		, PinDirection(InPin ? InPin->Direction.GetValue() : EEdGraphPinDirection::EGPD_Input)
 		, PersistentPinGuid(InPin ? InPin->PersistentGuid : FGuid())
-	{}
+		, PinIndex(INDEX_NONE)
+	{
+		if(const UEdGraphNode* Node = InPin->GetOwningNode())
+		{
+			PinIndex = Node->Pins.Find((UEdGraphPin*)InPin);
+		}
+	}
 		
 	FORCEINLINE FEdGraphPinHandle(const FEdGraphPinHandle& InOther)
 		: FEdGraphNodeHandle(InOther)
 		, PinName(InOther.PinName)
 		, PinDirection(InOther.PinDirection)
 		, PersistentPinGuid(InOther.PersistentPinGuid)
+		, PinIndex(InOther.PinIndex)
 	{}
 
 	friend FORCEINLINE uint32 GetTypeHash(const FEdGraphPinHandle& InHandle)
@@ -100,11 +107,23 @@ public:
 	{
 		if(UEdGraphNode* EdNode = GetNode())
 		{
+			// first try using the persistent guid
 			const UEdGraphPin*const* Pin = EdNode->Pins.FindByPredicate([this](UEdGraphPin* Pin) -> bool
 			{
 				return Pin->PersistentGuid.IsValid() && Pin->PersistentGuid == PersistentPinGuid && Pin->Direction == PinDirection;
 			});
 
+			// try to use the pin at the given index if name and direction matches.
+			// this deals with cases where we have multiple pins of the same name.
+			if(Pin == nullptr && EdNode->Pins.IsValidIndex(PinIndex))
+			{
+				if(EdNode->Pins[PinIndex]->GetFName() == PinName && EdNode->Pins[PinIndex]->Direction == PinDirection)
+				{
+					Pin = &EdNode->Pins[PinIndex];
+				}
+			}
+
+			// finally fall back to asking the node for the pin by name
 			if(Pin == nullptr)
 			{
 				Pin = EdNode->Pins.FindByPredicate([this](UEdGraphPin* Pin) -> bool
@@ -125,5 +144,6 @@ private:
 	FName PinName;
 	EEdGraphPinDirection PinDirection;
 	FGuid PersistentPinGuid;
+	int32 PinIndex;
 };
 
