[Feature] blueprint diff node outlines colored by diff state

diff --git a/Engine/Plugins/Experimental/BlueprintSnapNodes/Source/BlueprintSnapNodes/Private/K2Node_SnapContainer.cpp b/Engine/Plugins/Experimental/BlueprintSnapNodes/Source/BlueprintSnapNodes/Private/K2Node_SnapContainer.cpp
--- a/Engine/Plugins/Experimental/BlueprintSnapNodes/Source/BlueprintSnapNodes/Private/K2Node_SnapContainer.cpp
+++ b/Engine/Plugins/Experimental/BlueprintSnapNodes/Source/BlueprintSnapNodes/Private/K2Node_SnapContainer.cpp
@@ -351,7 +351,7 @@ void UK2Node_SnapContainer::FindDiffs(class UEdGraphNode* OtherNode, struct FDif
 		Args.Add(TEXT("Expression2"), Expression2);
 
 		Diff.ToolTip =  FText::Format(LOCTEXT("DIF_MathExpressionToolTip", "Math Expression '{Expression1}' changed to '{Expression2}'"), Args);
-		Diff.DisplayColor = FLinearColor(0.85f,0.71f,0.25f);
+		Diff.Category = EDiffType::MODIFICATION;
 		Diff.DisplayString = FText::Format(LOCTEXT("DIF_MathExpression", "Math Expression '{Expression1}' changed to '{Expression2}'"), Args);
 		Results.Add(Diff);
 	}
diff --git a/Engine/Plugins/Experimental/CommonConversation/Source/CommonConversationEditor/Private/SConversationDiff.cpp b/Engine/Plugins/Experimental/CommonConversation/Source/CommonConversationEditor/Private/SConversationDiff.cpp
--- a/Engine/Plugins/Experimental/CommonConversation/Source/CommonConversationEditor/Private/SConversationDiff.cpp
+++ b/Engine/Plugins/Experimental/CommonConversation/Source/CommonConversationEditor/Private/SConversationDiff.cpp
@@ -46,7 +46,7 @@ struct FTreeDiffResultItem : public TSharedFromThis<FTreeDiffResultItem>
 	TSharedRef<SWidget>	GenerateWidget() const
 	{
 		FText ToolTip = Result.ToolTip;
-		FLinearColor Color = Result.DisplayColor;
+		FLinearColor Color = Result.GetDisplayColor();
 		FText Text = Result.DisplayString;
 		if(Text.IsEmpty())
 		{
diff --git a/Engine/Source/Developer/Merge/Private/SMergeGraphView.cpp b/Engine/Source/Developer/Merge/Private/SMergeGraphView.cpp
--- a/Engine/Source/Developer/Merge/Private/SMergeGraphView.cpp
+++ b/Engine/Source/Developer/Merge/Private/SMergeGraphView.cpp
@@ -221,7 +221,7 @@ static TArray< FMergeGraphEntry > GenerateDiffListItems(const FBlueprintRevPair&
 							, ConflictingDifference ? (*ConflictingDifference)->Pin2 : nullptr /*UEdGraphPin* LocalPin*/
 							, Difference.Pin1 /*UEdGraphPin* BasePin*/
 							, Difference.Pin2 /*UEdGraphPin* RemotePin*/
-							, Difference.DisplayColor
+							, Difference.GetDisplayColor()
 							, ConflictingDifference ? true : false
 						};
 
@@ -246,7 +246,7 @@ static TArray< FMergeGraphEntry > GenerateDiffListItems(const FBlueprintRevPair&
 								, Difference.Pin2 /*UEdGraphPin* LocalPin*/
 								, Difference.Pin1 /*UEdGraphPin* BasePin*/
 								, nullptr
-								, Difference.DisplayColor
+								, Difference.GetDisplayColor()
 								, false
 							};
 
diff --git a/Engine/Source/Editor/AIGraph/Private/AIGraphNode.cpp b/Engine/Source/Editor/AIGraph/Private/AIGraphNode.cpp
--- a/Engine/Source/Editor/AIGraph/Private/AIGraphNode.cpp
+++ b/Engine/Source/Editor/AIGraph/Private/AIGraphNode.cpp
@@ -260,7 +260,7 @@ void UAIGraphNode::FindDiffs(UEdGraphNode* OtherNode, FDiffResults& Results)
 			Diff.Node1 = this;
 			Diff.Node2 = OtherNode;
 			Diff.ToolTip = LOCTEXT("DIF_NodeInstancePropertyToolTip", "A property of the node instance has changed");
-			Diff.DisplayColor = FLinearColor(0.25f, 0.71f, 0.85f);
+			Diff.Category = EDiffType::MODIFICATION;
 
 			DiffProperties(NodeInstance->GetClass(), OtherGraphNode->NodeInstance->GetClass(), NodeInstance, OtherGraphNode->NodeInstance, Results, Diff);
 		}
diff --git a/Engine/Source/Editor/BehaviorTreeEditor/Private/SBehaviorTreeDiff.cpp b/Engine/Source/Editor/BehaviorTreeEditor/Private/SBehaviorTreeDiff.cpp
--- a/Engine/Source/Editor/BehaviorTreeEditor/Private/SBehaviorTreeDiff.cpp
+++ b/Engine/Source/Editor/BehaviorTreeEditor/Private/SBehaviorTreeDiff.cpp
@@ -43,7 +43,7 @@ struct FTreeDiffResultItem : public TSharedFromThis<FTreeDiffResultItem>
 	TSharedRef<SWidget>	GenerateWidget() const
 	{
 		FText ToolTip = Result.ToolTip;
-		FLinearColor Color = Result.DisplayColor;
+		FSlateColor Color = Result.GetDisplayColor();
 		FText Text = Result.DisplayString;
 		if(Text.IsEmpty())
 		{
diff --git a/Engine/Source/Editor/BlueprintGraph/Private/K2Node_AddComponent.cpp b/Engine/Source/Editor/BlueprintGraph/Private/K2Node_AddComponent.cpp
--- a/Engine/Source/Editor/BlueprintGraph/Private/K2Node_AddComponent.cpp
+++ b/Engine/Source/Editor/BlueprintGraph/Private/K2Node_AddComponent.cpp
@@ -427,7 +427,7 @@ void UK2Node_AddComponent::FindDiffs(class UEdGraphNode* OtherNode, struct FDiff
 			Diff.Node1 = this;
 			Diff.Node2 = OtherNode;
 			Diff.ToolTip = LOCTEXT("DIF_ComponentTemplatePropertyToolTip", "A property of the component template has changed");
-			Diff.DisplayColor = FLinearColor(0.25f, 0.71f, 0.85f);
+			Diff.Category = EDiffType::MODIFICATION;
 
 			DiffProperties(MyComponent->GetClass(), OtherComponent->GetClass(), MyComponent, OtherComponent, Results, Diff);
 		}
diff --git a/Engine/Source/Editor/BlueprintGraph/Private/K2Node_Event.cpp b/Engine/Source/Editor/BlueprintGraph/Private/K2Node_Event.cpp
--- a/Engine/Source/Editor/BlueprintGraph/Private/K2Node_Event.cpp
+++ b/Engine/Source/Editor/BlueprintGraph/Private/K2Node_Event.cpp
@@ -883,7 +883,7 @@ void UK2Node_Event::FindDiffs(UEdGraphNode* OtherNode, struct FDiffResults& Resu
 			Diff.Node1 = this;
 			Diff.Node2 = OtherNode;
 			Diff.DisplayString = LOCTEXT("DIF_EventFlags", "Event flags have changed");
-			Diff.DisplayColor = FLinearColor(0.25f, 0.71f, 0.85f);
+			Diff.Category = EDiffType::MODIFICATION;
 
 			Results.Add(Diff);
 		}
diff --git a/Engine/Source/Editor/BlueprintGraph/Private/K2Node_FunctionEntry.cpp b/Engine/Source/Editor/BlueprintGraph/Private/K2Node_FunctionEntry.cpp
--- a/Engine/Source/Editor/BlueprintGraph/Private/K2Node_FunctionEntry.cpp
+++ b/Engine/Source/Editor/BlueprintGraph/Private/K2Node_FunctionEntry.cpp
@@ -673,7 +673,7 @@ void UK2Node_FunctionEntry::FindDiffs(UEdGraphNode* OtherNode, struct FDiffResul
 			Diff.Node1 = this;
 			Diff.Node2 = OtherNode;
 			Diff.DisplayString = LOCTEXT("DIF_FunctionFlags", "Function flags have changed");
-			Diff.DisplayColor = FLinearColor(0.25f, 0.71f, 0.85f);
+			Diff.Category = EDiffType::MODIFICATION;
 
 			Results.Add(Diff);
 		}
@@ -685,7 +685,7 @@ void UK2Node_FunctionEntry::FindDiffs(UEdGraphNode* OtherNode, struct FDiffResul
 			Diff.Node1 = this;
 			Diff.Node2 = OtherNode;
 			Diff.DisplayString = LOCTEXT("DIF_FunctionMetadata", "Function metadata has changed");
-			Diff.DisplayColor = FLinearColor(0.25f, 0.71f, 0.85f);
+			Diff.Category = EDiffType::MODIFICATION;
 
 			Results.Add(Diff);
 		}
@@ -717,7 +717,7 @@ void UK2Node_FunctionEntry::FindDiffs(UEdGraphNode* OtherNode, struct FDiffResul
 			Diff.Node1 = this;
 			Diff.Node2 = OtherNode;
 			Diff.DisplayString = LOCTEXT("DIF_FunctionLocalVariables", "Function local variables have changed in structure");
-			Diff.DisplayColor = FLinearColor(0.25f, 0.71f, 0.85f);
+			Diff.Category = EDiffType::MODIFICATION;
 
 			Results.Add(Diff);
 		}
@@ -734,7 +734,7 @@ void UK2Node_FunctionEntry::FindDiffs(UEdGraphNode* OtherNode, struct FDiffResul
 				Diff.Node1 = this;
 				Diff.Node2 = OtherNode;
 				Diff.ToolTip = LOCTEXT("DIF_FunctionLocalVariableDefaults", "Function local variable default values have changed");
-				Diff.DisplayColor = FLinearColor(0.25f, 0.71f, 0.85f);
+				Diff.Category = EDiffType::MODIFICATION;
 
 				DiffProperties(const_cast<UStruct*>(MyLocals->GetStruct()), const_cast<UStruct*>(OtherLocals->GetStruct()), MyLocals->GetStructMemory(), OtherLocals->GetStructMemory(), Results, Diff);
 			}
diff --git a/Engine/Source/Editor/BlueprintGraph/Private/K2Node_MathExpression.cpp b/Engine/Source/Editor/BlueprintGraph/Private/K2Node_MathExpression.cpp
--- a/Engine/Source/Editor/BlueprintGraph/Private/K2Node_MathExpression.cpp
+++ b/Engine/Source/Editor/BlueprintGraph/Private/K2Node_MathExpression.cpp
@@ -2877,7 +2877,7 @@ void UK2Node_MathExpression::FindDiffs(class UEdGraphNode* OtherNode, struct FDi
 		Args.Add(TEXT("Expression2"), Expression2);
 
 		Diff.ToolTip =  FText::Format(LOCTEXT("DIF_MathExpressionToolTip", "Math Expression '{Expression1}' changed to '{Expression2}'"), Args);
-		Diff.DisplayColor = FLinearColor(0.85f,0.71f,0.25f);
+		Diff.Category = EDiffType::MODIFICATION;
 		Diff.DisplayString = FText::Format(LOCTEXT("DIF_MathExpression", "Math Expression '{Expression1}' changed to '{Expression2}'"), Args);
 		Results.Add(Diff);
 	}
diff --git a/Engine/Source/Editor/BlueprintGraph/Private/K2Node_Timeline.cpp b/Engine/Source/Editor/BlueprintGraph/Private/K2Node_Timeline.cpp
--- a/Engine/Source/Editor/BlueprintGraph/Private/K2Node_Timeline.cpp
+++ b/Engine/Source/Editor/BlueprintGraph/Private/K2Node_Timeline.cpp
@@ -461,7 +461,7 @@ void FindExactTimelineDifference(struct FDiffResults& Results, FDiffSingleResult
 
 		Result.Diff = EDiffType::TIMELINE_NUM_TRACKS;
 		Result.ToolTip =  FText::Format(LOCTEXT("DIF_TimelineNumTracksToolTip", "The number of {TrackType} tracks in Timeline '{NodeName}' has changed"), Args);
-		Result.DisplayColor = FLinearColor(0.05f,0.261f,0.775f);
+		Result.Category = EDiffType::MODIFICATION;
 		Result.DisplayString = FText::Format(LOCTEXT("DIF_TimelineNumTracks", "{TrackType} Track Count '{NodeName}'"), Args);
 		Results.Add(Result);
 		return;
@@ -480,7 +480,7 @@ void FindExactTimelineDifference(struct FDiffResults& Results, FDiffSingleResult
 
 			Result.Diff = EDiffType::TIMELINE_TRACK_MODIFIED;
 			Result.ToolTip =  FText::Format(LOCTEXT("DIF_TimelineTrackModifiedToolTip", "Track '{TrackName}' of Timeline '{NodeName}' was Modified"), Args);
-			Result.DisplayColor = FLinearColor(0.75f,0.1f,0.15f);
+			Result.Category = EDiffType::MODIFICATION;
 			Result.DisplayString = FText::Format(LOCTEXT("DIF_TimelineTrackModified", "Track Modified '{TrackName}'"), Args);
 			Results.Add(Result);
 			break;
@@ -516,7 +516,7 @@ void UK2Node_Timeline::FindDiffs( class UEdGraphNode* OtherNode, struct FDiffRes
 			Args.Add(TEXT("NodeName"), NodeName);
 
 			Diff.ToolTip =  FText::Format(LOCTEXT("DIF_TimelineAutoPlayToolTip", "Timeline '{NodeName}' had its AutoPlay state changed"), Args);
-			Diff.DisplayColor = FLinearColor(0.15f,0.61f,0.15f);
+			Diff.Category = EDiffType::MODIFICATION;
 			Diff.DisplayString = FText::Format(LOCTEXT("DIF_TimelineAutoPlay", "Timeline AutoPlay Changed '{NodeName}'"), Args);
 			Results.Add(Diff);
 		}
@@ -529,7 +529,7 @@ void UK2Node_Timeline::FindDiffs( class UEdGraphNode* OtherNode, struct FDiffRes
 			Args.Add(TEXT("NodeName"), NodeName);
 
 			Diff.ToolTip =  FText::Format(LOCTEXT("DIF_TimelineLoopingToolTip", "Timeline '{NodeName}' had its looping state changed"), Args);
-			Diff.DisplayColor = FLinearColor(0.75f,0.1f,0.75f);
+			Diff.Category = EDiffType::MODIFICATION;
 			Diff.DisplayString =  FText::Format(LOCTEXT("DIF_TimelineLooping", "Timeline Loop Changed '{NodeName}'"), Args);
 			Results.Add(Diff);
 		}
@@ -544,7 +544,7 @@ void UK2Node_Timeline::FindDiffs( class UEdGraphNode* OtherNode, struct FDiffRes
 
 			Diff.Diff = EDiffType::TIMELINE_LENGTH;
 			Diff.ToolTip = FText::Format(LOCTEXT("DIF_TimelineLengthToolTip", "Length of Timeline '{NodeName}' has changed. Was {TimelineLength1}, but is now {TimelineLength2}"), Args);
-			Diff.DisplayColor = FLinearColor(0.25f,0.1f,0.15f);
+			Diff.Category = EDiffType::MODIFICATION;
 			Diff.DisplayString =  FText::Format(LOCTEXT("DIF_TimelineLength", "Timeline Length '{NodeName}' [{TimelineLength1} -> {TimelineLength2}]"), Args);
 			Results.Add(Diff);
 		}
@@ -557,7 +557,7 @@ void UK2Node_Timeline::FindDiffs( class UEdGraphNode* OtherNode, struct FDiffRes
 			Args.Add(TEXT("NodeName"), NodeName);
 
 			Diff.ToolTip = FText::Format(LOCTEXT("DIF_TimelineIgnoreDilationToolTip", "Timeline '{NodeName}' had its ignore time dilation state changed"), Args);
-			Diff.DisplayColor = FLinearColor(0.75f, 0.1f, 0.75f);
+			Diff.Category = EDiffType::MODIFICATION;
 			Diff.DisplayString = FText::Format(LOCTEXT("DIF_TimelineIgnoreDilation", "Timeline IgnoreTimeDilation Changed '{NodeName}'"), Args);
 			Results.Add(Diff);
 		}
diff --git a/Engine/Source/Editor/EditorStyle/Private/StarshipStyle.cpp b/Engine/Source/Editor/EditorStyle/Private/StarshipStyle.cpp
--- a/Engine/Source/Editor/EditorStyle/Private/StarshipStyle.cpp
+++ b/Engine/Source/Editor/EditorStyle/Private/StarshipStyle.cpp
@@ -3146,11 +3146,20 @@ void FStarshipEditorStyle::FStyle::SetupGraphEditorStyles()
 		Set( "Graph.Node.ShadowSize", FVector2D(12,12) );
 		Set( "Graph.Node.ShadowSelected", new BOX_BRUSH( "/Graph/RegularNode_shadow_selected", FMargin(18.0f/64.0f) ) );
 		Set( "Graph.Node.Shadow", new BOX_BRUSH( "/Graph/RegularNode_shadow", FMargin(18.0f/64.0f) ) );
+		
+		Set( "Graph.Node.DiffHighlight", new BOX_BRUSH( "/Graph/RegularNode_DiffHighlight", FMargin(18.0f/64.0f) ) );
+		Set( "Graph.Node.DiffHighlightShading", new BOX_BRUSH( "/Graph/RegularNode_DiffHighlightShading", FMargin(18.0f/64.0f) ) );
 
 		Set( "Graph.Node.RerouteShadow", new IMAGE_BRUSH( "/Graph/RerouteNode_shadow", FVector2D(64.0f, 64.0f) ) );
 		Set( "Graph.Node.RerouteShadowSelected", new IMAGE_BRUSH( "/Graph/RerouteNode_shadow_selected", FVector2D(64.0f, 64.0f) ) );
+		
+		Set( "Graph.Node.RerouteDiffHighlight", new BOX_BRUSH( "/Graph/RerouteNode_DiffHighlight", FMargin(18.0f/64.0f) ) );
+		Set( "Graph.Node.RerouteDiffHighlightShading", new BOX_BRUSH( "/Graph/RerouteNode_DiffHighlightShading", FMargin(18.0f/64.0f) ) );
 
 		Set( "Graph.CompactNode.ShadowSelected", new BOX_BRUSH( "/Graph/MathNode_shadow_selected", FMargin(18.0f/64.0f) ) );
+		
+		Set( "Graph.CompactNode.DiffHighlight", new BOX_BRUSH( "/Graph/MathNode_DiffHighlight", FMargin(18.0f/64.0f) ) );
+		Set( "Graph.CompactNode.DiffHighlightShading", new BOX_BRUSH( "/Graph/MathNode_DiffHighlightShading", FMargin(18.0f/64.0f) ) );
 
 		Set( "Graph.Node.CommentBubble", new BOX_BRUSH( "Old/Graph/CommentBubble", FMargin(8/32.0f) ) );
 		Set( "Graph.Node.CommentArrow", new IMAGE_BRUSH( "Old/Graph/CommentBubbleArrow", FVector2D(8,8) ) );
@@ -3199,6 +3208,9 @@ void FStarshipEditorStyle::FStyle::SetupGraphEditorStyles()
 		
 		Set( "Graph.VarNode.ShadowSelected", new BOX_BRUSH( "/Graph/VarNode_shadow_selected", FMargin(26.0f/64.0f) ) );
 		Set( "Graph.VarNode.Shadow", new BOX_BRUSH( "/Graph/VarNode_shadow", FMargin(26.0f/64.0f) ) );
+		
+		Set( "Graph.VarNode.DiffHighlight", new BOX_BRUSH( "/Graph/VarNode_DiffHighlight", FMargin(18.0f/64.0f) ) );
+		Set( "Graph.VarNode.DiffHighlightShading", new BOX_BRUSH( "/Graph/VarNode_DiffHighlightShading", FMargin(18.0f/64.0f) ) );
 
 		Set( "Graph.CollapsedNode.Body", new BOX_BRUSH( "/Graph/RegularNode_body", FMargin(16.f/64.f, 25.f/64.f, 16.f/64.f, 16.f/64.f) ) );
 		Set( "Graph.CollapsedNode.BodyColorSpill", new BOX_BRUSH( "/Graph/CollapsedNode_Body_ColorSpill", FMargin(16.f/64.f, 25.f/64.f, 16.f/64.f, 16.f/64.f) ) );
@@ -3612,6 +3624,7 @@ void FStarshipEditorStyle::FStyle::SetupGraphEditorStyles()
 
 		Set( "Graph.Pin.BackgroundHovered", new IMAGE_BRUSH( "/Graph/Pin_hover_cue", FVector2D(32,8)));
 		Set( "Graph.Pin.Background", new FSlateNoResource() );
+		Set( "Graph.Pin.DiffHighlight", new IMAGE_BRUSH( "/Graph/Pin_DiffHighlight", FVector2D(32,8)));
 
 		Set( "Graph.Pin.ObjectSet", new IMAGE_BRUSH( "Old/Graph/Pin_ObjectSet", Icon12x12 ) );
 		Set( "Graph.Pin.ObjectEmpty", new IMAGE_BRUSH( "Old/Graph/Pin_ObjectEmpty", Icon12x12 ) );
diff --git a/Engine/Source/Editor/GraphEditor/Private/ConnectionDrawingPolicy.cpp b/Engine/Source/Editor/GraphEditor/Private/ConnectionDrawingPolicy.cpp
--- a/Engine/Source/Editor/GraphEditor/Private/ConnectionDrawingPolicy.cpp
+++ b/Engine/Source/Editor/GraphEditor/Private/ConnectionDrawingPolicy.cpp
@@ -489,6 +489,10 @@ void FConnectionDrawingPolicy::DrawPinGeometries(TMap<TSharedRef<SWidget>, FArra
 				{
 					FConnectionParams Params;
 					DetermineWiringStyle(ThePin, TargetPin, /*inout*/ Params);
+					if ( PinWidget.AreConnectionsFaded() && PinToPinWidgetMap[TargetPin]->AreConnectionsFaded() )
+					{
+						Params.WireColor.A = 0.2f;
+					}
 					DrawSplineWithArrow(LinkStartWidgetGeometry->Geometry, LinkEndWidgetGeometry->Geometry, Params);
 				}
 			}
diff --git a/Engine/Source/Editor/GraphEditor/Private/GraphDiffControl.cpp b/Engine/Source/Editor/GraphEditor/Private/GraphDiffControl.cpp
--- a/Engine/Source/Editor/GraphEditor/Private/GraphDiffControl.cpp
+++ b/Engine/Source/Editor/GraphEditor/Private/GraphDiffControl.cpp
@@ -32,7 +32,7 @@ static void DiffR_NodeAdded( const FGraphDiffControl::FNodeDiffContext& DiffCont
 		Args.Add(TEXT("NodeTitle"), Node->GetNodeTitle(ENodeTitleType::ListView));
 		Diff.ToolTip =  FText::Format(LOCTEXT("DIF_AddNode", "Added {NodeType} '{NodeTitle}'"), Args);
 		Diff.DisplayString = Diff.ToolTip;
-		Diff.DisplayColor = FLinearColor(0.3f,1.0f,0.4f);
+		Diff.Category = EDiffType::ADDITION;
 	}
 	
 	Results.Add(Diff);
@@ -53,7 +53,7 @@ static void DiffR_NodeRemoved( const FGraphDiffControl::FNodeDiffContext& DiffCo
 		Args.Add(TEXT("NodeTitle"), Node->GetNodeTitle(ENodeTitleType::ListView));
 		Diff.ToolTip =  FText::Format(LOCTEXT("DIF_RemoveNode", "Removed {NodeType} '{NodeTitle}'"), Args);
 		Diff.DisplayString = Diff.ToolTip;
-		Diff.DisplayColor = FLinearColor(1.f,0.4f,0.4f);
+		Diff.Category = EDiffType::SUBTRACTION;
 	}
 
 	Results.Add(Diff);
@@ -75,7 +75,7 @@ static void DiffR_NodeCommentChanged(const FGraphDiffControl::FNodeDiffContext&
 		Args.Add(TEXT("NodeTitle"), NewNode->GetNodeTitle(ENodeTitleType::ListView));
 		Diff.ToolTip =  FText::Format(LOCTEXT("DIF_CommentModified", "Comment Modified {NodeType} '{NodeTitle}'"), Args);
 		Diff.DisplayString = Diff.ToolTip;
-		Diff.DisplayColor = FLinearColor(0.25f,0.4f,0.5f);
+		Diff.Category = EDiffType::MINOR;
 	}
 
 	Results.Add(Diff);
@@ -86,6 +86,7 @@ static void DiffR_NodeMoved(const FGraphDiffControl::FNodeDiffContext& DiffConte
 {
 	FDiffSingleResult Diff;
 	Diff.Diff = EDiffType::NODE_MOVED;
+	Diff.Category = EDiffType::MINOR;
 	Diff.Node1 = OldNode;
 	Diff.Node2 = NewNode;
 
@@ -97,7 +98,6 @@ static void DiffR_NodeMoved(const FGraphDiffControl::FNodeDiffContext& DiffConte
 		Args.Add(TEXT("NodeTitle"), NewNode->GetNodeTitle(ENodeTitleType::ListView));
 		Diff.ToolTip = FText::Format(LOCTEXT("DIF_MoveNode", "Moved {NodeType} '{NodeTitle}'"), Args);
 		Diff.DisplayString = Diff.ToolTip;
-		Diff.DisplayColor = FLinearColor(0.9f, 0.84f, 0.43f);
 	}
 
 	Results.Add(Diff);
@@ -122,7 +122,7 @@ static void DiffR_PinTypeChanged(FDiffResults& Results, UEdGraphPin* NewPin, UEd
 		if(Results.CanStoreResults())
 		{
 			Diff.ToolTip = FText::Format(LOCTEXT("DIF_PinCategoryToolTipFmt", "Pin '{0}' Category was '{1}', but is now '{2}"), FText::FromName(NewPin->PinName), FText::FromName(OldPin->PinType.PinCategory), FText::FromName(NewPin->PinType.PinCategory));
-			Diff.DisplayColor = FLinearColor(0.15f,0.53f,0.15f);
+			Diff.Category = EDiffType::MODIFICATION;
 			Diff.DisplayString = FText::Format(LOCTEXT("DIF_PinCategoryFmt", "Pin Category '{0}' ['{1}' -> '{2}']"), FText::FromName(NewPin->PinName), FText::FromName(OldPin->PinType.PinCategory), FText::FromName(NewPin->PinType.PinCategory));
 		}
 	}
@@ -134,7 +134,7 @@ static void DiffR_PinTypeChanged(FDiffResults& Results, UEdGraphPin* NewPin, UEd
 		if(Results.CanStoreResults())
 		{
 			Diff.ToolTip = FText::Format(LOCTEXT("DIF_PinSubCategoryToolTipFmt", "Pin '{0}' SubCategory was '{1}', but is now '{2}"), FText::FromName(NewPin->PinName), FText::FromName(OldPin->PinType.PinSubCategory), FText::FromName(NewPin->PinType.PinSubCategory));
-			Diff.DisplayColor = FLinearColor(0.45f,0.53f,0.65f);
+			Diff.Category = EDiffType::MODIFICATION;
 			Diff.DisplayString = FText::Format(LOCTEXT("DIF_PinSubCategoryFmt", "Pin SubCategory '{0}'  ['{1}' -> '{2}']"), FText::FromName(NewPin->PinName), FText::FromName(OldPin->PinType.PinSubCategory), FText::FromName(NewPin->PinType.PinSubCategory));
 		}
 	}
@@ -150,7 +150,7 @@ static void DiffR_PinTypeChanged(FDiffResults& Results, UEdGraphPin* NewPin, UEd
 			const FName Obj2 = T2Obj->GetFName();
 
 			Diff.ToolTip = FText::Format(LOCTEXT("DIF_PinSubCategorObjToolTipFmt", "Pin '{0}' was SubCategoryObject '{1}', but is now '{2}"), FText::FromName(NewPin->PinName), FText::FromName(Obj1), FText::FromName(Obj2));
-			Diff.DisplayColor = FLinearColor(0.45f,0.13f,0.25f);
+			Diff.Category = EDiffType::MODIFICATION;
 			Diff.DisplayString = FText::Format(LOCTEXT("DIF_PinSubCategoryObjFmt", "Pin SubCategoryObject '{0}' ['{1}' -> '{2}']"), FText::FromName(NewPin->PinName), FText::FromName(Obj1), FText::FromName(Obj2));
 		}
 	}
@@ -167,7 +167,7 @@ static void DiffR_PinTypeChanged(FDiffResults& Results, UEdGraphPin* NewPin, UEd
 
 			Diff.ToolTip = FText::Format(LOCTEXT("DIF_PinIsArrayToolTipFmt", "PinType IsArray for '{0}' modified. Was '{1}', but is now '{2}"), FText::FromName(NewPin->PinName), IsArray1, IsArray2);
 			Diff.DisplayString = FText::Format(LOCTEXT("DIF_PinIsArrayFmt", "Pin IsArray '{0}' ['{1}' -> '{2}']"), FText::FromName(NewPin->PinName), IsArray1, IsArray2);
-			Diff.DisplayColor = FLinearColor(0.45f,0.33f,0.35f);
+			Diff.Category = EDiffType::MODIFICATION;
 		}
 	}
 	else if(Type1.bIsReference != Type2.bIsReference)
@@ -181,7 +181,7 @@ static void DiffR_PinTypeChanged(FDiffResults& Results, UEdGraphPin* NewPin, UEd
 			FText IsRef2 = NewPin->PinType.bIsReference ? LOCTEXT("true", "true") : LOCTEXT("false", "false");
 
 			Diff.ToolTip = FText::Format(LOCTEXT("DIF_PinIsRefToolTipFmt", "PinType IsReference for '{0}' modified. Was '{1}', but is now '{2}"), FText::FromName(NewPin->PinName), IsRef1, IsRef2);
-			Diff.DisplayColor = FLinearColor(0.25f,0.43f,0.35f);
+			Diff.Category = EDiffType::MODIFICATION;
 			Diff.DisplayString = FText::Format(LOCTEXT("DIF_PinIsRefFmt", "Pin IsReference '{0}' ['{1}' -> '{2}']"), FText::FromName(NewPin->PinName), IsRef1, IsRef2);
 		}
 	}
@@ -205,13 +205,13 @@ static void DiffR_PinLinkCountChanged(FDiffResults& Results, UEdGraphPin* NewPin
 		if(Diff.Diff == EDiffType::PIN_LINKEDTO_NUM_INC)
 		{
 			Diff.ToolTip = FText::Format(LOCTEXT("DIF_PinLinkCountIncToolTipFmt", "Pin '{0}' has more links (was {1} now {2})"), FText::FromName(OldPin->PinName), FText::AsNumber(OldPin->LinkedTo.Num()), FText::AsNumber(NewPin->LinkedTo.Num()));
-			Diff.DisplayColor = FLinearColor(0.5f,0.3f,0.85f);
+			Diff.Category = EDiffType::ADDITION;
 			Diff.DisplayString = FText::Format(LOCTEXT("DIF_PinLinkCountIncFmt", "Added Link to '{0}'"), FText::FromName(OldPin->PinName));
 		}
 		else
 		{
 			Diff.ToolTip = FText::Format(LOCTEXT("DIF_PinLinkCountDecToolTipFmt", "Pin '{0}' has fewer links (was {1} now {2})"), FText::FromName(OldPin->PinName), FText::AsNumber(OldPin->LinkedTo.Num()), FText::AsNumber(NewPin->LinkedTo.Num()));
-			Diff.DisplayColor = FLinearColor(0.5f,0.3f,0.85f);
+			Diff.Category = EDiffType::SUBTRACTION;
 			Diff.DisplayString = FText::Format(LOCTEXT("DIF_PinLinkCountDecFmt", "Removed Link to '{0}'"), FText::FromName(OldPin->PinName));
 		}
 	}
@@ -241,7 +241,7 @@ static void DiffR_LinkedToNode(FDiffResults& Results, UEdGraphPin* OldPin, UEdGr
 		Args.Add(TEXT("NodeName2"), Node2Name);
 		Diff.ToolTip = FText::Format(LOCTEXT("DIF_PinLinkMovedToolTip", "Pin '{PinNameForNode1}' was linked to Node '{NodeName1}', but is now linked to Node '{NodeName2}'"), Args);
 
-		Diff.DisplayColor = FLinearColor(0.85f,0.71f,0.25f);
+		Diff.Category = EDiffType::MODIFICATION;
 		Diff.DisplayString = FText::Format(LOCTEXT("DIF_PinLinkMoved", "Link Moved  '{PinNameForNode1}' ['{NodeName1}' -> '{NodeName2}']"), Args);
 	}
 
@@ -264,7 +264,7 @@ static void DiffR_PinDefaultValueChanged(FDiffResults& Results, UEdGraphPin* New
 		Args.Add(TEXT("PinValue1"), FText::FromString(OldPin->GetDefaultAsString()));
 		Args.Add(TEXT("PinValue2"), FText::FromString(NewPin->GetDefaultAsString()));
 		Diff.ToolTip = FText::Format(LOCTEXT("DIF_PinDefaultValueToolTip", "Pin '{PinNameForValue1}' Default Value was '{PinValue1}', but is now '{PinValue2}"), Args);
-		Diff.DisplayColor = FLinearColor(0.665f,0.13f,0.455f);
+		Diff.Category = EDiffType::MODIFICATION;
 		Diff.DisplayString = FText::Format(LOCTEXT("DIF_PinDefaultValue", "Pin Default '{PinNameForValue1}' '{PinValue1}' -> '{PinValue2}']"), Args);
 	}
 
@@ -289,7 +289,6 @@ static void DiffR_NodePinCount(FDiffResults& Results, UEdGraphNode* NewNode, UEd
 		Args.Add(TEXT("NodeName"), NodeName);
 		Args.Add(TEXT("OriginalCount"), OriginalCount);
 		Args.Add(TEXT("NewCount"), NewCount);
-		Diff.DisplayColor = FLinearColor(0.45f,0.4f,0.4f);
 
 		struct FMatchName
 		{
@@ -331,16 +330,19 @@ static void DiffR_NodePinCount(FDiffResults& Results, UEdGraphNode* NewNode, UEd
 		if (RemovedPins.Num() > 0 && AddedPins.Num() > 0)
 		{
 			Diff.DisplayString = FText::Format(LOCTEXT("DIF_NodePinsAddedAndRemoved", "Added and removed Pins from '{NodeName}'"), Args);
+			Diff.Category = EDiffType::ADDITION;
 		}
 		else if (AddedPins.Num() > 0)
 		{
 			if (AddedPins.Num() == 1)
 			{
 				Diff.DisplayString = FText::Format(LOCTEXT("DIF_NodePinCountIncreased", "Added Pin to '{NodeName}'"), Args);
+				Diff.Category = EDiffType::ADDITION;
 			}
 			else
 			{
 				Diff.DisplayString = FText::Format(LOCTEXT("DIF_NodePinCountIncreasedSeveral", "Added Pins to '{NodeName}'"), Args);
+				Diff.Category = EDiffType::ADDITION;
 			}
 		}
 		else if (RemovedPins.Num() > 0)
@@ -348,10 +350,12 @@ static void DiffR_NodePinCount(FDiffResults& Results, UEdGraphNode* NewNode, UEd
 			if (RemovedPins.Num() == 1)
 			{
 				Diff.DisplayString = FText::Format(LOCTEXT("DIF_NodePinCountDecreased", "Removed Pin from '{NodeName}'"), Args);
+				Diff.Category = EDiffType::SUBTRACTION;
 			}
 			else
 			{
 				Diff.DisplayString = FText::Format(LOCTEXT("DIF_NodePinCountDecreasedSeveral", "Removed Pins from '{NodeName}'"), Args);
+				Diff.Category = EDiffType::SUBTRACTION;
 			}
 		}
 
diff --git a/Engine/Source/Editor/GraphEditor/Private/KismetNodes/SGraphNodeK2Base.cpp b/Engine/Source/Editor/GraphEditor/Private/KismetNodes/SGraphNodeK2Base.cpp
--- a/Engine/Source/Editor/GraphEditor/Private/KismetNodes/SGraphNodeK2Base.cpp
+++ b/Engine/Source/Editor/GraphEditor/Private/KismetNodes/SGraphNodeK2Base.cpp
@@ -669,6 +669,22 @@ const FSlateBrush* SGraphNodeK2Base::GetShadowBrush(bool bSelected) const
 	}
 }
 
+void SGraphNodeK2Base::GetDiffHighlightBrushes(const FSlateBrush*& BackgroundOut, const FSlateBrush*& ForegroundOut) const
+{
+	const UK2Node* K2Node = CastChecked<UK2Node>(GraphNode);
+	
+	if (K2Node->ShouldDrawCompact())
+	{
+		BackgroundOut = FAppStyle::GetBrush(TEXT("Graph.VarNode.DiffHighlight"));
+		ForegroundOut = FAppStyle::GetBrush(TEXT("Graph.VarNode.DiffHighlightShading"));
+	}
+	else
+	{
+		BackgroundOut = FAppStyle::GetBrush(TEXT("Graph.Node.DiffHighlight"));
+		ForegroundOut = FAppStyle::GetBrush(TEXT("Graph.Node.DiffHighlightShading"));
+	}
+}
+
 void SGraphNodeK2Base::PerformSecondPassLayout(const TMap< UObject*, TSharedRef<SNode> >& NodeToWidgetLookup) const
 {
 	TSet<UEdGraphNode*> PrevNodes;
diff --git a/Engine/Source/Editor/GraphEditor/Private/KismetNodes/SGraphNodeK2Var.cpp b/Engine/Source/Editor/GraphEditor/Private/KismetNodes/SGraphNodeK2Var.cpp
--- a/Engine/Source/Editor/GraphEditor/Private/KismetNodes/SGraphNodeK2Var.cpp
+++ b/Engine/Source/Editor/GraphEditor/Private/KismetNodes/SGraphNodeK2Var.cpp
@@ -331,3 +331,10 @@ const FSlateBrush* SGraphNodeK2Var::GetShadowBrush(bool bSelected) const
 {
 	return bSelected ? FAppStyle::GetBrush(TEXT("Graph.VarNode.ShadowSelected")) : FAppStyle::GetBrush(TEXT("Graph.VarNode.Shadow"));
 }
+
+void SGraphNodeK2Var::GetDiffHighlightBrushes(const FSlateBrush*& BackgroundOut, const FSlateBrush*& ForegroundOut) const
+{
+	BackgroundOut = FAppStyle::GetBrush(TEXT("Graph.VarNode.DiffHighlight"));
+	ForegroundOut = FAppStyle::GetBrush(TEXT("Graph.VarNode.DiffHighlightShading"));
+}
+
diff --git a/Engine/Source/Editor/GraphEditor/Private/SGraphPanel.cpp b/Engine/Source/Editor/GraphEditor/Private/SGraphPanel.cpp
--- a/Engine/Source/Editor/GraphEditor/Private/SGraphPanel.cpp
+++ b/Engine/Source/Editor/GraphEditor/Private/SGraphPanel.cpp
@@ -104,7 +104,7 @@ int32 SGraphPanel::OnPaint( const FPaintArgs& Args, const FGeometry& AllottedGeo
 
 	//Style used for objects that are the same between revisions
 	FWidgetStyle FadedStyle = InWidgetStyle;
-	FadedStyle.BlendColorAndOpacityTint(FLinearColor(0.45f,0.45f,0.45f,0.45f));
+	FadedStyle.BlendColorAndOpacityTint(FLinearColor(0.45f,0.45f,0.45f,0.30f));
 
 	// First paint the background
 	const UEditorExperimentalSettings& Options = *GetDefault<UEditorExperimentalSettings>();
@@ -129,6 +129,8 @@ int32 SGraphPanel::OnPaint( const FPaintArgs& Args, const FGeometry& AllottedGeo
 	const int32 CommentNodeShadowLayerId = LayerId++;
 	const int32 CommentNodeLayerId = LayerId++;
 
+	const int32 NodeDiffHighlightLayerID = LayerId++;
+
 	// Save a LayerId for wires, which appear below nodes but above comments
 	// We will draw them later, along with the arrows which appear above nodes.
 	const int32 WireLayerId = LayerId++;
@@ -154,7 +156,12 @@ int32 SGraphPanel::OnPaint( const FPaintArgs& Args, const FGeometry& AllottedGeo
 			if (Result.Pin1)
 			{
 				PinDiffResults.FindOrAdd(Result.Pin1, Result);
-				NodeDiffResults.FindOrAdd(Result.Pin1->GetOwningNode(), Result);
+
+				// when zoomed out, make it easier to see diffed pins by also highlighting the node
+				if(ZoomLevel <= 6)
+				{
+					NodeDiffResults.FindOrAdd(Result.Pin1->GetOwningNode(), Result);
+				}
 			}
 			else if (Result.Node1)
 			{
@@ -168,9 +175,14 @@ int32 SGraphPanel::OnPaint( const FPaintArgs& Args, const FGeometry& AllottedGeo
 			if (Result.Pin2)
 			{
 				PinDiffResults.FindOrAdd(Result.Pin2, Result);
-				NodeDiffResults.FindOrAdd(Result.Pin2->GetOwningNode(), Result);
+
+				// when zoomed out, make it easier to see diffed pins by also highlighting the node
+				if(ZoomLevel <= 6)
+				{
+					NodeDiffResults.FindOrAdd(Result.Pin2->GetOwningNode(), Result);
+				}
 			}
-			else if (Result.Node2)
+			else if (!Result.Pin1 && Result.Node2)
 			{
 				NodeDiffResults.FindOrAdd(Result.Node2, Result);
 			}
@@ -252,6 +264,58 @@ int32 SGraphPanel::OnPaint( const FPaintArgs& Args, const FGeometry& AllottedGeo
 					}
 				}
 
+				/** if this graph is being diffed, highlight the changes in the graph */
+				if(DiffResults.IsValid())
+				{
+					/** When diffing nodes, color code shadow based on diff result */
+					if (NodeDiffResults.Contains(NodeObj))
+					{
+						const FDiffSingleResult& DiffResult = NodeDiffResults[NodeObj];
+						for (const SNode::DiffHighlightInfo& Highlight : ChildNode->GetDiffHighlights(DiffResult))
+						{
+							FSlateDrawElement::MakeBox(
+								OutDrawElements,
+								NodeDiffHighlightLayerID,
+								CurWidget.Geometry.ToInflatedPaintGeometry(NodeShadowSize),
+								Highlight.Brush,
+								ESlateDrawEffect::None,
+								Highlight.Tint
+								);
+						}
+					}
+					
+					/** When diffing, set the backround of the differing pins to their diff colors */
+					for (UEdGraphPin* Pin : NodeObj->Pins)
+					{
+						if (TSharedPtr<SGraphPin> PinWidget = ChildNode->FindWidgetForPin(Pin))
+						{
+							if (PinDiffResults.Contains(Pin))
+							{
+								// if the diff result associated with this pin is focused, highlight the pin
+								const FDiffSingleResult& DiffResult = PinDiffResults[Pin];
+								if (DiffResults.IsValid() && FocusedDiffResult.IsSet())
+								{
+									const FDiffSingleResult& Focused = (*DiffResults.Get())[FocusedDiffResult.Get()];
+									PinWidget->SetDiffHighlighted(DiffResult == Focused);
+								}
+							
+								FLinearColor PinDiffColor = DiffResult.GetDisplayColor();
+								PinDiffColor.A = 0.7f;
+								PinWidget->SetPinDiffColor(PinDiffColor);
+								PinWidget->SetFadeConnections(false);
+							}
+							else
+							{
+								PinWidget->SetDiffHighlighted(false);
+								PinWidget->SetPinDiffColor(TOptional<FLinearColor>());
+
+								// when zoomed out, fade out pin connections that aren't involved in a diff
+								PinWidget->SetFadeConnections(ZoomLevel <= 6 && (!NodeDiffResults.Contains(NodeObj) || NodeDiffResults[NodeObj].Pin1));
+							}
+						}
+					}
+				}
+
 				// Draw the node's shadow.
 				if (bDrawShadowsThisFrame || bSelected)
 				{
@@ -297,7 +361,7 @@ int32 SGraphPanel::OnPaint( const FPaintArgs& Args, const FGeometry& AllottedGeo
 					{
 						NodeStyleToUse = FadedStyle;
 					}
-					else if (bCleanDiff)
+					else if (ZoomLevel <= 6 && bCleanDiff)
 					{
 						NodeStyleToUse = FadedStyle;
 					}
diff --git a/Engine/Source/Editor/GraphEditor/Private/SGraphPin.cpp b/Engine/Source/Editor/GraphEditor/Private/SGraphPin.cpp
--- a/Engine/Source/Editor/GraphEditor/Private/SGraphPin.cpp
+++ b/Engine/Source/Editor/GraphEditor/Private/SGraphPin.cpp
@@ -147,6 +147,8 @@ SGraphPin::SGraphPin()
 
 	CachedImg_Pin_Background = FAppStyle::GetBrush( NAME_Pin_Background );
 	CachedImg_Pin_BackgroundHovered = FAppStyle::GetBrush( NAME_Pin_BackgroundHovered );
+
+	CachedImg_Pin_DiffOutline = FAppStyle::GetBrush( NAME_Pin_DiffOutline );
 }
 
 SGraphPin::~SGraphPin()
@@ -304,21 +306,27 @@ void SGraphPin::Construct(const FArguments& InArgs, UEdGraphPin* InPin)
 	}
 
 	// Set up a hover for pins that is tinted the color of the pin.
+	
 	SBorder::Construct(SBorder::FArguments()
 		.BorderImage(this, &SGraphPin::GetPinBorder)
-		.BorderBackgroundColor(this, &SGraphPin::GetPinColor)
+		.BorderBackgroundColor(this, &SGraphPin::GetHighlightColor)
 		.OnMouseButtonDown(this, &SGraphPin::OnPinNameMouseDown)
 		[
-			SNew(SLevelOfDetailBranchNode)
-			.UseLowDetailSlot(this, &SGraphPin::UseLowDetailPinNames)
-			.LowDetail()
-			[
-				//@TODO: Try creating a pin-colored line replacement that doesn't measure text / call delegates but still renders
-				PinWidgetRef
-			]
-			.HighDetail()
+			SNew(SBorder)
+			.BorderImage(CachedImg_Pin_DiffOutline)
+			.BorderBackgroundColor(this, &SGraphPin::GetPinDiffColor)
 			[
-				PinContent.ToSharedRef()
+				SNew(SLevelOfDetailBranchNode)
+				.UseLowDetailSlot(this, &SGraphPin::UseLowDetailPinNames)
+				.LowDetail()
+				[
+					//@TODO: Try creating a pin-colored line replacement that doesn't measure text / call delegates but still renders
+					PinWidgetRef
+				]
+				.HighDetail()
+				[
+					PinContent.ToSharedRef()
+				]
 			]
 		]
 	);
@@ -952,6 +960,11 @@ bool SGraphPin::IsConnected() const
 	return GraphPin? GraphPin->LinkedTo.Num() > 0 : false;
 }
 
+bool SGraphPin::AreConnectionsFaded() const
+{
+	return bFadeConnections;
+}
+
 /** @return The brush with which to pain this graph pin's incoming/outgoing bullet point */
 const FSlateBrush* SGraphPin::GetPinIcon() const
 {
@@ -1092,6 +1105,24 @@ FSlateColor SGraphPin::GetPinColor() const
 	return FLinearColor::White;
 }
 
+FSlateColor SGraphPin::GetHighlightColor() const
+{
+	if (PinDiffColor.IsSet())
+	{
+		return PinDiffColor.GetValue();
+	}
+	return GetPinColor();
+}
+
+FSlateColor SGraphPin::GetPinDiffColor() const
+{
+	if (PinDiffColor.IsSet())
+	{
+		return PinDiffColor.GetValue();
+	}
+	return FLinearColor(0.f,0.f,0.f,0.f);
+}
+
 FSlateColor SGraphPin::GetSecondaryPinColor() const
 {
 	if (UEdGraphPin* GraphPin = GetPinObj())
diff --git a/Engine/Source/Editor/GraphEditor/Private/SNodePanel.cpp b/Engine/Source/Editor/GraphEditor/Private/SNodePanel.cpp
--- a/Engine/Source/Editor/GraphEditor/Private/SNodePanel.cpp
+++ b/Engine/Source/Editor/GraphEditor/Private/SNodePanel.cpp
@@ -228,6 +228,36 @@ void SNodePanel::SNode::FNodeSlot::Construct(const FChildren& SlotOwner, FSlotAr
 	}
 }
 
+TArray<SNodePanel::SNode::DiffHighlightInfo> SNodePanel::SNode::GetDiffHighlights(
+	const FDiffSingleResult& DiffResult) const
+{
+	FLinearColor BackgroundColor = DiffResult.GetDisplayColor();
+	BackgroundColor.A = 1.f; // give highlight some transparency so it's not so 'in your face'
+		
+	FLinearColor ShadingColorHSV = BackgroundColor.LinearRGBToHSV();
+	ShadingColorHSV.R -= 15.f; // shift hue
+	if (ShadingColorHSV.R < 0.f)
+	{
+		ShadingColorHSV.R += 360.f;
+	}
+	ShadingColorHSV.B *= 0.2f; // darken
+
+	const FSlateBrush* BackgroundBrush;
+	const FSlateBrush* ForegroundBrush;
+	GetDiffHighlightBrushes(BackgroundBrush, ForegroundBrush);
+
+	return {
+		{
+			BackgroundBrush,
+			BackgroundColor
+		},
+		{
+			ForegroundBrush,
+			ShadingColorHSV.HSVToLinearRGB()
+		},
+	};
+}
+
 SNodePanel::SNodePanel()
 	: Children(this)
 	, VisibleChildren(this)
diff --git a/Engine/Source/Editor/GraphEditor/Public/KismetNodes/SGraphNodeK2Base.h b/Engine/Source/Editor/GraphEditor/Public/KismetNodes/SGraphNodeK2Base.h
--- a/Engine/Source/Editor/GraphEditor/Public/KismetNodes/SGraphNodeK2Base.h
+++ b/Engine/Source/Editor/GraphEditor/Public/KismetNodes/SGraphNodeK2Base.h
@@ -20,6 +20,7 @@ public:
 	virtual void GetOverlayBrushes(bool bSelected, const FVector2D WidgetSize, TArray<FOverlayBrushInfo>& Brushes) const override;
 	virtual void GetNodeInfoPopups(FNodeInfoContext* Context, TArray<FGraphInformationPopupInfo>& Popups) const override;
 	virtual const FSlateBrush* GetShadowBrush(bool bSelected) const override;
+	virtual void GetDiffHighlightBrushes(const FSlateBrush*& BackgroundOut, const FSlateBrush*& ForegroundOut) const override;
 	void PerformSecondPassLayout(const TMap< UObject*, TSharedRef<SNode> >& NodeToWidgetLookup) const override;
 
 protected :
diff --git a/Engine/Source/Editor/GraphEditor/Public/KismetNodes/SGraphNodeK2Var.h b/Engine/Source/Editor/GraphEditor/Public/KismetNodes/SGraphNodeK2Var.h
--- a/Engine/Source/Editor/GraphEditor/Public/KismetNodes/SGraphNodeK2Var.h
+++ b/Engine/Source/Editor/GraphEditor/Public/KismetNodes/SGraphNodeK2Var.h
@@ -19,6 +19,8 @@ public:
 
 	virtual const FSlateBrush* GetShadowBrush(bool bSelected) const override;
 
+	virtual void GetDiffHighlightBrushes(const FSlateBrush*& BackgroundOut, const FSlateBrush*& ForegroundOut) const override;
+
 	// SGraphNode interface
 	virtual void UpdateGraphNode() override;
 	// End of SGraphNode interface
diff --git a/Engine/Source/Editor/GraphEditor/Public/SGraphPin.h b/Engine/Source/Editor/GraphEditor/Public/SGraphPin.h
--- a/Engine/Source/Editor/GraphEditor/Public/SGraphPin.h
+++ b/Engine/Source/Editor/GraphEditor/Public/SGraphPin.h
@@ -150,6 +150,9 @@ public:
 	/** @return whether this pin is connected to another pin */
 	bool IsConnected() const;
 
+	/** @return whether to fade out the pin's connections */
+	bool AreConnectionsFaded() const;
+
 	/** Tries to handle making a connection to another pin, depending on the schema and the pins it may do:  
 		 - Nothing
 		 - Break existing links on either side while making the new one
@@ -175,6 +178,18 @@ public:
 		bIsDiffHighlighted = bHighlighted;
 	}
 
+	/** Allows Diff to highlight pins */
+	void SetPinDiffColor(TOptional<FLinearColor> InColor)
+	{
+		PinDiffColor = InColor;
+	}
+
+	/** Makes Pin Connection Wires transparent */
+	void SetFadeConnections(bool bInFadeConnections)
+	{
+		bFadeConnections = bInFadeConnections;
+	}
+
 	/** Set this pin to only be used to display default value */
 	void SetOnlyShowDefaultValue(bool bNewOnlyShowDefaultValue);
 
@@ -256,6 +271,11 @@ protected:
 	/** @return The color that we should use to draw this pin */
 	virtual FSlateColor GetPinColor() const;
 
+	/** @return The color that we should use to draw the highlight for this pin */
+	virtual FSlateColor GetHighlightColor() const;
+
+	virtual FSlateColor GetPinDiffColor() const;
+
 	/** @return The secondary color that we should use to draw this pin (e.g. value color for Map pins) */
 	FSlateColor GetSecondaryPinColor() const;
 
@@ -321,6 +341,9 @@ protected:
 	/** Set of pins that are currently being hovered */
 	TSet< FEdGraphPinReference > HoverPinSet;
 
+	/** If set, this will change the color of the pin's background highlight */
+	TOptional<FLinearColor> PinDiffColor;
+
 	//@TODO: Want to cache these once for all SGraphPins, but still handle slate style updates
 	const FSlateBrush* CachedImg_ArrayPin_Connected;
 	const FSlateBrush* CachedImg_ArrayPin_Disconnected;
@@ -338,6 +361,8 @@ protected:
 
 	const FSlateBrush* CachedImg_Pin_Background;
 	const FSlateBrush* CachedImg_Pin_BackgroundHovered;
+	
+	const FSlateBrush* CachedImg_Pin_DiffOutline;
 
 	const FSlateBrush* Custom_Brush_Connected;
 	const FSlateBrush* Custom_Brush_Disconnected;
@@ -362,4 +387,7 @@ protected:
 
 	/** True if this pin is being diffed and it's currently selected in the diff view. Highlights this pin */
 	bool bIsDiffHighlighted;
+
+	/** TRUE if the connections from this pin should be drawn at a lower opacity */
+	bool bFadeConnections;
 };
diff --git a/Engine/Source/Editor/GraphEditor/Public/SNodePanel.h b/Engine/Source/Editor/GraphEditor/Public/SNodePanel.h
--- a/Engine/Source/Editor/GraphEditor/Public/SNodePanel.h
+++ b/Engine/Source/Editor/GraphEditor/Public/SNodePanel.h
@@ -3,6 +3,7 @@
 #pragma once
 
 #include "CoreMinimal.h"
+#include "DiffResults.h"
 #include "Misc/Attribute.h"
 #include "Misc/Guid.h"
 #include "Layout/Visibility.h"
@@ -567,6 +568,22 @@ public:
 			return bSelected ? FAppStyle::GetBrush(TEXT("Graph.Node.ShadowSelected")) : FAppStyle::GetBrush(TEXT("Graph.Node.Shadow"));
 		}
 
+		struct DiffHighlightInfo
+		{
+			const FSlateBrush* Brush;
+			const FLinearColor Tint;
+		};
+		
+		/** @return Collection of brushes layered to outline node with the DiffResult color */
+		TArray<SNodePanel::SNode::DiffHighlightInfo> GetDiffHighlights(const FDiffSingleResult& DiffResult) const;
+
+		/** used by GetDiffHighlights to generate outlines for diffed nodes */
+		virtual void GetDiffHighlightBrushes(const FSlateBrush*& BackgroundOut, const FSlateBrush*& ForegroundOut) const
+		{
+			BackgroundOut = FAppStyle::GetBrush(TEXT("Graph.Node.DiffHighlight"));
+			ForegroundOut = FAppStyle::GetBrush(TEXT("Graph.Node.DiffHighlightShading"));
+		}
+
 		/** Populate the brushes array with any overlay brushes to render */
 		virtual void GetOverlayBrushes(bool bSelected, const FVector2D WidgetSize, TArray<FOverlayBrushInfo>& Brushes) const
 		{
diff --git a/Engine/Source/Editor/Kismet/Private/SBlueprintDiff.cpp b/Engine/Source/Editor/Kismet/Private/SBlueprintDiff.cpp
--- a/Engine/Source/Editor/Kismet/Private/SBlueprintDiff.cpp
+++ b/Engine/Source/Editor/Kismet/Private/SBlueprintDiff.cpp
@@ -53,7 +53,7 @@ const FName GraphMode = FName(TEXT("GraphMode"));
 TSharedRef<SWidget>	FDiffResultItem::GenerateWidget() const
 {
 	FText ToolTip = Result.ToolTip;
-	FLinearColor Color = Result.DisplayColor;
+	FLinearColor Color = Result.GetDisplayColor();
 	FText Text = Result.DisplayString;
 	if (Text.IsEmpty())
 	{
@@ -1480,6 +1480,7 @@ void SBlueprintDiff::HandleGraphChanged( const FString& GraphPath )
 		}
 	}
 
+	
 	PanelOld.GeneratePanel(GraphOld, DiffResults, FocusedDiffResult);
 	PanelNew.GeneratePanel(GraphNew, DiffResults, FocusedDiffResult);
 }
diff --git a/Engine/Source/Editor/UMGEditor/Private/WidgetBlueprint.cpp b/Engine/Source/Editor/UMGEditor/Private/WidgetBlueprint.cpp
--- a/Engine/Source/Editor/UMGEditor/Private/WidgetBlueprint.cpp
+++ b/Engine/Source/Editor/UMGEditor/Private/WidgetBlueprint.cpp
@@ -864,7 +864,7 @@ bool UWidgetBlueprint::FindDiffs(const UBlueprint* OtherBlueprint, FDiffResults&
 				Args.Add(TEXT("WidgetPath"), FText::FromString(Pair.Key));
 				Diff.ToolTip = FText::Format(LOCTEXT("DIF_RequestWidgetTooltip", "Widget {WidgetTitle}\nPath: {WidgetPath}"), Args);
 				Diff.DisplayString = FText::Format(LOCTEXT("DIF_RequestWidgetLabel", "Widget {WidgetTitle}"), Args);
-				Diff.DisplayColor = FLinearColor(1.f, 0.4f, 0.4f);
+				Diff.Category = EDiffType::CONTROL;
 
 				Results.Add(Diff);
 
@@ -883,7 +883,7 @@ bool UWidgetBlueprint::FindDiffs(const UBlueprint* OtherBlueprint, FDiffResults&
 					SlotArgs.Add(TEXT("WidgetPath"), FText::FromString(Pair.Key));
 					SlotDiff.ToolTip = FText::Format(LOCTEXT("DIF_RequestSlotTooltip", "Slot for {WidgetTitle}\nPath: {WidgetPath}"), SlotArgs);
 					SlotDiff.DisplayString = FText::Format(LOCTEXT("DIF_RequestSlotLabel", "Slot for {WidgetTitle}"), SlotArgs);
-					SlotDiff.DisplayColor = FLinearColor(1.f, 0.4f, 0.4f);
+					Diff.Category = EDiffType::CONTROL;
 
 					Results.Add(SlotDiff);
 				}
@@ -904,7 +904,7 @@ bool UWidgetBlueprint::FindDiffs(const UBlueprint* OtherBlueprint, FDiffResults&
 				Args.Add(TEXT("WidgetPath"), FText::FromString(Pair.Key));
 				Diff.ToolTip = FText::Format(LOCTEXT("DIF_AddedWidgetTooltip", "Added Widget {WidgetTitle}\nPath: {WidgetPath}"), Args);
 				Diff.DisplayString = FText::Format(LOCTEXT("DIF_AddedWidgetLabel", "Added Widget {WidgetTitle}"), Args);
-				Diff.DisplayColor = FLinearColor(1.f, 0.4f, 0.4f);
+				Diff.Category = EDiffType::ADDITION;
 			}
 
 			Results.Add(Diff);
@@ -931,7 +931,7 @@ bool UWidgetBlueprint::FindDiffs(const UBlueprint* OtherBlueprint, FDiffResults&
 				Args.Add(TEXT("WidgetPath"), FText::FromString(Pair.Key));
 				Diff.ToolTip = FText::Format(LOCTEXT("DIF_RemovedWidgetTooltip", "Removed Widget {WidgetTitle}\nPath:{WidgetPath}"), Args);
 				Diff.DisplayString = FText::Format(LOCTEXT("DIF_RemovedWidgetLabel", "Removed Widget {WidgetTitle}"), Args);
-				Diff.DisplayColor = FLinearColor(1.f, 0.4f, 0.4f);
+				Diff.Category = EDiffType::SUBTRACTION;
 			}
 
 			Results.Add(Diff);
@@ -943,7 +943,7 @@ bool UWidgetBlueprint::FindDiffs(const UBlueprint* OtherBlueprint, FDiffResults&
 	{
 		FDiffSingleResult Diff;
 		Diff.Diff = EDiffType::INFO_MESSAGE;
-		Diff.DisplayColor = FLinearColor(.7f, .7f, .7f);
+		Diff.Category = EDiffType::CONTROL;
 		Diff.ToolTip = LOCTEXT("DIF_WidgetWarningMessage", "Warning: This may be missing changes to Animations and Bindings");
 		Diff.DisplayString = Diff.ToolTip;
 
diff --git a/Engine/Source/Editor/WorldBrowser/Private/Tiles/SWorldTileItem.cpp b/Engine/Source/Editor/WorldBrowser/Private/Tiles/SWorldTileItem.cpp
--- a/Engine/Source/Editor/WorldBrowser/Private/Tiles/SWorldTileItem.cpp
+++ b/Engine/Source/Editor/WorldBrowser/Private/Tiles/SWorldTileItem.cpp
@@ -223,6 +223,12 @@ const FSlateBrush* SWorldTileItem::GetShadowBrush(bool bSelected) const
 	return bSelected ? FAppStyle::GetBrush(TEXT("Graph.CompactNode.ShadowSelected")) : FAppStyle::GetBrush(TEXT("Graph.Node.Shadow"));
 }
 
+void SWorldTileItem::GetDiffHighlightBrushes(const FSlateBrush*& BackgroundOut, const FSlateBrush*& ForegroundOut) const
+{
+	BackgroundOut = FAppStyle::GetBrush(TEXT("Graph.CompactNode.DiffHighlight"));
+	ForegroundOut = FAppStyle::GetBrush(TEXT("Graph.CompactNode.DiffHighlightShading"));
+}
+
 FOptionalSize SWorldTileItem::GetItemWidth() const
 {
 	return FOptionalSize(TileModel->GetLevelSize2D().X);
diff --git a/Engine/Source/Editor/WorldBrowser/Private/Tiles/SWorldTileItem.h b/Engine/Source/Editor/WorldBrowser/Private/Tiles/SWorldTileItem.h
--- a/Engine/Source/Editor/WorldBrowser/Private/Tiles/SWorldTileItem.h
+++ b/Engine/Source/Editor/WorldBrowser/Private/Tiles/SWorldTileItem.h
@@ -70,6 +70,7 @@ public:
 	virtual UObject* GetObjectBeingDisplayed() const override;
 	virtual FVector2D GetPosition() const override;
 	virtual const FSlateBrush* GetShadowBrush(bool bSelected) const override;
+	virtual void GetDiffHighlightBrushes(const FSlateBrush*& BackgroundOut, const FSlateBrush*& ForegroundOut) const override;
 	// SNodePanel::SNode interface end
 	
 	/** @return Deferred item refresh */
diff --git a/Engine/Source/Runtime/Engine/Private/EdGraph/EdGraphNode.cpp b/Engine/Source/Runtime/Engine/Private/EdGraph/EdGraphNode.cpp
--- a/Engine/Source/Runtime/Engine/Private/EdGraph/EdGraphNode.cpp
+++ b/Engine/Source/Runtime/Engine/Private/EdGraph/EdGraphNode.cpp
@@ -781,7 +781,7 @@ void UEdGraphNode::FindDiffs(UEdGraphNode* OtherNode, struct FDiffResults& Resul
 		Diff.Node1 = this;
 		Diff.Node2 = OtherNode;
 		Diff.ToolTip = LOCTEXT("DIF_NodePropertyToolTip", "A Property of the node has changed");
-		Diff.DisplayColor = FLinearColor(0.25f, 0.71f, 0.85f);
+		Diff.Category = EDiffType::MODIFICATION;
 
 		// Diff the properties between the nodes
 		DiffProperties(GetClass(), OtherNode->GetClass(), this, OtherNode, Results, Diff);
diff --git a/Engine/Source/Runtime/Engine/Public/DiffResults.h b/Engine/Source/Runtime/Engine/Public/DiffResults.h
--- a/Engine/Source/Runtime/Engine/Public/DiffResults.h
+++ b/Engine/Source/Runtime/Engine/Public/DiffResults.h
@@ -50,6 +50,19 @@ namespace EDiffType
 		// Informational message, does't count as a real diff
 		INFO_MESSAGE
 	};
+	
+	enum Category
+	{
+		ADDITION,
+		SUBTRACTION,
+		MODIFICATION,
+		
+		// used for small changes like moving nodes that don't effect the compilation
+		MINOR,
+
+		// used for items that are purely informational and don't show up in the diff results
+		CONTROL,
+	};
 }
 
 /** Result of a single difference found on graph or object */
@@ -64,12 +77,13 @@ struct FDiffSingleResult
 		Pin2 = nullptr;
 		Object1 = nullptr;
 		Object2 = nullptr;
-		DisplayColor = FLinearColor::White;
 	}
 
 	/** The type of diff */
 	EDiffType::Type Diff;
 
+	EDiffType::Category Category;
+
 	/** The first node involved in diff */
 	class UEdGraphNode* Node1;
 
@@ -94,8 +108,19 @@ struct FDiffSingleResult
 	/** Optional tooltip containing more information */
 	FText ToolTip; 
 
-	/** User can override color to use for display string */
-	FLinearColor DisplayColor;
+	/** Get the color that is associated with this diff category */
+	FLinearColor GetDisplayColor() const
+	{
+		switch(Category)
+		{
+		case EDiffType::ADDITION: return FLinearColor(0.02f, 0.94f, 0.f);
+		case EDiffType::SUBTRACTION: return FLinearColor(1.f, 0.16f, 0.16f);
+		case EDiffType::MODIFICATION: return FLinearColor(0.04f, 0.87f, 1.f);
+		case EDiffType::MINOR:  return FLinearColor(0.74f, 0.69f, 0.79f);
+		
+		default: return FLinearColor::White;
+		}
+	}
 
 	/** Path string of graph, relative to blueprint/asset root */
 	FString OwningObjectPath;
@@ -126,7 +151,7 @@ FORCEINLINE bool operator==( const FDiffSingleResult& LHS, const FDiffSingleResu
 			LHS.Object2 == RHS.Object2 &&
 			LHS.DisplayString.ToString() == RHS.DisplayString.ToString() &&
 			LHS.ToolTip.ToString() == RHS.ToolTip.ToString() &&
-			LHS.DisplayColor == RHS.DisplayColor;
+			LHS.Category == RHS.Category;
 }
 
 /** Collects the Diffs found for a node/object */
