[BugFix] Crash using blueprint Diff view on complex blueprints
+ also added in Pin_DiffHighlight brush that never got committed

diff --git a/Engine/Source/Editor/GraphEditor/Private/SGraphPanel.cpp b/Engine/Source/Editor/GraphEditor/Private/SGraphPanel.cpp
--- a/Engine/Source/Editor/GraphEditor/Private/SGraphPanel.cpp
+++ b/Engine/Source/Editor/GraphEditor/Private/SGraphPanel.cpp
@@ -253,8 +253,12 @@ int32 SGraphPanel::OnPaint( const FPaintArgs& Args, const FGeometry& AllottedGeo
 							// if the diff result associated with this pin is focused, highlight the pin
 							if (DiffResults.IsValid() && FocusedDiffResult.IsSet())
 							{
-								const FDiffSingleResult& Focused = (*DiffResults)[FocusedDiffResult.Get()];
-								PinWidget->SetDiffHighlighted(*DiffResult == Focused);
+								const int32 Index = FocusedDiffResult.Get();
+								if (DiffResults->IsValidIndex(Index))
+								{
+									const FDiffSingleResult& Focused = (*DiffResults)[Index];
+									PinWidget->SetDiffHighlighted(*DiffResult == Focused);
+								}
 							}
 						}
 						else
