Fixed an issue where an uninitialized flag would cause some graph connections to be faded out.

diff --git a/Engine/Source/Editor/GraphEditor/Private/SGraphPin.cpp b/Engine/Source/Editor/GraphEditor/Private/SGraphPin.cpp
--- a/Engine/Source/Editor/GraphEditor/Private/SGraphPin.cpp
+++ b/Engine/Source/Editor/GraphEditor/Private/SGraphPin.cpp
@@ -80,6 +80,7 @@ SGraphPin::SGraphPin()
 	, bIsMovingLinks(false)
 	, bUsePinColorForText(false)
 	, bDragAndDropEnabled(true)
+	, bFadeConnections(false)
 {
 	IsEditable = true;
 
