[BugFix] Crash using blueprint Diff view without a selected diff item
- This was the result of a merge error. I removed the duplicate/incorrect code

[BugFix] Selected diff Pins aren't highlighted in blueprints with more than one graph
- FocusedDiffResult attribute returned the index of the diff result in a list of all diff results when it needed to return the index of a subset array that only includes results diffs between the currently rendered graph

diff --git a/Engine/Source/Editor/GraphEditor/Private/SGraphPanel.cpp b/Engine/Source/Editor/GraphEditor/Private/SGraphPanel.cpp
--- a/Engine/Source/Editor/GraphEditor/Private/SGraphPanel.cpp
+++ b/Engine/Source/Editor/GraphEditor/Private/SGraphPanel.cpp
@@ -242,31 +242,6 @@ int32 SGraphPanel::OnPaint( const FPaintArgs& Args, const FGeometry& AllottedGeo
 						ChildNode->ApplyRename();
 					}
 				}
-				
-				/** Highlight Pins that are focused by the diff view */
-				for (UEdGraphPin* Pin : NodeObj->Pins)
-				{
-					if (TSharedPtr<SGraphPin> PinWidget = ChildNode->FindWidgetForPin(Pin))
-					{
-						if (FDiffSingleResult* DiffResult = PinDiffResults.Find(Pin))
-						{
-							// if the diff result associated with this pin is focused, highlight the pin
-							if (DiffResults.IsValid() && FocusedDiffResult.IsSet())
-							{
-								const int32 Index = FocusedDiffResult.Get();
-								if (DiffResults->IsValidIndex(Index))
-								{
-									const FDiffSingleResult& Focused = (*DiffResults)[Index];
-									PinWidget->SetDiffHighlighted(*DiffResult == Focused);
-								}
-							}
-						}
-						else
-						{
-							PinWidget->SetDiffHighlighted(false);
-						}
-					}
-				}
 
 				/** if this graph is being diffed, highlight the changes in the graph */
 				if(DiffResults.IsValid())
@@ -287,35 +262,38 @@ int32 SGraphPanel::OnPaint( const FPaintArgs& Args, const FGeometry& AllottedGeo
 								);
 						}
 					}
-					
-					/** When diffing, set the backround of the differing pins to their diff colors */
-					for (UEdGraphPin* Pin : NodeObj->Pins)
+				}
+				
+				/** When diffing, set the backround of the differing pins to their diff colors */
+				for (UEdGraphPin* Pin : NodeObj->Pins)
+				{
+					if (TSharedPtr<SGraphPin> PinWidget = ChildNode->FindWidgetForPin(Pin))
 					{
-						if (TSharedPtr<SGraphPin> PinWidget = ChildNode->FindWidgetForPin(Pin))
+						if (FDiffSingleResult* DiffResult = PinDiffResults.Find(Pin))
 						{
-							if (PinDiffResults.Contains(Pin))
+							// if the diff result associated with this pin is focused, highlight the pin
+							if (DiffResults.IsValid() && FocusedDiffResult.IsSet())
 							{
-								// if the diff result associated with this pin is focused, highlight the pin
-								const FDiffSingleResult& DiffResult = PinDiffResults[Pin];
-								if (DiffResults.IsValid() && FocusedDiffResult.IsSet())
+								const int32 Index = FocusedDiffResult.Get();
+								if (DiffResults->IsValidIndex(Index))
 								{
-									const FDiffSingleResult& Focused = (*DiffResults.Get())[FocusedDiffResult.Get()];
-									PinWidget->SetDiffHighlighted(DiffResult == Focused);
+									const FDiffSingleResult& Focused = (*DiffResults)[Index];
+									PinWidget->SetDiffHighlighted(*DiffResult == Focused);
 								}
-							
-								FLinearColor PinDiffColor = DiffResult.GetDisplayColor();
-								PinDiffColor.A = 0.7f;
-								PinWidget->SetPinDiffColor(PinDiffColor);
-								PinWidget->SetFadeConnections(false);
 							}
-							else
-							{
-								PinWidget->SetDiffHighlighted(false);
-								PinWidget->SetPinDiffColor(TOptional<FLinearColor>());
+						
+							FLinearColor PinDiffColor = DiffResult->GetDisplayColor();
+							PinDiffColor.A = 0.7f;
+							PinWidget->SetPinDiffColor(PinDiffColor);
+							PinWidget->SetFadeConnections(false);
+						}
+						else
+						{
+							PinWidget->SetDiffHighlighted(false);
+							PinWidget->SetPinDiffColor(TOptional<FLinearColor>());
 
-								// when zoomed out, fade out pin connections that aren't involved in a diff
-								PinWidget->SetFadeConnections(ZoomLevel <= 6 && (!NodeDiffResults.Contains(NodeObj) || NodeDiffResults[NodeObj].Pin1));
-							}
+							// when zoomed out, fade out pin connections that aren't involved in a diff
+							PinWidget->SetFadeConnections(ZoomLevel <= 6 && (!NodeDiffResults.Contains(NodeObj) || NodeDiffResults[NodeObj].Pin1));
 						}
 					}
 				}
diff --git a/Engine/Source/Editor/Kismet/Private/SBlueprintDiff.cpp b/Engine/Source/Editor/Kismet/Private/SBlueprintDiff.cpp
--- a/Engine/Source/Editor/Kismet/Private/SBlueprintDiff.cpp
+++ b/Engine/Source/Editor/Kismet/Private/SBlueprintDiff.cpp
@@ -1461,7 +1461,6 @@ void SBlueprintDiff::HandleGraphChanged( const FString& GraphPath )
 	UEdGraph* GraphOld = nullptr;
 	UEdGraph* GraphNew = nullptr;
 	TSharedPtr<TArray<FDiffSingleResult>> DiffResults;
-	TAttribute<int32> FocusedDiffResult;
 	for (const TSharedPtr<FGraphToDiff>& GraphToDiff : Graphs)
 	{
 		UEdGraph* NewGraph = GraphToDiff->GetGraphNew();
@@ -1471,14 +1470,30 @@ void SBlueprintDiff::HandleGraphChanged( const FString& GraphPath )
 			GraphNew = NewGraph;
 			GraphOld = OldGraph;
 			DiffResults = GraphToDiff->FoundDiffs;
-			FocusedDiffResult = TAttribute<int32>::CreateLambda(
-				[DifferencesTreeView = DifferencesTreeView.ToSharedRef(), &RealDifferences = RealDifferences]()
-				{
-					return DiffTreeView::CurrentDifference(DifferencesTreeView, RealDifferences);
-				});
 			break;
 		}
 	}
+	
+	const TAttribute<int32> FocusedDiffResult = TAttribute<int32>::CreateLambda(
+        [this, GraphPath]()
+        {
+        	// because DifferencesTreeView has differences from several graphs, we need to find the index of the first entry in this graph
+        	// that way we can return the selected index of just this graph's entries
+        	int32 startIndex = 0;
+			for (const TSharedPtr<FGraphToDiff>& GraphToDiff : Graphs)
+			{
+				if (GraphPath.Equals(FGraphDiffControl::GetGraphPath(GraphToDiff->GetGraphNew())))
+				{
+					break;
+				}
+				if (GraphToDiff->FoundDiffs.IsValid())
+				{
+					startIndex += GraphToDiff->FoundDiffs->Num();
+				}
+			}
+			// find selected index in all the graphs, and subtract the index of the first entry in this graph
+			return DiffTreeView::CurrentDifference(DifferencesTreeView.ToSharedRef(), RealDifferences) - startIndex;
+        });
 
 	
 	PanelOld.GeneratePanel(GraphOld, DiffResults, FocusedDiffResult);
