[BugFix] Crash when attempting to open AnimGraph in ABP_BlendSpaceGraph blueprint

diff --git a/Engine/Source/Editor/GraphEditor/Private/ConnectionDrawingPolicy.cpp b/Engine/Source/Editor/GraphEditor/Private/ConnectionDrawingPolicy.cpp
--- a/Engine/Source/Editor/GraphEditor/Private/ConnectionDrawingPolicy.cpp
+++ b/Engine/Source/Editor/GraphEditor/Private/ConnectionDrawingPolicy.cpp
@@ -489,9 +489,13 @@ void FConnectionDrawingPolicy::DrawPinGeometries(TMap<TSharedRef<SWidget>, FArra
 				{
 					FConnectionParams Params;
 					DetermineWiringStyle(ThePin, TargetPin, /*inout*/ Params);
-					if ( PinWidget.AreConnectionsFaded() && PinToPinWidgetMap[TargetPin]->AreConnectionsFaded() )
+					const TSharedPtr<SGraphPin>* ConnectedPinWidget = PinToPinWidgetMap.Find(TargetPin);
+					if (ConnectedPinWidget && ConnectedPinWidget->IsValid())
 					{
-						Params.WireColor.A = 0.2f;
+						if ( PinWidget.AreConnectionsFaded() && (*ConnectedPinWidget)->AreConnectionsFaded() )
+						{
+							Params.WireColor.A = 0.2f;
+						}
 					}
 					DrawSplineWithArrow(LinkStartWidgetGeometry->Geometry, LinkEndWidgetGeometry->Geometry, Params);
 				}
