[Bug Fix] Navigating the Event Graph in Blueprint Diff window causes visual snapping issues

diff --git a/Engine/Source/Editor/GraphEditor/Private/SGraphEditorImpl.cpp b/Engine/Source/Editor/GraphEditor/Private/SGraphEditorImpl.cpp
--- a/Engine/Source/Editor/GraphEditor/Private/SGraphEditorImpl.cpp
+++ b/Engine/Source/Editor/GraphEditor/Private/SGraphEditorImpl.cpp
@@ -1424,32 +1424,48 @@ void SGraphEditorImpl::AddNotification( FNotificationInfo& Info, bool bSuccess )
 
 EActiveTimerReturnType SGraphEditorImpl::HandleFocusEditorDeferred(double InCurrentTime, float InDeltaTime)
 {
+	
+	// If GraphPanel is going to pan to a target but hasn't yet, wait until it does so we don't miss it
+	if (GraphPanel->HasDeferredZoomDestination())
+	{
+		return EActiveTimerReturnType::Continue;
+	}
+	
 	for( int i = 0; i < LockedGraphs.Num(); ++i )
 	{
 		TSharedPtr<SGraphEditor> LockedGraph = LockedGraphs[i].Pin();
 		if (LockedGraph != TSharedPtr<SGraphEditor>())
 		{
-			LockedGraph->SetViewLocation(GraphPanel->GetViewOffset(), GraphPanel->GetZoomAmount());
+			// If the locked graph is going to pan to a target but hasn't yet, wait until it does so we don't miss it
+			if (LockedGraph->GetGraphPanel()->HasDeferredZoomDestination())
+			{
+				return EActiveTimerReturnType::Continue;
+			}
 
-			// if GraphPanel is currently panning to a target, pan with it
 			FVector2D TopLeft, BottomRight;
+
+			// If the locked graph was instructed to pan to a destination, let it ignore the lock to reach that destination.
+			// this way we can support diffs of moved nodes.
+			if (LockedGraph->GetGraphPanel()->GetZoomTargetRect(TopLeft, BottomRight))
+			{
+				continue;
+			}
+			
+			// Send the locked graph to the same place as this graph
 			if (GraphPanel->GetZoomTargetRect(TopLeft, BottomRight))
 			{
-				// if this graph is zooming/panning to a target location, have the locked graph do the same
 				LockedGraph->GetGraphPanel()->JumpToRect(TopLeft, BottomRight);
 			}
+			else
+			{
+				LockedGraph->SetViewLocation(GraphPanel->GetViewOffset(), GraphPanel->GetZoomAmount());
+			}
 		}
 		else
 		{
 			LockedGraphs.RemoveAtSwap(i--);
 		}
 	}
-	
-	// If GraphPanel is going to pan to a target but hasn't yet, wait until it does so we don't miss it
-	if (GraphPanel->HasDeferredObjectFocus())
-	{
-		return EActiveTimerReturnType::Continue;
-	}
 	return EActiveTimerReturnType::Stop;
 }
 
diff --git a/Engine/Source/Editor/GraphEditor/Private/SNodePanel.cpp b/Engine/Source/Editor/GraphEditor/Private/SNodePanel.cpp
--- a/Engine/Source/Editor/GraphEditor/Private/SNodePanel.cpp
+++ b/Engine/Source/Editor/GraphEditor/Private/SNodePanel.cpp
@@ -1779,6 +1779,11 @@ bool SNodePanel::HasDeferredObjectFocus() const
 	return DeferredMovementTargetObject != nullptr;
 }
 
+bool SNodePanel::HasDeferredZoomDestination() const
+{
+	return HasDeferredObjectFocus() || bDeferredZoomToSelection || bDeferredZoomToNodeExtents;
+}
+
 void SNodePanel::FinalizeNodeMovements()
 {
 	// Process moved nodes on focus lost
@@ -1848,6 +1853,10 @@ void SNodePanel::CancelZoomToFit()
 {
 	if (ActiveTimerHandle.IsValid())
 	{
+		// Reset Zoom destination
+		ZoomPadding = NodePanelDefs::DefaultZoomPadding;
+		ZoomTargetTopLeft = FVector2D::ZeroVector;
+		ZoomTargetBottomRight = FVector2D::ZeroVector;
 		UnRegisterActiveTimer(ActiveTimerHandle.Pin().ToSharedRef());
 	}
 }
diff --git a/Engine/Source/Editor/GraphEditor/Public/SNodePanel.h b/Engine/Source/Editor/GraphEditor/Public/SNodePanel.h
--- a/Engine/Source/Editor/GraphEditor/Public/SNodePanel.h
+++ b/Engine/Source/Editor/GraphEditor/Public/SNodePanel.h
@@ -777,6 +777,9 @@ public:
 	/** If it is focusing on a particular object */
 	bool HasDeferredObjectFocus() const;
 
+	/** Query whether this graph is about to start panning/zooming towards a destination */
+	bool HasDeferredZoomDestination() const;
+
 	/** Commit transactions for any node movements */
 	void FinalizeNodeMovements();
 
diff --git a/Engine/Source/Editor/Kismet/Private/SBlueprintDiff.cpp b/Engine/Source/Editor/Kismet/Private/SBlueprintDiff.cpp
--- a/Engine/Source/Editor/Kismet/Private/SBlueprintDiff.cpp
+++ b/Engine/Source/Editor/Kismet/Private/SBlueprintDiff.cpp
@@ -1499,9 +1499,17 @@ void SBlueprintDiff::HandleGraphChanged( const FString& GraphPath )
 			return DiffTreeView::CurrentDifference(DifferencesTreeView.ToSharedRef(), RealDifferences) - startIndex;
         });
 
+	// only regenerate PanelOld if the old graph has changed
+	if (!PanelOld.GraphEditor.IsValid() || GraphOld != PanelOld.GraphEditor.Pin()->GetCurrentGraph())
+	{
+		PanelOld.GeneratePanel(GraphOld, DiffResults, FocusedDiffResult);
+	}
 	
-	PanelOld.GeneratePanel(GraphOld, DiffResults, FocusedDiffResult);
-	PanelNew.GeneratePanel(GraphNew, DiffResults, FocusedDiffResult);
+	// only regenerate PanelNew if the old graph has changed
+	if (!PanelNew.GraphEditor.IsValid() || GraphNew != PanelNew.GraphEditor.Pin()->GetCurrentGraph())
+	{
+		PanelNew.GeneratePanel(GraphNew, DiffResults, FocusedDiffResult);
+	}
 }
 
 void SBlueprintDiff::GenerateDifferencesList()
