UE-157166 - Vector literal pins lack precision with high values

Similar to other numerical Slate widgets, we need to templatize SGraphPinVector so that it can choose between a float and double representation. Currently, only the BP editor will use the double variant, while other tools continue to use a float.

This also update SGraphPinVector2D and SGraphPinVector4 to use a template. While we do use SGraphPinVector2D in BPs, we currently don't use SGraphPinVector4. However, that may change in the future.

diff --git a/Engine/Plugins/FX/Niagara/Source/NiagaraEditor/Private/NiagaraEditorModule.cpp b/Engine/Plugins/FX/Niagara/Source/NiagaraEditor/Private/NiagaraEditorModule.cpp
--- a/Engine/Plugins/FX/Niagara/Source/NiagaraEditor/Private/NiagaraEditorModule.cpp
+++ b/Engine/Plugins/FX/Niagara/Source/NiagaraEditor/Private/NiagaraEditorModule.cpp
@@ -1019,16 +1019,16 @@ void FNiagaraEditorModule::StartupModule()
 		[](UEdGraphPin* GraphPin) -> TSharedRef<SGraphPin> { return SNew(TNiagaraGraphPinEditableName<SGraphPinInteger>, GraphPin); }));
 
 	GraphPanelPinFactory->RegisterTypePin(FNiagaraTypeDefinition::GetVec2Struct(), FNiagaraScriptGraphPanelPinFactory::FCreateGraphPin::CreateLambda(
-		[](UEdGraphPin* GraphPin) -> TSharedRef<SGraphPin> { return SNew(TNiagaraGraphPinEditableName<SGraphPinVector2D>, GraphPin); }));
+		[](UEdGraphPin* GraphPin) -> TSharedRef<SGraphPin> { return SNew(TNiagaraGraphPinEditableName<SGraphPinVector2D<float>>, GraphPin); }));
 
 	GraphPanelPinFactory->RegisterTypePin(FNiagaraTypeDefinition::GetVec3Struct(), FNiagaraScriptGraphPanelPinFactory::FCreateGraphPin::CreateLambda(
-		[](UEdGraphPin* GraphPin) -> TSharedRef<SGraphPin> { return SNew(TNiagaraGraphPinEditableName<SGraphPinVector>, GraphPin); }));
+		[](UEdGraphPin* GraphPin) -> TSharedRef<SGraphPin> { return SNew(TNiagaraGraphPinEditableName<SGraphPinVector<float>>, GraphPin); }));
 
 	GraphPanelPinFactory->RegisterTypePin(FNiagaraTypeDefinition::GetPositionStruct(), FNiagaraScriptGraphPanelPinFactory::FCreateGraphPin::CreateLambda(
-		[](UEdGraphPin* GraphPin) -> TSharedRef<SGraphPin> { return SNew(TNiagaraGraphPinEditableName<SGraphPinVector>, GraphPin); }));
+		[](UEdGraphPin* GraphPin) -> TSharedRef<SGraphPin> { return SNew(TNiagaraGraphPinEditableName<SGraphPinVector<float>>, GraphPin); }));
 
 	GraphPanelPinFactory->RegisterTypePin(FNiagaraTypeDefinition::GetVec4Struct(), FNiagaraScriptGraphPanelPinFactory::FCreateGraphPin::CreateLambda(
-		[](UEdGraphPin* GraphPin) -> TSharedRef<SGraphPin> { return SNew(TNiagaraGraphPinEditableName<SGraphPinVector4>, GraphPin); }));
+		[](UEdGraphPin* GraphPin) -> TSharedRef<SGraphPin> { return SNew(TNiagaraGraphPinEditableName<SGraphPinVector4<float>>, GraphPin); }));
 
 	GraphPanelPinFactory->RegisterTypePin(FNiagaraTypeDefinition::GetColorStruct(), FNiagaraScriptGraphPanelPinFactory::FCreateGraphPin::CreateLambda(
 		[](UEdGraphPin* GraphPin) -> TSharedRef<SGraphPin> { return SNew(TNiagaraGraphPinEditableName<SGraphPinColor>, GraphPin); }));
diff --git a/Engine/Source/Editor/GraphEditor/Private/KismetPins/SGraphPinVector.cpp b/Engine/Source/Editor/GraphEditor/Private/KismetPins/SGraphPinVector.cpp
deleted file mode 100644
diff --git a/Engine/Source/Editor/GraphEditor/Private/KismetPins/SGraphPinVector2D.cpp b/Engine/Source/Editor/GraphEditor/Private/KismetPins/SGraphPinVector2D.cpp
deleted file mode 100644
diff --git a/Engine/Source/Editor/GraphEditor/Private/KismetPins/SGraphPinVector4.cpp b/Engine/Source/Editor/GraphEditor/Private/KismetPins/SGraphPinVector4.cpp
deleted file mode 100644
diff --git a/Engine/Source/Editor/GraphEditor/Private/NodeFactory.cpp b/Engine/Source/Editor/GraphEditor/Private/NodeFactory.cpp
--- a/Engine/Source/Editor/GraphEditor/Private/NodeFactory.cpp
+++ b/Engine/Source/Editor/GraphEditor/Private/NodeFactory.cpp
@@ -332,11 +332,11 @@ TSharedPtr<SGraphPin> FNodeFactory::CreateK2PinWidget(UEdGraphPin* InPin)
 		}
 		else if ((InPin->PinType.PinSubCategoryObject == VectorStruct) || (InPin->PinType.PinSubCategoryObject == Vector3fStruct) || (InPin->PinType.PinSubCategoryObject == RotatorStruct))
 		{
-			return SNew(SGraphPinVector, InPin);
+			return SNew(SGraphPinVector<double>, InPin);
 		}
 		else if (InPin->PinType.PinSubCategoryObject == Vector2DStruct)
 		{
-			return SNew(SGraphPinVector2D, InPin);
+			return SNew(SGraphPinVector2D<double>, InPin);
 		}
 		else if (InPin->PinType.PinSubCategoryObject == FKey::StaticStruct())
 		{
diff --git a/Engine/Source/Editor/GraphEditor/Public/KismetPins/SGraphPinVector.h b/Engine/Source/Editor/GraphEditor/Public/KismetPins/SGraphPinVector.h
--- a/Engine/Source/Editor/GraphEditor/Public/KismetPins/SGraphPinVector.h
+++ b/Engine/Source/Editor/GraphEditor/Public/KismetPins/SGraphPinVector.h
@@ -6,14 +6,19 @@
 #include "Widgets/DeclarativeSyntaxSupport.h"
 #include "Widgets/SWidget.h"
 #include "SGraphPin.h"
+#include "SVectorTextBox.h"
 
-class GRAPHEDITOR_API SGraphPinVector : public SGraphPin
+template <typename NumericType>
+class SGraphPinVector : public SGraphPin
 {
 public:
 	SLATE_BEGIN_ARGS(SGraphPinVector) {}
 	SLATE_END_ARGS()
 
-	void Construct(const FArguments& InArgs, UEdGraphPin* InGraphPinObj);
+	void Construct(const FArguments& InArgs, UEdGraphPin* InGraphPinObj)
+	{
+		SGraphPin::Construct(SGraphPin::FArguments(), InGraphPinObj);
+	}
 
 protected:
 
@@ -22,7 +27,21 @@ protected:
 	 *
 	 *	@return Reference to the newly created widget object
 	 */
-	virtual TSharedRef<SWidget>	GetDefaultValueWidget() override;
+	virtual TSharedRef<SWidget>	GetDefaultValueWidget() override
+	{
+		UScriptStruct* RotatorStruct = TBaseStructure<FRotator>::Get();
+		bIsRotator = (GraphPinObj->PinType.PinSubCategoryObject == RotatorStruct) ? true : false;
+	
+		return	SNew( SVectorTextBox<NumericType>, bIsRotator )
+				.VisibleText_0(this, &SGraphPinVector::GetCurrentValue_0)
+				.VisibleText_1(this, &SGraphPinVector::GetCurrentValue_1)
+				.VisibleText_2(this, &SGraphPinVector::GetCurrentValue_2)
+				.Visibility( this, &SGraphPin::GetDefaultValueVisibility)
+				.IsEnabled(this, &SGraphPin::GetDefaultValueIsEditable)
+				.OnNumericCommitted_Box_0(this, &SGraphPinVector::OnChangedValueTextBox_0)
+				.OnNumericCommitted_Box_1(this, &SGraphPinVector::OnChangedValueTextBox_1)
+				.OnNumericCommitted_Box_2(this, &SGraphPinVector::OnChangedValueTextBox_2);
+	}
 
 private:
 
@@ -34,26 +53,40 @@ private:
 		TextBox_2,
 	};
 
+	// Rotator is represented as X->Roll, Y->Pitch, Z->Yaw
+
 	/*
 	 *	Function to get current value in text box 0
 	 *
 	 *	@return current string value
 	 */
-	FString GetCurrentValue_0() const;
+	FString GetCurrentValue_0() const
+	{
+		// Text box 0: Rotator->Roll, Vector->X
+		return GetValue(bIsRotator ? TextBox_2 : TextBox_0);
+	}
 
 	/*
 	 *	Function to get current value in text box 1
 	 *
 	 *	@return current string value
 	 */
-	FString GetCurrentValue_1() const;
+	FString GetCurrentValue_1() const
+	{
+		// Text box 1: Rotator->Pitch, Vector->Y
+		return GetValue(bIsRotator ? TextBox_0 : TextBox_1);
+	}
 
 	/*
 	 *	Function to get current value in text box 2
 	 *
 	 *	@return current string value
 	 */
-	FString GetCurrentValue_2() const;
+	FString GetCurrentValue_2() const
+	{
+		// Text box 2: Rotator->Yaw, Vector->Z
+		return GetValue(bIsRotator ? TextBox_1 : TextBox_2);
+	}
 
 	/*
 	 *	Function to getch current value based on text box index value
@@ -62,30 +95,135 @@ private:
 	 *
 	 *	@return current string value
 	 */
-	FString GetValue( ETextBoxIndex Index ) const;
+	FString GetValue( ETextBoxIndex Index ) const
+	{
+		FString DefaultString = GraphPinObj->GetDefaultAsString();
+		TArray<FString> ResultString;
+
+		// Parse string to split its contents separated by ','
+		DefaultString.TrimStartInline();
+		DefaultString.TrimEndInline();
+		DefaultString.ParseIntoArray(ResultString, TEXT(","), true);
+
+		if (Index < ResultString.Num())
+		{
+			return ResultString[Index];
+		}
+		else
+		{
+			return FString(TEXT("0"));
+		}
+	}
 
 	/*
 	 *	Function to store value when text box 0 value in modified
 	 *
-	 *	@param 0: Updated Float Value
+	 *	@param 0: Updated numeric Value
 	 */
-	void OnChangedValueTextBox_0(float NewValue, ETextCommit::Type CommitInfo);
+	void OnChangedValueTextBox_0(NumericType NewValue, ETextCommit::Type CommitInfo)
+	{
+		if (GraphPinObj->IsPendingKill())
+		{
+			return;
+		}
+
+		const FString ValueStr = FString::Printf(TEXT("%f"), NewValue);
+
+		FString DefaultValue;
+		if (bIsRotator)
+		{
+			// Update Roll value
+			DefaultValue = GetValue(TextBox_0) + FString(TEXT(",")) + GetValue(TextBox_1) + FString(TEXT(",")) + ValueStr;
+		}
+		else
+		{
+			// Update X value
+			DefaultValue = ValueStr + FString(TEXT(",")) + GetValue(TextBox_1) + FString(TEXT(",")) + GetValue(TextBox_2);
+		}
+
+		if (GraphPinObj->GetDefaultAsString() != DefaultValue)
+		{
+			const FScopedTransaction Transaction(NSLOCTEXT("GraphEditor", "ChangeVectorPinValue", "Change Vector Pin Value"));
+			GraphPinObj->Modify();
+
+			// Set new default value
+			GraphPinObj->GetSchema()->TrySetDefaultValue(*GraphPinObj, DefaultValue);
+		}
+	}
 
 	/*
 	 *	Function to store value when text box 1 value in modified
 	 *
-	 *	@param 0: Updated Float Value
+	 *	@param 0: Updated numeric Value
 	 */
-	void OnChangedValueTextBox_1(float NewValue, ETextCommit::Type CommitInfo);
+	void OnChangedValueTextBox_1(NumericType NewValue, ETextCommit::Type CommitInfo)
+	{
+		if (GraphPinObj->IsPendingKill())
+		{
+			return;
+		}
+
+		const FString ValueStr = FString::Printf(TEXT("%f"), NewValue);
+
+		FString DefaultValue;
+		if (bIsRotator)
+		{
+			// Update Pitch value
+			DefaultValue = ValueStr + FString(TEXT(",")) + GetValue(TextBox_1) + FString(TEXT(",")) + GetValue(TextBox_2);
+		}
+		else
+		{
+			// Update Y value
+			DefaultValue = GetValue(TextBox_0) + FString(TEXT(",")) + ValueStr + FString(TEXT(",")) + GetValue(TextBox_2);
+		}
+
+		if (GraphPinObj->GetDefaultAsString() != DefaultValue)
+		{
+			const FScopedTransaction Transaction(NSLOCTEXT("GraphEditor", "ChangeVectorPinValue", "Change Vector Pin Value"));
+			GraphPinObj->Modify();
+
+			// Set new default value
+			GraphPinObj->GetSchema()->TrySetDefaultValue(*GraphPinObj, DefaultValue);
+		}
+	}
 
 	/*
 	 *	Function to store value when text box 2 value in modified
 	 *
-	 *	@param 0: Updated Float Value
+	 *	@param 0: Updated numeric Value
 	 */
-	void OnChangedValueTextBox_2(float NewValue, ETextCommit::Type CommitInfo);
+	void OnChangedValueTextBox_2(NumericType NewValue, ETextCommit::Type CommitInfo)
+	{
+		if (GraphPinObj->IsPendingKill())
+		{
+			return;
+		}
+
+		const FString ValueStr = FString::Printf(TEXT("%f"), NewValue);
+
+		FString DefaultValue;
+		if (bIsRotator)
+		{
+			// Update Yaw value
+			DefaultValue = GetValue(TextBox_0) + FString(TEXT(",")) + ValueStr + FString(TEXT(",")) + GetValue(TextBox_2);
+		}
+		else
+		{
+			// Update Z value
+			DefaultValue = GetValue(TextBox_0) + FString(TEXT(",")) + GetValue(TextBox_1) + FString(TEXT(",")) + ValueStr;
+		}
+
+		if (GraphPinObj->GetDefaultAsString() != DefaultValue)
+		{
+			const FScopedTransaction Transaction(NSLOCTEXT("GraphEditor", "ChangeVectorPinValue", "Change Vector Pin Value"));
+			GraphPinObj->Modify();
+
+			// Set new default value
+			GraphPinObj->GetSchema()->TrySetDefaultValue(*GraphPinObj, DefaultValue);
+		}
+	}
 
 private:
-	/** Flag is true if the widget is used to represent a rotator; false otherwise */
+	// Flag is true if the widget is used to represent a rotator; false otherwise
 	bool bIsRotator;
 };
diff --git a/Engine/Source/Editor/GraphEditor/Public/KismetPins/SGraphPinVector2D.h b/Engine/Source/Editor/GraphEditor/Public/KismetPins/SGraphPinVector2D.h
--- a/Engine/Source/Editor/GraphEditor/Public/KismetPins/SGraphPinVector2D.h
+++ b/Engine/Source/Editor/GraphEditor/Public/KismetPins/SGraphPinVector2D.h
@@ -6,14 +6,19 @@
 #include "Widgets/DeclarativeSyntaxSupport.h"
 #include "Widgets/SWidget.h"
 #include "SGraphPin.h"
+#include "SVector2DTextBox.h"
 
-class GRAPHEDITOR_API SGraphPinVector2D : public SGraphPin
+template <typename NumericType>
+class SGraphPinVector2D : public SGraphPin
 {
 public:
 	SLATE_BEGIN_ARGS(SGraphPinVector2D) {}
 	SLATE_END_ARGS()
 
-	void Construct(const FArguments& InArgs, UEdGraphPin* InGraphPinObj);
+	void Construct(const FArguments& InArgs, UEdGraphPin* InGraphPinObj)
+	{
+		SGraphPin::Construct(SGraphPin::FArguments(), InGraphPinObj);
+	}
 
 protected:
 
@@ -22,7 +27,17 @@ protected:
 	 *
 	 *	@return Reference to the newly created widget object
 	 */
-	virtual TSharedRef<SWidget>	GetDefaultValueWidget() override;
+	virtual TSharedRef<SWidget>	GetDefaultValueWidget() override
+	{
+		// Create widget
+		return SNew(SVector2DTextBox<NumericType>)
+			.VisibleText_X(this, &SGraphPinVector2D::GetCurrentValue_X)
+			.VisibleText_Y(this, &SGraphPinVector2D::GetCurrentValue_Y)
+			.Visibility(this, &SGraphPin::GetDefaultValueVisibility)
+			.IsEnabled(this, &SGraphPin::GetDefaultValueIsEditable)
+			.OnNumericCommitted_Box_X(this, &SGraphPinVector2D::OnChangedValueTextBox_X)
+			.OnNumericCommitted_Box_Y(this, &SGraphPinVector2D::OnChangedValueTextBox_Y);
+	}
 
 private:
 
@@ -38,14 +53,20 @@ private:
 	 *
 	 *	@return current string value
 	 */
-	FString GetCurrentValue_X() const;
+	FString GetCurrentValue_X() const
+	{
+		return GetValue(TextBox_X);
+	}
 
 	/*
 	 *	Function to get current value in text box 1
 	 *
 	 *	@return current string value
 	 */
-	FString GetCurrentValue_Y() const;
+	FString GetCurrentValue_Y() const
+	{
+		return GetValue(TextBox_Y);
+	}
 
 	/*
 	 *	Function to getch current value based on text box index value
@@ -54,19 +75,72 @@ private:
 	 *
 	 *	@return current string value
 	 */
-	FString GetValue( ETextBoxIndex Index ) const;
+	FString GetValue(ETextBoxIndex Index) const
+	{
+		FString DefaultString = GraphPinObj->GetDefaultAsString();
+		TArray<FString> ResultString;
+
+		FVector2D Value;
+		Value.InitFromString(DefaultString);
+
+		if (Index == TextBox_X)
+		{
+			return FString::Printf(TEXT("%f"), Value.X);
+		}
+		else
+		{
+			return FString::Printf(TEXT("%f"), Value.Y);
+		}
+	}
+
+	FString MakeVector2DString(const FString& X, const FString& Y)
+	{
+		return FString(TEXT("(X=")) + X + FString(TEXT(",Y=")) + Y + FString(TEXT(")"));
+	}
 
 	/*
 	 *	Function to store value when text box 0 value in modified
 	 *
-	 *	@param 0: Updated Float Value
+	 *	@param 0: Updated numeric value
 	 */
-	void OnChangedValueTextBox_X(float NewValue, ETextCommit::Type CommitInfo);
+	void OnChangedValueTextBox_X(NumericType NewValue, ETextCommit::Type CommitInfo)
+	{
+		if (GraphPinObj->IsPendingKill())
+		{
+			return;
+		}
+
+		const FString ValueStr = FString::Printf(TEXT("%f"), NewValue);
+		const FString Vector2DString = MakeVector2DString(ValueStr, GetValue(TextBox_Y));
+
+		if (GraphPinObj->GetDefaultAsString() != Vector2DString)
+		{
+			const FScopedTransaction Transaction(NSLOCTEXT("GraphEditor", "ChangeVectorPinValue", "Change Vector Pin Value"));
+			GraphPinObj->Modify();
+			GraphPinObj->GetSchema()->TrySetDefaultValue(*GraphPinObj, Vector2DString);
+		}
+	}
 
 	/*
 	 *	Function to store value when text box 1 value in modified
 	 *
-	 *	@param 0: Updated Float Value
+	 *	@param 0: Updated numeric value
 	 */
-	void OnChangedValueTextBox_Y(float NewValue, ETextCommit::Type CommitInfo);
+	void OnChangedValueTextBox_Y(NumericType NewValue, ETextCommit::Type CommitInfo)
+	{
+		if (GraphPinObj->IsPendingKill())
+		{
+			return;
+		}
+
+		const FString ValueStr = FString::Printf(TEXT("%f"), NewValue);
+		const FString Vector2DString = MakeVector2DString(GetValue(TextBox_X), ValueStr);
+
+		if (GraphPinObj->GetDefaultAsString() != Vector2DString)
+		{
+			const FScopedTransaction Transaction(NSLOCTEXT("GraphEditor", "ChangeVectorPinValue", "Change Vector Pin Value"));
+			GraphPinObj->Modify();
+			GraphPinObj->GetSchema()->TrySetDefaultValue(*GraphPinObj, Vector2DString);
+		}
+	}
 };
diff --git a/Engine/Source/Editor/GraphEditor/Public/KismetPins/SGraphPinVector4.h b/Engine/Source/Editor/GraphEditor/Public/KismetPins/SGraphPinVector4.h
--- a/Engine/Source/Editor/GraphEditor/Public/KismetPins/SGraphPinVector4.h
+++ b/Engine/Source/Editor/GraphEditor/Public/KismetPins/SGraphPinVector4.h
@@ -6,23 +6,43 @@
 #include "Widgets/DeclarativeSyntaxSupport.h"
 #include "Widgets/SWidget.h"
 #include "SGraphPin.h"
+#include "SVector4TextBox.h"
 
-class GRAPHEDITOR_API SGraphPinVector4 : public SGraphPin
+template <typename NumericType>
+class SGraphPinVector4 : public SGraphPin
 {
 public:
 	SLATE_BEGIN_ARGS(SGraphPinVector4) {}
 	SLATE_END_ARGS()
 
-		void Construct(const FArguments& InArgs, UEdGraphPin* InGraphPinObj);
+	void Construct(const FArguments& InArgs, UEdGraphPin* InGraphPinObj)
+	{
+		SGraphPin::Construct(SGraphPin::FArguments(), InGraphPinObj);
+	}
 
 protected:
 
 	/**
-	*	Function to create class specific widget.
-	*
-	*	@return Reference to the newly created widget object
-	*/
-	virtual TSharedRef<SWidget>	GetDefaultValueWidget() override;
+	 *	Function to create class specific widget.
+	 *
+	 *	@return Reference to the newly created widget object
+	 */
+	virtual TSharedRef<SWidget>	GetDefaultValueWidget() override
+	{
+		UScriptStruct* RotatorStruct = TBaseStructure<FRotator>::Get();
+
+		return SNew(SVector4TextBox<NumericType>)
+			.VisibleText_0(this, &SGraphPinVector4::GetCurrentValue_0)
+			.VisibleText_1(this, &SGraphPinVector4::GetCurrentValue_1)
+			.VisibleText_2(this, &SGraphPinVector4::GetCurrentValue_2)
+			.VisibleText_3(this, &SGraphPinVector4::GetCurrentValue_3)
+			.Visibility(this, &SGraphPin::GetDefaultValueVisibility)
+			.IsEnabled(this, &SGraphPin::GetDefaultValueIsEditable)
+			.OnNumericCommitted_Box_0(this, &SGraphPinVector4::OnChangedValueTextBox_0)
+			.OnNumericCommitted_Box_1(this, &SGraphPinVector4::OnChangedValueTextBox_1)
+			.OnNumericCommitted_Box_2(this, &SGraphPinVector4::OnChangedValueTextBox_2)
+			.OnNumericCommitted_Box_3(this, &SGraphPinVector4::OnChangedValueTextBox_3);
+	}
 
 private:
 
@@ -35,10 +55,31 @@ private:
 		TextBox_3,
 	};
 
-	FString GetCurrentValue_0() const;
-	FString GetCurrentValue_1() const;
-	FString GetCurrentValue_2() const;
-	FString GetCurrentValue_3() const;
+	// Rotator is represented as X->Roll, Y->Pitch, Z->Yaw
+
+	FString GetCurrentValue_0() const
+	{
+		// Text box 0: Rotator->Roll, Vector->X
+		return GetValue(TextBox_0);
+	}
+
+	FString GetCurrentValue_1() const
+	{
+		// Text box 1: Rotator->Pitch, Vector->Y
+		return GetValue(TextBox_1);
+	}
+
+	FString GetCurrentValue_2() const
+	{
+		// Text box 2: Rotator->Yaw, Vector->Z
+		return GetValue(TextBox_2);
+	}
+
+	FString GetCurrentValue_3() const
+	{
+		// Text box 3: Vector->W
+		return GetValue(TextBox_3);
+	}
 
 	/*
 	*	Function to getch current value based on text box index value
@@ -47,19 +88,102 @@ private:
 	*
 	*	@return current string value
 	*/
-	FString GetValue(ETextBoxIndex Index) const;
+	FString GetValue(ETextBoxIndex Index) const
+	{
+		FString DefaultString = GraphPinObj->GetDefaultAsString();
+		TArray<FString> ResultString;
+
+		// Parse string to split its contents separated by ','
+		DefaultString.TrimStartAndEndInline();
+		DefaultString.ParseIntoArray(ResultString, TEXT(","), true);
+
+		if (Index < ResultString.Num())
+		{
+			return ResultString[Index];
+		}
+		else
+		{
+			return FString(TEXT("0"));
+		}
+	}
 
 	/*
 	*	Function to store value when text box value in modified
 	*
-	*	@param 0: Updated Float Value
+	*	@param 0: Updated numeric value
 	*/
-	void OnChangedValueTextBox_0(float NewValue, ETextCommit::Type CommitInfo);
-	void OnChangedValueTextBox_1(float NewValue, ETextCommit::Type CommitInfo);
-	void OnChangedValueTextBox_2(float NewValue, ETextCommit::Type CommitInfo);
-	void OnChangedValueTextBox_3(float NewValue, ETextCommit::Type CommitInfo);
+	void OnChangedValueTextBox_0(NumericType NewValue, ETextCommit::Type CommitInfo)
+	{
+		const FString ValueStr = FString::Printf(TEXT("%f"), NewValue);
+
+		FString DefaultValue;
+		// Update X value
+		DefaultValue = ValueStr + FString(TEXT(",")) + GetValue(TextBox_1) + FString(TEXT(",")) + GetValue(TextBox_2) + FString(TEXT(",")) + GetValue(TextBox_3);
+
+		if (GraphPinObj->GetDefaultAsString() != DefaultValue)
+		{
+			const FScopedTransaction Transaction(NSLOCTEXT("GraphEditor", "ChangeVector4PinValue", "Change Vector4 Pin Value"));
+			GraphPinObj->Modify();
+
+			// Set new default value
+			GraphPinObj->GetSchema()->TrySetDefaultValue(*GraphPinObj, DefaultValue);
+		}
+	}
+
+	void OnChangedValueTextBox_1(NumericType NewValue, ETextCommit::Type CommitInfo)
+	{
+		const FString ValueStr = FString::Printf(TEXT("%f"), NewValue);
+
+		FString DefaultValue;
+		// Update Y value
+		DefaultValue = GetValue(TextBox_0) + FString(TEXT(",")) + ValueStr + FString(TEXT(",")) + GetValue(TextBox_2) + FString(TEXT(",")) + GetValue(TextBox_3);
+
+		if (GraphPinObj->GetDefaultAsString() != DefaultValue)
+		{
+			const FScopedTransaction Transaction(NSLOCTEXT("GraphEditor", "ChangeVector4PinValue", "Change Vector4 Pin Value"));
+			GraphPinObj->Modify();
+
+			// Set new default value
+			GraphPinObj->GetSchema()->TrySetDefaultValue(*GraphPinObj, DefaultValue);
+		}
+	}
+
+	void OnChangedValueTextBox_2(NumericType NewValue, ETextCommit::Type CommitInfo)
+	{
+		const FString ValueStr = FString::Printf(TEXT("%f"), NewValue);
+
+		FString DefaultValue;
+		// Update Z value
+		DefaultValue = GetValue(TextBox_0) + FString(TEXT(",")) + GetValue(TextBox_1) + FString(TEXT(",")) + ValueStr + FString(TEXT(",")) + GetValue(TextBox_3);
+
+		if (GraphPinObj->GetDefaultAsString() != DefaultValue)
+		{
+			const FScopedTransaction Transaction(NSLOCTEXT("GraphEditor", "ChangeVector4PinValue", "Change Vector4 Pin Value"));
+			GraphPinObj->Modify();
+
+			// Set new default value
+			GraphPinObj->GetSchema()->TrySetDefaultValue(*GraphPinObj, DefaultValue);
+		}
+	}
+
+	void OnChangedValueTextBox_3(NumericType NewValue, ETextCommit::Type CommitInfo)
+	{
+		const FString ValueStr = FString::Printf(TEXT("%f"), NewValue);
+
+		FString DefaultValue;
+		// Update W value
+		DefaultValue = GetValue(TextBox_0) + FString(TEXT(",")) + GetValue(TextBox_1) + FString(TEXT(",")) + GetValue(TextBox_2) + FString(TEXT(",")) + ValueStr;
+
+		if (GraphPinObj->GetDefaultAsString() != DefaultValue)
+		{
+			const FScopedTransaction Transaction(NSLOCTEXT("GraphEditor", "ChangeVector4PinValue", "Change Vector4 Pin Value"));
+			GraphPinObj->Modify();
+
+			// Set new default value
+			GraphPinObj->GetSchema()->TrySetDefaultValue(*GraphPinObj, DefaultValue);
+		}
+	}
 
-	static FVector4 ParseValue(FString Value);
 private:
 	/** Flag is true if the widget is used to represent a rotator; false otherwise */
 	bool bIsRotator;
diff --git a/Engine/Source/Editor/GraphEditor/Public/KismetPins/SVector2DTextBox.h b/Engine/Source/Editor/GraphEditor/Public/KismetPins/SVector2DTextBox.h
new file mode 100644
--- /dev/null
+++ b/Engine/Source/Editor/GraphEditor/Public/KismetPins/SVector2DTextBox.h
@@ -0,0 +1,123 @@
+// Copyright Epic Games, Inc. All Rights Reserved.
+
+#pragma once
+
+#include "CoreMinimal.h"
+#include "Widgets/SCompoundWidget.h"
+
+#define LOCTEXT_NAMESPACE "Vector2DTextBox"
+
+// Class implementation to create 2 editable text boxes to represent vector2D graph pin
+template <typename NumericType>
+class SVector2DTextBox : public SCompoundWidget
+{
+public:
+
+	// Notification for numeric value committed
+	DECLARE_DELEGATE_TwoParams(FOnNumericValueCommitted, NumericType, ETextCommit::Type);
+
+	SLATE_BEGIN_ARGS( SVector2DTextBox ){}	
+		SLATE_ATTRIBUTE( FString, VisibleText_X )
+		SLATE_ATTRIBUTE( FString, VisibleText_Y )
+		SLATE_EVENT( FOnNumericValueCommitted, OnNumericCommitted_Box_X )
+		SLATE_EVENT( FOnNumericValueCommitted, OnNumericCommitted_Box_Y )
+	SLATE_END_ARGS()
+
+	// Construct editable text boxes with the appropriate getter & setter functions along with tool tip text
+	void Construct( const FArguments& InArgs )
+	{
+		VisibleText_X = InArgs._VisibleText_X;
+		VisibleText_Y = InArgs._VisibleText_Y;
+		const FLinearColor LabelClr = FLinearColor( 1.f, 1.f, 1.f, 0.4f );
+
+		this->ChildSlot
+		[
+			SNew(SVerticalBox)
+			+SVerticalBox::Slot()
+			.AutoHeight()
+			.Padding(0)
+			[
+				SNew(SHorizontalBox)
+				+SHorizontalBox::Slot()
+				.AutoWidth().Padding(2) .HAlign(HAlign_Fill)
+				[
+					// Create Text box 0 
+					SNew( SNumericEntryBox<NumericType> )
+					.LabelVAlign(VAlign_Center)
+					.Label()
+					[
+						SNew( STextBlock )
+						.Font( FAppStyle::GetFontStyle( "Graph.VectorEditableTextBox" ) )
+						.Text( LOCTEXT("VectorNodeXAxisValueLabel", "X") )
+						.ColorAndOpacity( LabelClr )
+					]
+					.Value( this, &SVector2DTextBox::GetTypeInValue_X )
+					.OnValueCommitted( InArgs._OnNumericCommitted_Box_X )
+					.Font( FAppStyle::GetFontStyle( "Graph.VectorEditableTextBox" ) )
+					.UndeterminedString( LOCTEXT("MultipleValues", "Multiple Values") )
+					.ToolTipText(LOCTEXT("VectorNodeXAxisValueLabel_ToolTip", "X value") )
+					.EditableTextBoxStyle( &FAppStyle::GetWidgetStyle<FEditableTextBoxStyle>( "Graph.VectorEditableTextBox" ))
+					.BorderForegroundColor( FLinearColor::White )
+					.BorderBackgroundColor( FLinearColor::White )
+				]
+				+SHorizontalBox::Slot()
+				.AutoWidth().Padding(2) .HAlign(HAlign_Fill)
+				[
+					// Create Text box 1
+					SNew( SNumericEntryBox<NumericType> )
+					.LabelVAlign(VAlign_Center)
+					.Label()
+					[
+						SNew( STextBlock )
+						.Font( FAppStyle::GetFontStyle( "Graph.VectorEditableTextBox" ) )
+						.Text( LOCTEXT("VectorNodeYAxisValueLabel", "Y") )
+						.ColorAndOpacity( LabelClr )
+					]
+					.Value( this, &SVector2DTextBox::GetTypeInValue_Y )
+					.OnValueCommitted( InArgs._OnNumericCommitted_Box_Y )
+					.Font( FAppStyle::GetFontStyle( "Graph.VectorEditableTextBox" ) )
+					.UndeterminedString( LOCTEXT("MultipleValues", "Multiple Values") )
+					.ToolTipText(LOCTEXT("VectorNodeYAxisValueLabel_ToolTip", "Y value"))
+					.EditableTextBoxStyle( &FAppStyle::GetWidgetStyle<FEditableTextBoxStyle>( "Graph.VectorEditableTextBox" ))
+					.BorderForegroundColor( FLinearColor::White )
+					.BorderBackgroundColor( FLinearColor::White )
+				]
+			]
+		];
+	}
+
+private:
+
+	NumericType GetValueType(const FString& InString) const
+	{
+		static_assert(std::is_floating_point_v<NumericType>);
+
+		if constexpr (std::is_same_v<float, NumericType>)
+		{
+			return FCString::Atof(*InString);
+		}
+		else if constexpr (std::is_same_v<double, NumericType>)
+		{
+			return FCString::Atod(*InString);
+		}
+
+		return NumericType{};
+	}
+
+	// Get value for X text box
+	TOptional<NumericType> GetTypeInValue_X() const
+	{
+		return GetValueType(VisibleText_X.Get());
+	}
+
+	// Get value for Y text box
+	TOptional<NumericType> GetTypeInValue_Y() const
+	{
+		return GetValueType(VisibleText_Y.Get());
+	}
+
+	TAttribute<FString> VisibleText_X;
+	TAttribute<FString> VisibleText_Y;
+};
+
+#undef LOCTEXT_NAMESPACE
diff --git a/Engine/Source/Editor/GraphEditor/Public/KismetPins/SVector4TextBox.h b/Engine/Source/Editor/GraphEditor/Public/KismetPins/SVector4TextBox.h
new file mode 100644
--- /dev/null
+++ b/Engine/Source/Editor/GraphEditor/Public/KismetPins/SVector4TextBox.h
@@ -0,0 +1,184 @@
+// Copyright Epic Games, Inc. All Rights Reserved.
+
+#pragma once
+
+#include "CoreMinimal.h"
+#include "Widgets/SCompoundWidget.h"
+
+#define LOCTEXT_NAMESPACE "VectorTextBox"
+
+// Class implementation to create 3 editable text boxes to represent vector/rotator graph pin
+template <typename NumericType>
+class SVector4TextBox : public SCompoundWidget
+{
+public:
+
+	// Notification for numeric value committed
+	DECLARE_DELEGATE_TwoParams(FOnNumericValueCommitted, NumericType, ETextCommit::Type);
+
+	SLATE_BEGIN_ARGS(SVector4TextBox) {}
+		SLATE_ATTRIBUTE(FString, VisibleText_0)
+		SLATE_ATTRIBUTE(FString, VisibleText_1)
+		SLATE_ATTRIBUTE(FString, VisibleText_2)
+		SLATE_ATTRIBUTE(FString, VisibleText_3)
+		SLATE_EVENT(FOnNumericValueCommitted, OnNumericCommitted_Box_0)
+		SLATE_EVENT(FOnNumericValueCommitted, OnNumericCommitted_Box_1)
+		SLATE_EVENT(FOnNumericValueCommitted, OnNumericCommitted_Box_2)
+		SLATE_EVENT(FOnNumericValueCommitted, OnNumericCommitted_Box_3)
+	SLATE_END_ARGS()
+	
+	// Construct editable text boxes with the appropriate getter & setter functions along with tool tip text
+	void Construct(const FArguments& InArgs)
+	{
+		VisibleText_0 = InArgs._VisibleText_0;
+		VisibleText_1 = InArgs._VisibleText_1;
+		VisibleText_2 = InArgs._VisibleText_2;
+		VisibleText_3 = InArgs._VisibleText_3;
+		const FLinearColor LabelClr = FLinearColor(1.f, 1.f, 1.f, 0.4f);
+
+		this->ChildSlot
+			[
+				SNew(SVerticalBox)
+				+ SVerticalBox::Slot()
+				.AutoHeight()
+				.Padding(0)
+				[
+					SNew(SHorizontalBox)
+					+ SHorizontalBox::Slot()
+					.AutoWidth().Padding(2).HAlign(HAlign_Fill)
+					[
+						// Create Text box 0 
+						SNew(SNumericEntryBox<NumericType>)
+						.LabelVAlign(VAlign_Center)
+						.Label()
+						[
+							SNew(STextBlock)
+							.Font(FAppStyle::GetFontStyle("Graph.VectorEditableTextBox"))
+							.Text(LOCTEXT("VectorNodeXAxisValueLabel", "X"))
+							.ColorAndOpacity(LabelClr)
+						]
+						.Value(this, &SVector4TextBox::GetTypeInValue_0)
+							.OnValueCommitted(InArgs._OnNumericCommitted_Box_0)
+							.Font(FAppStyle::GetFontStyle("Graph.VectorEditableTextBox"))
+							.UndeterminedString(LOCTEXT("MultipleValues", "Multiple Values"))
+							.ToolTipText(LOCTEXT("VectorNodeXAxisValueLabel_ToolTip", "X value"))
+							.EditableTextBoxStyle(&FAppStyle::GetWidgetStyle<FEditableTextBoxStyle>("Graph.VectorEditableTextBox"))
+							.BorderForegroundColor(FLinearColor::White)
+							.BorderBackgroundColor(FLinearColor::White)
+					]
+					+ SHorizontalBox::Slot()
+						.AutoWidth().Padding(2).HAlign(HAlign_Fill)
+						[
+							// Create Text box 1
+							SNew(SNumericEntryBox<NumericType>)
+							.LabelVAlign(VAlign_Center)
+							.Label()
+							[
+								SNew(STextBlock)
+								.Font(FAppStyle::GetFontStyle("Graph.VectorEditableTextBox"))
+								.Text(LOCTEXT("VectorNodeYAxisValueLabel", "Y"))
+								.ColorAndOpacity(LabelClr)
+							]
+							.Value(this, &SVector4TextBox::GetTypeInValue_1)
+								.OnValueCommitted(InArgs._OnNumericCommitted_Box_1)
+								.Font(FAppStyle::GetFontStyle("Graph.VectorEditableTextBox"))
+								.UndeterminedString(LOCTEXT("MultipleValues", "Multiple Values"))
+								.ToolTipText(LOCTEXT("VectorNodeYAxisValueLabel_ToolTip", "Y value"))
+								.EditableTextBoxStyle(&FAppStyle::GetWidgetStyle<FEditableTextBoxStyle>("Graph.VectorEditableTextBox"))
+								.BorderForegroundColor(FLinearColor::White)
+								.BorderBackgroundColor(FLinearColor::White)
+						]
+					+ SHorizontalBox::Slot()
+						.AutoWidth().Padding(2).HAlign(HAlign_Fill)
+						[
+							// Create Text box 2
+							SNew(SNumericEntryBox<NumericType>)
+							.LabelVAlign(VAlign_Center)
+							.Label()
+							[
+								SNew(STextBlock)
+								.Font(FAppStyle::GetFontStyle("Graph.VectorEditableTextBox"))
+								.Text(LOCTEXT("VectorNodeZAxisValueLabel", "Z"))
+								.ColorAndOpacity(LabelClr)
+							]
+							.Value(this, &SVector4TextBox::GetTypeInValue_2)
+								.OnValueCommitted(InArgs._OnNumericCommitted_Box_2)
+								.Font(FAppStyle::GetFontStyle("Graph.VectorEditableTextBox"))
+								.UndeterminedString(LOCTEXT("MultipleValues", "Multiple Values"))
+								.ToolTipText(LOCTEXT("VectorNodeZAxisValueLabel_ToolTip", "Z value"))
+								.EditableTextBoxStyle(&FAppStyle::GetWidgetStyle<FEditableTextBoxStyle>("Graph.VectorEditableTextBox"))
+								.BorderForegroundColor(FLinearColor::White)
+								.BorderBackgroundColor(FLinearColor::White)
+						]
+					+ SHorizontalBox::Slot()
+						.AutoWidth().Padding(2).HAlign(HAlign_Fill)
+						[
+							// Create Text box 3
+							SNew(SNumericEntryBox<NumericType>)
+							.LabelVAlign(VAlign_Center)
+							.Label()
+							[
+								SNew(STextBlock)
+								.Font(FAppStyle::GetFontStyle("Graph.VectorEditableTextBox"))
+								.Text(LOCTEXT("VectorNodeWAxisValueLabel", "W"))
+								.ColorAndOpacity(LabelClr)
+							]
+							.Value(this, &SVector4TextBox::GetTypeInValue_3)
+								.OnValueCommitted(InArgs._OnNumericCommitted_Box_3)
+								.Font(FAppStyle::GetFontStyle("Graph.VectorEditableTextBox"))
+								.UndeterminedString(LOCTEXT("MultipleValues", "Multiple Values"))
+								.ToolTipText(LOCTEXT("VectorNodeWAxisValueLabel_ToolTip", "W value"))
+								.EditableTextBoxStyle(&FAppStyle::GetWidgetStyle<FEditableTextBoxStyle>("Graph.VectorEditableTextBox"))
+								.BorderForegroundColor(FLinearColor::White)
+								.BorderBackgroundColor(FLinearColor::White)
+						]
+
+				]
+			];
+		}
+
+private:
+
+	NumericType GetValueType(const FString& InString) const
+	{
+		static_assert(std::is_floating_point_v<NumericType>);
+
+		if constexpr (std::is_same_v<float, NumericType>)
+		{
+			return FCString::Atof(*InString);
+		}
+		else if constexpr (std::is_same_v<double, NumericType>)
+		{
+			return FCString::Atod(*InString);
+		}
+
+		return NumericType{};
+	}
+
+	TOptional<NumericType> GetTypeInValue_0() const
+	{
+		return GetValueType(VisibleText_0.Get());
+	}
+
+	TOptional<NumericType> GetTypeInValue_1() const
+	{
+		return GetValueType(VisibleText_1.Get());
+	}
+
+	TOptional<NumericType> GetTypeInValue_2() const
+	{
+		return GetValueType(VisibleText_2.Get());
+	}
+
+	TOptional<NumericType> GetTypeInValue_3() const
+	{
+		return GetValueType(VisibleText_3.Get());
+	}
+
+	TAttribute<FString> VisibleText_0;
+	TAttribute<FString> VisibleText_1;
+	TAttribute<FString> VisibleText_2;
+	TAttribute<FString> VisibleText_3;
+};
+
+#undef LOCTEXT_NAMESPACE
diff --git a/Engine/Source/Editor/GraphEditor/Public/KismetPins/SVectorTextBox.h b/Engine/Source/Editor/GraphEditor/Public/KismetPins/SVectorTextBox.h
new file mode 100644
--- /dev/null
+++ b/Engine/Source/Editor/GraphEditor/Public/KismetPins/SVectorTextBox.h
@@ -0,0 +1,160 @@
+// Copyright Epic Games, Inc. All Rights Reserved.
+
+#pragma once
+
+#include "CoreMinimal.h"
+#include "Widgets/SCompoundWidget.h"
+
+#define LOCTEXT_NAMESPACE "VectorTextBox"
+
+// Class implementation to create 3 editable text boxes to represent vector/rotator graph pin
+template <typename NumericType>
+class SVectorTextBox : public SCompoundWidget
+{
+public:
+
+	// Notification for numeric value committed
+	DECLARE_DELEGATE_TwoParams(FOnNumericValueCommitted, NumericType, ETextCommit::Type);
+
+	SLATE_BEGIN_ARGS( SVectorTextBox ) {}	
+		SLATE_ATTRIBUTE( FString, VisibleText_0 )
+		SLATE_ATTRIBUTE( FString, VisibleText_1 )
+		SLATE_ATTRIBUTE( FString, VisibleText_2 )
+		SLATE_EVENT( FOnNumericValueCommitted, OnNumericCommitted_Box_0 )
+		SLATE_EVENT( FOnNumericValueCommitted, OnNumericCommitted_Box_1 )
+		SLATE_EVENT( FOnNumericValueCommitted, OnNumericCommitted_Box_2 )
+	SLATE_END_ARGS()
+
+	// Construct editable text boxes with the appropriate getter & setter functions along with tool tip text
+	void Construct( const FArguments& InArgs, const bool bInIsRotator )
+	{
+		bIsRotator = bInIsRotator;
+		const bool bUseRPY = false;
+
+		VisibleText_0 = InArgs._VisibleText_0;
+		VisibleText_1 = InArgs._VisibleText_1;
+		VisibleText_2 = InArgs._VisibleText_2;
+		const FLinearColor LabelClr = FLinearColor( 1.f, 1.f, 1.f, 0.4f );
+
+		this->ChildSlot
+		[
+			SNew(SVerticalBox)
+			+SVerticalBox::Slot()
+			.AutoHeight() 
+			.Padding(0)
+			[
+				SNew(SHorizontalBox)
+				+SHorizontalBox::Slot()
+				.AutoWidth().Padding(2) .HAlign(HAlign_Fill)
+				[
+					// Create Text box 0 
+					SNew( SNumericEntryBox<NumericType> )
+					.LabelVAlign(VAlign_Center)
+					.Label()
+					[
+						SNew( STextBlock )
+						.Font( FAppStyle::GetFontStyle( "Graph.VectorEditableTextBox" ) )
+						.Text( bIsRotator && bUseRPY ? LOCTEXT("VectorNodeRollValueLabel", "R") : LOCTEXT("VectorNodeXAxisValueLabel", "X") )
+						.ColorAndOpacity( LabelClr )
+					]
+					.Value( this, &SVectorTextBox::GetTypeInValue_0 )
+					.OnValueCommitted( InArgs._OnNumericCommitted_Box_0 )
+					.Font( FAppStyle::GetFontStyle( "Graph.VectorEditableTextBox" ) )
+					.UndeterminedString( LOCTEXT("MultipleValues", "Multiple Values") )
+					.ToolTipText( bIsRotator ? LOCTEXT("VectorNodeRollValueLabel_ToolTip", "Roll value (around X)") : LOCTEXT("VectorNodeXAxisValueLabel_ToolTip", "X value"))
+					.EditableTextBoxStyle( &FAppStyle::GetWidgetStyle<FEditableTextBoxStyle>( "Graph.VectorEditableTextBox" ))
+					.BorderForegroundColor( FLinearColor::White )
+					.BorderBackgroundColor( FLinearColor::White )
+				]
+				+SHorizontalBox::Slot()
+				.AutoWidth().Padding(2) .HAlign(HAlign_Fill)
+				[
+					// Create Text box 1
+					SNew( SNumericEntryBox<NumericType> )
+					.LabelVAlign(VAlign_Center)
+					.Label()
+					[
+						SNew( STextBlock )
+						.Font( FAppStyle::GetFontStyle( "Graph.VectorEditableTextBox" ) )
+						.Text( bIsRotator && bUseRPY ? LOCTEXT("VectorNodePitchValueLabel", "P") : LOCTEXT("VectorNodeYAxisValueLabel", "Y") )
+						.ColorAndOpacity( LabelClr )
+					]
+					.Value( this, &SVectorTextBox::GetTypeInValue_1 )
+					.OnValueCommitted( InArgs._OnNumericCommitted_Box_1 )
+					.Font( FAppStyle::GetFontStyle( "Graph.VectorEditableTextBox" ) )
+					.UndeterminedString( LOCTEXT("MultipleValues", "Multiple Values") )
+					.ToolTipText(bIsRotator ? LOCTEXT("VectorNodePitchValueLabel_ToolTip", "Pitch value (around Y)") : LOCTEXT("VectorNodeYAxisValueLabel_ToolTip", "Y value"))
+					.EditableTextBoxStyle( &FAppStyle::GetWidgetStyle<FEditableTextBoxStyle>( "Graph.VectorEditableTextBox" ))
+					.BorderForegroundColor( FLinearColor::White )
+					.BorderBackgroundColor( FLinearColor::White )
+				]
+				+SHorizontalBox::Slot()
+				.AutoWidth().Padding(2) .HAlign(HAlign_Fill)
+				[
+					// Create Text box 2
+					SNew( SNumericEntryBox<NumericType> )
+					.LabelVAlign(VAlign_Center)
+					.Label()
+					[
+						SNew( STextBlock )
+						.Font( FAppStyle::GetFontStyle( "Graph.VectorEditableTextBox" ) )
+						.Text( bIsRotator && bUseRPY ? LOCTEXT("VectorNodeYawValueLabel", "Y") : LOCTEXT("VectorNodeZAxisValueLabel", "Z") )
+						.ColorAndOpacity( LabelClr )
+					]
+					.Value( this, &SVectorTextBox::GetTypeInValue_2 )
+					.OnValueCommitted( InArgs._OnNumericCommitted_Box_2 )
+					.Font( FAppStyle::GetFontStyle( "Graph.VectorEditableTextBox" ) )
+					.UndeterminedString( LOCTEXT("MultipleValues", "Multiple Values") )
+					.ToolTipText(bIsRotator ? LOCTEXT("VectorNodeYawValueLabel_Tooltip", "Yaw value (around Z)") : LOCTEXT("VectorNodeZAxisValueLabel_ToolTip", "Z value"))
+					.EditableTextBoxStyle( &FAppStyle::GetWidgetStyle<FEditableTextBoxStyle>( "Graph.VectorEditableTextBox" ))
+					.BorderForegroundColor( FLinearColor::White )
+					.BorderBackgroundColor( FLinearColor::White )
+				]
+			]
+		];
+	}
+
+private:
+
+	NumericType GetValueType(const FString& InString) const
+	{
+		static_assert(std::is_floating_point_v<NumericType>);
+
+		if constexpr (std::is_same_v<float, NumericType>)
+		{
+			return FCString::Atof(*InString);
+		}
+		else if constexpr (std::is_same_v<double, NumericType>)
+		{
+			return FCString::Atod(*InString);
+		}
+
+		return NumericType{};
+	}
+
+	// Get value for text box 0
+	TOptional<NumericType> GetTypeInValue_0() const
+	{
+		return GetValueType(VisibleText_0.Get());
+	}
+
+	// Get value for text box 1
+	TOptional<NumericType> GetTypeInValue_1() const
+	{
+		return GetValueType(VisibleText_1.Get());
+	}
+
+	// Get value for text box 2
+	TOptional<NumericType> GetTypeInValue_2() const
+	{
+		return GetValueType(VisibleText_2.Get());
+	}
+
+	TAttribute<FString> VisibleText_0;
+	TAttribute<FString> VisibleText_1;
+	TAttribute<FString> VisibleText_2;
+
+	bool bIsRotator;
+};
+
+#undef LOCTEXT_NAMESPACE
diff --git a/Engine/Source/Editor/MaterialEditor/Private/MaterialEditorGraphPanelPinFactory.cpp b/Engine/Source/Editor/MaterialEditor/Private/MaterialEditorGraphPanelPinFactory.cpp
--- a/Engine/Source/Editor/MaterialEditor/Private/MaterialEditorGraphPanelPinFactory.cpp
+++ b/Engine/Source/Editor/MaterialEditor/Private/MaterialEditorGraphPanelPinFactory.cpp
@@ -34,11 +34,11 @@ TSharedPtr<class SGraphPin> FMaterialEditorGraphPanelPinFactory::CreatePin(class
 			}
 			else if (InPin->PinType.PinSubCategory == MaterialGraphSchema->PSC_RG)
 			{
-				return SNew(SGraphPinVector2D, InPin);
+				return SNew(SGraphPinVector2D<float>, InPin);
 			}
 			else if (InPin->PinType.PinSubCategory == MaterialGraphSchema->PSC_RGB)
 			{
-				return SNew(SGraphPinVector, InPin);
+				return SNew(SGraphPinVector<float>, InPin);
 			}
 			else if (InPin->PinType.PinSubCategory == MaterialGraphSchema->PSC_RGBA)
 			{
@@ -46,7 +46,7 @@ TSharedPtr<class SGraphPin> FMaterialEditorGraphPanelPinFactory::CreatePin(class
 			}
 			else if (InPin->PinType.PinSubCategory == MaterialGraphSchema->PSC_Vector4)
 			{
-				return SNew(SGraphPinVector4, InPin);
+				return SNew(SGraphPinVector4<float>, InPin);
 			}
 			else if (InPin->PinType.PinSubCategory == MaterialGraphSchema->PSC_Int)
 			{
