Build fix - missing headers.

diff --git a/Engine/Source/Editor/GraphEditor/Public/KismetPins/SGraphPinVector.h b/Engine/Source/Editor/GraphEditor/Public/KismetPins/SGraphPinVector.h
--- a/Engine/Source/Editor/GraphEditor/Public/KismetPins/SGraphPinVector.h
+++ b/Engine/Source/Editor/GraphEditor/Public/KismetPins/SGraphPinVector.h
@@ -7,6 +7,7 @@
 #include "Widgets/SWidget.h"
 #include "SGraphPin.h"
 #include "SVectorTextBox.h"
+#include "ScopedTransaction.h"
 
 template <typename NumericType>
 class SGraphPinVector : public SGraphPin
diff --git a/Engine/Source/Editor/GraphEditor/Public/KismetPins/SGraphPinVector2D.h b/Engine/Source/Editor/GraphEditor/Public/KismetPins/SGraphPinVector2D.h
--- a/Engine/Source/Editor/GraphEditor/Public/KismetPins/SGraphPinVector2D.h
+++ b/Engine/Source/Editor/GraphEditor/Public/KismetPins/SGraphPinVector2D.h
@@ -7,6 +7,7 @@
 #include "Widgets/SWidget.h"
 #include "SGraphPin.h"
 #include "SVector2DTextBox.h"
+#include "ScopedTransaction.h"
 
 template <typename NumericType>
 class SGraphPinVector2D : public SGraphPin
diff --git a/Engine/Source/Editor/GraphEditor/Public/KismetPins/SGraphPinVector4.h b/Engine/Source/Editor/GraphEditor/Public/KismetPins/SGraphPinVector4.h
--- a/Engine/Source/Editor/GraphEditor/Public/KismetPins/SGraphPinVector4.h
+++ b/Engine/Source/Editor/GraphEditor/Public/KismetPins/SGraphPinVector4.h
@@ -7,6 +7,7 @@
 #include "Widgets/SWidget.h"
 #include "SGraphPin.h"
 #include "SVector4TextBox.h"
+#include "ScopedTransaction.h"
 
 template <typename NumericType>
 class SGraphPinVector4 : public SGraphPin
