[Bug Fix] Crash when diffing a graph that was deleted

diff --git a/Engine/Source/Editor/GraphEditor/Private/GraphDiffControl.cpp b/Engine/Source/Editor/GraphEditor/Private/GraphDiffControl.cpp
--- a/Engine/Source/Editor/GraphEditor/Private/GraphDiffControl.cpp
+++ b/Engine/Source/Editor/GraphEditor/Private/GraphDiffControl.cpp
@@ -730,8 +730,11 @@ bool FGraphDiffControl::DiffGraphs(UEdGraph* const LhsGraph, UEdGraph* const Rhs
 
 FString FGraphDiffControl::GetGraphPath(UEdGraph* Graph)
 {
-	FString GraphPath;
-	if (UBlueprint* Blueprint = FBlueprintEditorUtils::FindBlueprintForGraph(Graph))
+	if (Graph == nullptr)
+	{
+		return FString();
+	}
+	else if (UBlueprint* Blueprint = FBlueprintEditorUtils::FindBlueprintForGraph(Graph))
 	{
 		return Graph->GetPathName(Blueprint);
 	}
@@ -739,11 +742,8 @@ FString FGraphDiffControl::GetGraphPath(UEdGraph* Graph)
 	{
 		return Graph->GetPathName(Package);
 	}
-	else if (Graph)
-	{
-		return Graph->GetName();
-	}
-	return FString();
+	
+	return Graph->GetName();
 }
 
 #undef LOCTEXT_NAMESPACE
