Make Graph Knot API public so it can be extended

diff --git a/Engine/Source/Editor/GraphEditor/Private/SGraphNodeKnot.cpp b/Engine/Source/Editor/GraphEditor/Private/SGraphNodeKnot.cpp
--- a/Engine/Source/Editor/GraphEditor/Private/SGraphNodeKnot.cpp
+++ b/Engine/Source/Editor/GraphEditor/Private/SGraphNodeKnot.cpp
@@ -163,27 +163,6 @@ void FAmbivalentDirectionDragConnection::ValidateGraphPinList(TArray<UEdGraphPin
 /////////////////////////////////////////////////////
 // SGraphPinKnot
 
-class SGraphPinKnot : public SGraphPin
-{
-public:
-	SLATE_BEGIN_ARGS(SGraphPinKnot) {}
-	SLATE_END_ARGS()
-
-	void Construct(const FArguments& InArgs, UEdGraphPin* InPin);
-
-	// SWidget interface
-	virtual void OnDragEnter(const FGeometry& MyGeometry, const FDragDropEvent& DragDropEvent) override;
-	// End of SWidget interface
-
-protected:
-	// Begin SGraphPin interface
-	virtual TSharedRef<SWidget>	GetDefaultValueWidget() override;
-	virtual TSharedRef<FDragDropOperation> SpawnPinDragEvent(const TSharedRef<SGraphPanel>& InGraphPanel, const TArray< TSharedRef<SGraphPin> >& InStartingPins) override;
-	virtual FReply OnPinMouseDown(const FGeometry& SenderGeometry, const FPointerEvent& MouseEvent) override;
-	virtual FSlateColor GetPinColor() const override;
-	// End SGraphPin interface
-};
-
 void SGraphPinKnot::Construct(const FArguments& InArgs, UEdGraphPin* InPin)
 {
 	SGraphPin::Construct(SGraphPin::FArguments().SideToSideMargin(0.0f), InPin);
diff --git a/Engine/Source/Editor/GraphEditor/Public/SGraphNodeKnot.h b/Engine/Source/Editor/GraphEditor/Public/SGraphNodeKnot.h
--- a/Engine/Source/Editor/GraphEditor/Public/SGraphNodeKnot.h
+++ b/Engine/Source/Editor/GraphEditor/Public/SGraphNodeKnot.h
@@ -48,3 +48,25 @@ protected:
 	const FSlateBrush* ShadowBrush;
 	const FSlateBrush* ShadowBrushSelected;
 };
+
+class GRAPHEDITOR_API SGraphPinKnot : public SGraphPin
+{
+public:
+	SLATE_BEGIN_ARGS(SGraphPinKnot) {}
+	SLATE_END_ARGS()
+
+	void Construct(const FArguments& InArgs, UEdGraphPin* InPin);
+
+	// SWidget interface
+	virtual void OnDragEnter(const FGeometry& MyGeometry, const FDragDropEvent& DragDropEvent) override;
+	// End of SWidget interface
+
+protected:
+	// Begin SGraphPin interface
+	virtual TSharedRef<SWidget>	GetDefaultValueWidget() override;
+	virtual TSharedRef<FDragDropOperation> SpawnPinDragEvent(const TSharedRef<SGraphPanel>& InGraphPanel, const TArray< TSharedRef<SGraphPin> >& InStartingPins) override;
+	virtual FReply OnPinMouseDown(const FGeometry& SenderGeometry, const FPointerEvent& MouseEvent) override;
+	virtual FSlateColor GetPinColor() const override;
+	// End SGraphPin interface
+};
+
