Fixed assertion crash occurring when mouse-overing collapsed material attributes.

Crash occured while generating preview graph displayed in tooltip.

diff --git a/Engine/Source/Editor/GraphEditor/Private/MaterialNodes/SGraphNodeMaterialBase.cpp b/Engine/Source/Editor/GraphEditor/Private/MaterialNodes/SGraphNodeMaterialBase.cpp
--- a/Engine/Source/Editor/GraphEditor/Private/MaterialNodes/SGraphNodeMaterialBase.cpp
+++ b/Engine/Source/Editor/GraphEditor/Private/MaterialNodes/SGraphNodeMaterialBase.cpp
@@ -273,7 +273,7 @@ void SGraphNodeMaterialBase::CreatePinWidgets()
 		bool bPinDesiresToBeHidden = CurPin->bHidden || (bHideNoConnectionPins && !bPinHasConections); 
 
 		UMaterialGraph* MaterialGraph = CastChecked<UMaterialGraph>(GraphNode->GetGraph());
-		if (MaterialNode && MaterialNode->MaterialExpression && MaterialGraph->MaterialFunction == nullptr)
+		if (MaterialNode && MaterialNode->MaterialExpression && MaterialGraph->MaterialFunction == nullptr && !MaterialGraph->MaterialInputs.IsEmpty())
 		{
 			bool bIsAMakeAttrNode = MaterialNode->MaterialExpression->IsA(UMaterialExpressionMakeMaterialAttributes::StaticClass());
 			bool bIsABreakAttrNode = MaterialNode->MaterialExpression->IsA(UMaterialExpressionBreakMaterialAttributes::StaticClass());
