Fixed Scalar Parameter and Vector node to avoid triggering a component value change transaction if the user only clicks on the textbox (without effectively changing the value).

Currently when the user clicks on a Scalar or Vector parameter text box and then simply deselects the node, a data change is erroneously detected and a transaction is pushed onto the Undo stack. This change fixes the logic that determines whether the value is actually changed, removing these spurious and redundant transactions.

diff --git a/Engine/Source/Editor/GraphEditor/Public/KismetPins/SGraphPinNum.h b/Engine/Source/Editor/GraphEditor/Public/KismetPins/SGraphPinNum.h
--- a/Engine/Source/Editor/GraphEditor/Public/KismetPins/SGraphPinNum.h
+++ b/Engine/Source/Editor/GraphEditor/Public/KismetPins/SGraphPinNum.h
@@ -50,18 +50,17 @@ protected:
 
 	void SetNumericValue(NumericType InValue, ETextCommit::Type CommitType)
 	{
-		if(GraphPinObj->IsPendingKill())
+		if (GraphPinObj->IsPendingKill())
 		{
 			return;
 		}
 
-		const FString TypeValueString = LexToString(InValue);
-		if (GraphPinObj->GetDefaultAsString() != TypeValueString)
+		TOptional<NumericType> OldValue = TDefaultNumericTypeInterface<NumericType>{}.FromString(GraphPinObj->GetDefaultAsString(), InValue);
+		if (OldValue && InValue != *OldValue)
 		{
 			const FScopedTransaction Transaction(NSLOCTEXT("GraphEditor", "ChangeNumberPinValue", "Change Number Pin Value"));
 			GraphPinObj->Modify();
-
-			GraphPinObj->GetSchema()->TrySetDefaultValue(*GraphPinObj, *TypeValueString);
+			GraphPinObj->GetSchema()->TrySetDefaultValue(*GraphPinObj, *TDefaultNumericTypeInterface<NumericType>{}.ToString(InValue));
 		}
 	}
 };
diff --git a/Engine/Source/Editor/GraphEditor/Public/KismetPins/SGraphPinVector.h b/Engine/Source/Editor/GraphEditor/Public/KismetPins/SGraphPinVector.h
--- a/Engine/Source/Editor/GraphEditor/Public/KismetPins/SGraphPinVector.h
+++ b/Engine/Source/Editor/GraphEditor/Public/KismetPins/SGraphPinVector.h
@@ -54,6 +54,8 @@ private:
 		TextBox_2,
 	};
 
+	using FVectorType = UE::Math::TVector<NumericType>;
+
 	// Rotator is represented as X->Roll, Y->Pitch, Z->Yaw
 
 	/*
@@ -64,7 +66,8 @@ private:
 	FString GetCurrentValue_0() const
 	{
 		// Text box 0: Rotator->Roll, Vector->X
-		return GetValue(bIsRotator ? TextBox_2 : TextBox_0);
+		FVectorType Vector = ConvertDefaultValueStringToVector();
+		return TDefaultNumericTypeInterface<NumericType>{}.ToString(bIsRotator ? Vector.Z : Vector.X);
 	}
 
 	/*
@@ -75,7 +78,8 @@ private:
 	FString GetCurrentValue_1() const
 	{
 		// Text box 1: Rotator->Pitch, Vector->Y
-		return GetValue(bIsRotator ? TextBox_0 : TextBox_1);
+		FVectorType Vector = ConvertDefaultValueStringToVector();
+		return TDefaultNumericTypeInterface<NumericType>{}.ToString(bIsRotator ? Vector.X : Vector.Y);
 	}
 
 	/*
@@ -86,34 +90,8 @@ private:
 	FString GetCurrentValue_2() const
 	{
 		// Text box 2: Rotator->Yaw, Vector->Z
-		return GetValue(bIsRotator ? TextBox_1 : TextBox_2);
-	}
-
-	/*
-	 *	Function to getch current value based on text box index value
-	 *
-	 *	@param: Text box index
-	 *
-	 *	@return current string value
-	 */
-	FString GetValue( ETextBoxIndex Index ) const
-	{
-		FString DefaultString = GraphPinObj->GetDefaultAsString();
-		TArray<FString> ResultString;
-
-		// Parse string to split its contents separated by ','
-		DefaultString.TrimStartInline();
-		DefaultString.TrimEndInline();
-		DefaultString.ParseIntoArray(ResultString, TEXT(","), true);
-
-		if (Index < ResultString.Num())
-		{
-			return ResultString[Index];
-		}
-		else
-		{
-			return FString(TEXT("0"));
-		}
+		FVectorType Vector = ConvertDefaultValueStringToVector();
+		return TDefaultNumericTypeInterface<NumericType>{}.ToString(bIsRotator ? Vector.Y : Vector.Z);
 	}
 
 	/*
@@ -128,28 +106,24 @@ private:
 			return;
 		}
 
-		const FString ValueStr = FString::Printf(TEXT("%f"), NewValue);
+		FVectorType NewVector = ConvertDefaultValueStringToVector();
+		NumericType OldValue;
 
-		FString DefaultValue;
 		if (bIsRotator)
 		{
-			// Update Roll value
-			DefaultValue = GetValue(TextBox_0) + FString(TEXT(",")) + GetValue(TextBox_1) + FString(TEXT(",")) + ValueStr;
+			/* update roll */
+			OldValue = NewVector.Z;
+			NewVector.Z = NewValue;
 		}
 		else
 		{
-			// Update X value
-			DefaultValue = ValueStr + FString(TEXT(",")) + GetValue(TextBox_1) + FString(TEXT(",")) + GetValue(TextBox_2);
+			/* X value */
+			OldValue = NewVector.X;
+			NewVector.X = NewValue;
 		}
 
-		if (GraphPinObj->GetDefaultAsString() != DefaultValue)
-		{
-			const FScopedTransaction Transaction(NSLOCTEXT("GraphEditor", "ChangeVectorPinValue", "Change Vector Pin Value"));
-			GraphPinObj->Modify();
+		SetNewValueHelper(OldValue, NewValue, NewVector);
 
-			// Set new default value
-			GraphPinObj->GetSchema()->TrySetDefaultValue(*GraphPinObj, DefaultValue);
-		}
 	}
 
 	/*
@@ -164,28 +138,23 @@ private:
 			return;
 		}
 
-		const FString ValueStr = FString::Printf(TEXT("%f"), NewValue);
+		FVectorType NewVector = ConvertDefaultValueStringToVector();
+		NumericType OldValue;
 
-		FString DefaultValue;
 		if (bIsRotator)
 		{
-			// Update Pitch value
-			DefaultValue = ValueStr + FString(TEXT(",")) + GetValue(TextBox_1) + FString(TEXT(",")) + GetValue(TextBox_2);
+			/* update pitch */
+			OldValue = NewVector.X;
+			NewVector.X = NewValue;
 		}
 		else
 		{
-			// Update Y value
-			DefaultValue = GetValue(TextBox_0) + FString(TEXT(",")) + ValueStr + FString(TEXT(",")) + GetValue(TextBox_2);
+			/* Y value */
+			OldValue = NewVector.Y;
+			NewVector.Y = NewValue;
 		}
 
-		if (GraphPinObj->GetDefaultAsString() != DefaultValue)
-		{
-			const FScopedTransaction Transaction(NSLOCTEXT("GraphEditor", "ChangeVectorPinValue", "Change Vector Pin Value"));
-			GraphPinObj->Modify();
-
-			// Set new default value
-			GraphPinObj->GetSchema()->TrySetDefaultValue(*GraphPinObj, DefaultValue);
-		}
+		SetNewValueHelper(OldValue, NewValue, NewVector);
 	}
 
 	/*
@@ -199,29 +168,69 @@ private:
 		{
 			return;
 		}
+		
+		FVectorType NewVector = ConvertDefaultValueStringToVector();
+		NumericType OldValue;
 
-		const FString ValueStr = FString::Printf(TEXT("%f"), NewValue);
-
-		FString DefaultValue;
 		if (bIsRotator)
 		{
-			// Update Yaw value
-			DefaultValue = GetValue(TextBox_0) + FString(TEXT(",")) + ValueStr + FString(TEXT(",")) + GetValue(TextBox_2);
+			/* update yaw */
+			OldValue = NewVector.Y;
+			NewVector.Y = NewValue;
 		}
 		else
 		{
-			// Update Z value
-			DefaultValue = GetValue(TextBox_0) + FString(TEXT(",")) + GetValue(TextBox_1) + FString(TEXT(",")) + ValueStr;
+			/* Z value */
+			OldValue = NewVector.Z;
+			NewVector.Z = NewValue;
 		}
 
-		if (GraphPinObj->GetDefaultAsString() != DefaultValue)
-		{
-			const FScopedTransaction Transaction(NSLOCTEXT("GraphEditor", "ChangeVectorPinValue", "Change Vector Pin Value"));
-			GraphPinObj->Modify();
+		SetNewValueHelper(OldValue, NewValue, NewVector);
+	}
 
-			// Set new default value
-			GraphPinObj->GetSchema()->TrySetDefaultValue(*GraphPinObj, DefaultValue);
+	void SetNewValueHelper(NumericType OldValue, NumericType NewValue, FVectorType const& NewVector)
+	{
+		if (OldValue == NewValue)
+		{
+			return;
 		}
+
+		const FScopedTransaction Transaction(NSLOCTEXT("GraphEditor", "ChangeVectorPinValue", "Change Vector Pin Value"));
+		GraphPinObj->Modify();
+
+		// Create the new value string
+		FString DefaultValue = FString::Format(TEXT("{0},{1},{2}"), { NewVector.X, NewVector.Y, NewVector.Z });
+
+		//Set new default value
+		GraphPinObj->GetSchema()->TrySetDefaultValue(*GraphPinObj, DefaultValue);
+	}
+
+	/*
+	 * @brief Converts the default string value to a value of VectorType.
+	 * 
+	 * Example: it converts the string "(2.00,3.00,4.00)" to the corresponding 3D vector.
+	 */
+	FVectorType ConvertDefaultValueStringToVector() const
+	{
+		FString DefaultString = GraphPinObj->GetDefaultAsString();
+		DefaultString.TrimStartInline();
+		DefaultString.TrimEndInline();
+
+		// Parse string to split its contents separated by ','
+		TArray<FString> VecComponentStrings;
+		DefaultString.ParseIntoArray(VecComponentStrings, TEXT(","), true);
+
+		check(VecComponentStrings.Num() == 3);
+
+		TDefaultNumericTypeInterface<NumericType> NumericTypeInterface{};
+
+		// Construct the vector from the string parts.
+		FVectorType Vec = FVectorType::ZeroVector;
+		Vec.X = NumericTypeInterface.FromString(VecComponentStrings[0], 0).Get(0);
+		Vec.Y = NumericTypeInterface.FromString(VecComponentStrings[1], 0).Get(0);
+		Vec.Z = NumericTypeInterface.FromString(VecComponentStrings[2], 0).Get(0);
+
+		return Vec;
 	}
 
 private:
diff --git a/Engine/Source/Editor/GraphEditor/Public/KismetPins/SGraphPinVector2D.h b/Engine/Source/Editor/GraphEditor/Public/KismetPins/SGraphPinVector2D.h
--- a/Engine/Source/Editor/GraphEditor/Public/KismetPins/SGraphPinVector2D.h
+++ b/Engine/Source/Editor/GraphEditor/Public/KismetPins/SGraphPinVector2D.h
@@ -56,7 +56,7 @@ private:
 	 */
 	FString GetCurrentValue_X() const
 	{
-		return GetValue(TextBox_X);
+		return FString::Printf(TEXT("%f"), GetValue().X);
 	}
 
 	/*
@@ -66,7 +66,7 @@ private:
 	 */
 	FString GetCurrentValue_Y() const
 	{
-		return GetValue(TextBox_Y);
+		return FString::Printf(TEXT("%f"), GetValue().Y);
 	}
 
 	/*
@@ -74,29 +74,21 @@ private:
 	 *
 	 *	@param: Text box index
 	 *
-	 *	@return current string value
+	 *	@return current value
 	 */
-	FString GetValue(ETextBoxIndex Index) const
+	FVector2D GetValue() const
 	{
-		FString DefaultString = GraphPinObj->GetDefaultAsString();
-		TArray<FString> ResultString;
+		const FString& DefaultString = GraphPinObj->GetDefaultAsString();
 
 		FVector2D Value;
 		Value.InitFromString(DefaultString);
 
-		if (Index == TextBox_X)
-		{
-			return FString::Printf(TEXT("%f"), Value.X);
-		}
-		else
-		{
-			return FString::Printf(TEXT("%f"), Value.Y);
-		}
+		return Value;
 	}
 
-	FString MakeVector2DString(const FString& X, const FString& Y)
+	static FString MakeVector2DString(NumericType X, NumericType Y)
 	{
-		return FString(TEXT("(X=")) + X + FString(TEXT(",Y=")) + Y + FString(TEXT(")"));
+		return FString::Format(TEXT("(X={0},Y={1})"), { X, Y });
 	}
 
 	/*
@@ -111,13 +103,13 @@ private:
 			return;
 		}
 
-		const FString ValueStr = FString::Printf(TEXT("%f"), NewValue);
-		const FString Vector2DString = MakeVector2DString(ValueStr, GetValue(TextBox_Y));
-
-		if (GraphPinObj->GetDefaultAsString() != Vector2DString)
+		const FVector2D OldValue = GetValue();
+		if (NewValue != OldValue.X)
 		{
 			const FScopedTransaction Transaction(NSLOCTEXT("GraphEditor", "ChangeVectorPinValue", "Change Vector Pin Value"));
 			GraphPinObj->Modify();
+
+			const FString Vector2DString = MakeVector2DString(NewValue, OldValue.Y);
 			GraphPinObj->GetSchema()->TrySetDefaultValue(*GraphPinObj, Vector2DString);
 		}
 	}
@@ -134,13 +126,13 @@ private:
 			return;
 		}
 
-		const FString ValueStr = FString::Printf(TEXT("%f"), NewValue);
-		const FString Vector2DString = MakeVector2DString(GetValue(TextBox_X), ValueStr);
-
-		if (GraphPinObj->GetDefaultAsString() != Vector2DString)
+		const FVector2D OldValue = GetValue();
+		if (NewValue != OldValue.Y)
 		{
 			const FScopedTransaction Transaction(NSLOCTEXT("GraphEditor", "ChangeVectorPinValue", "Change Vector Pin Value"));
 			GraphPinObj->Modify();
+
+			const FString Vector2DString = MakeVector2DString(OldValue.X, NewValue);
 			GraphPinObj->GetSchema()->TrySetDefaultValue(*GraphPinObj, Vector2DString);
 		}
 	}
