Revert 21442104 because it crashes a blueprint when you open it.

#rnx

diff --git a/Engine/Source/Editor/GraphEditor/Public/KismetPins/SGraphPinNum.h b/Engine/Source/Editor/GraphEditor/Public/KismetPins/SGraphPinNum.h
--- a/Engine/Source/Editor/GraphEditor/Public/KismetPins/SGraphPinNum.h
+++ b/Engine/Source/Editor/GraphEditor/Public/KismetPins/SGraphPinNum.h
@@ -50,17 +50,18 @@ protected:
 
 	void SetNumericValue(NumericType InValue, ETextCommit::Type CommitType)
 	{
-		if (GraphPinObj->IsPendingKill())
+		if(GraphPinObj->IsPendingKill())
 		{
 			return;
 		}
 
-		TOptional<NumericType> OldValue = TDefaultNumericTypeInterface<NumericType>{}.FromString(GraphPinObj->GetDefaultAsString(), InValue);
-		if (OldValue && InValue != *OldValue)
+		const FString TypeValueString = LexToString(InValue);
+		if (GraphPinObj->GetDefaultAsString() != TypeValueString)
 		{
 			const FScopedTransaction Transaction(NSLOCTEXT("GraphEditor", "ChangeNumberPinValue", "Change Number Pin Value"));
 			GraphPinObj->Modify();
-			GraphPinObj->GetSchema()->TrySetDefaultValue(*GraphPinObj, *TDefaultNumericTypeInterface<NumericType>{}.ToString(InValue));
+
+			GraphPinObj->GetSchema()->TrySetDefaultValue(*GraphPinObj, *TypeValueString);
 		}
 	}
 };
diff --git a/Engine/Source/Editor/GraphEditor/Public/KismetPins/SGraphPinVector.h b/Engine/Source/Editor/GraphEditor/Public/KismetPins/SGraphPinVector.h
--- a/Engine/Source/Editor/GraphEditor/Public/KismetPins/SGraphPinVector.h
+++ b/Engine/Source/Editor/GraphEditor/Public/KismetPins/SGraphPinVector.h
@@ -54,8 +54,6 @@ private:
 		TextBox_2,
 	};
 
-	using FVectorType = UE::Math::TVector<NumericType>;
-
 	// Rotator is represented as X->Roll, Y->Pitch, Z->Yaw
 
 	/*
@@ -66,8 +64,7 @@ private:
 	FString GetCurrentValue_0() const
 	{
 		// Text box 0: Rotator->Roll, Vector->X
-		FVectorType Vector = ConvertDefaultValueStringToVector();
-		return TDefaultNumericTypeInterface<NumericType>{}.ToString(bIsRotator ? Vector.Z : Vector.X);
+		return GetValue(bIsRotator ? TextBox_2 : TextBox_0);
 	}
 
 	/*
@@ -78,8 +75,7 @@ private:
 	FString GetCurrentValue_1() const
 	{
 		// Text box 1: Rotator->Pitch, Vector->Y
-		FVectorType Vector = ConvertDefaultValueStringToVector();
-		return TDefaultNumericTypeInterface<NumericType>{}.ToString(bIsRotator ? Vector.X : Vector.Y);
+		return GetValue(bIsRotator ? TextBox_0 : TextBox_1);
 	}
 
 	/*
@@ -90,8 +86,34 @@ private:
 	FString GetCurrentValue_2() const
 	{
 		// Text box 2: Rotator->Yaw, Vector->Z
-		FVectorType Vector = ConvertDefaultValueStringToVector();
-		return TDefaultNumericTypeInterface<NumericType>{}.ToString(bIsRotator ? Vector.Y : Vector.Z);
+		return GetValue(bIsRotator ? TextBox_1 : TextBox_2);
+	}
+
+	/*
+	 *	Function to getch current value based on text box index value
+	 *
+	 *	@param: Text box index
+	 *
+	 *	@return current string value
+	 */
+	FString GetValue( ETextBoxIndex Index ) const
+	{
+		FString DefaultString = GraphPinObj->GetDefaultAsString();
+		TArray<FString> ResultString;
+
+		// Parse string to split its contents separated by ','
+		DefaultString.TrimStartInline();
+		DefaultString.TrimEndInline();
+		DefaultString.ParseIntoArray(ResultString, TEXT(","), true);
+
+		if (Index < ResultString.Num())
+		{
+			return ResultString[Index];
+		}
+		else
+		{
+			return FString(TEXT("0"));
+		}
 	}
 
 	/*
@@ -106,24 +128,28 @@ private:
 			return;
 		}
 
-		FVectorType NewVector = ConvertDefaultValueStringToVector();
-		NumericType OldValue;
+		const FString ValueStr = FString::Printf(TEXT("%f"), NewValue);
 
+		FString DefaultValue;
 		if (bIsRotator)
 		{
-			/* update roll */
-			OldValue = NewVector.Z;
-			NewVector.Z = NewValue;
+			// Update Roll value
+			DefaultValue = GetValue(TextBox_0) + FString(TEXT(",")) + GetValue(TextBox_1) + FString(TEXT(",")) + ValueStr;
 		}
 		else
 		{
-			/* X value */
-			OldValue = NewVector.X;
-			NewVector.X = NewValue;
+			// Update X value
+			DefaultValue = ValueStr + FString(TEXT(",")) + GetValue(TextBox_1) + FString(TEXT(",")) + GetValue(TextBox_2);
 		}
 
-		SetNewValueHelper(OldValue, NewValue, NewVector);
+		if (GraphPinObj->GetDefaultAsString() != DefaultValue)
+		{
+			const FScopedTransaction Transaction(NSLOCTEXT("GraphEditor", "ChangeVectorPinValue", "Change Vector Pin Value"));
+			GraphPinObj->Modify();
 
+			// Set new default value
+			GraphPinObj->GetSchema()->TrySetDefaultValue(*GraphPinObj, DefaultValue);
+		}
 	}
 
 	/*
@@ -138,23 +164,28 @@ private:
 			return;
 		}
 
-		FVectorType NewVector = ConvertDefaultValueStringToVector();
-		NumericType OldValue;
+		const FString ValueStr = FString::Printf(TEXT("%f"), NewValue);
 
+		FString DefaultValue;
 		if (bIsRotator)
 		{
-			/* update pitch */
-			OldValue = NewVector.X;
-			NewVector.X = NewValue;
+			// Update Pitch value
+			DefaultValue = ValueStr + FString(TEXT(",")) + GetValue(TextBox_1) + FString(TEXT(",")) + GetValue(TextBox_2);
 		}
 		else
 		{
-			/* Y value */
-			OldValue = NewVector.Y;
-			NewVector.Y = NewValue;
+			// Update Y value
+			DefaultValue = GetValue(TextBox_0) + FString(TEXT(",")) + ValueStr + FString(TEXT(",")) + GetValue(TextBox_2);
 		}
 
-		SetNewValueHelper(OldValue, NewValue, NewVector);
+		if (GraphPinObj->GetDefaultAsString() != DefaultValue)
+		{
+			const FScopedTransaction Transaction(NSLOCTEXT("GraphEditor", "ChangeVectorPinValue", "Change Vector Pin Value"));
+			GraphPinObj->Modify();
+
+			// Set new default value
+			GraphPinObj->GetSchema()->TrySetDefaultValue(*GraphPinObj, DefaultValue);
+		}
 	}
 
 	/*
@@ -168,69 +199,29 @@ private:
 		{
 			return;
 		}
-		
-		FVectorType NewVector = ConvertDefaultValueStringToVector();
-		NumericType OldValue;
 
+		const FString ValueStr = FString::Printf(TEXT("%f"), NewValue);
+
+		FString DefaultValue;
 		if (bIsRotator)
 		{
-			/* update yaw */
-			OldValue = NewVector.Y;
-			NewVector.Y = NewValue;
+			// Update Yaw value
+			DefaultValue = GetValue(TextBox_0) + FString(TEXT(",")) + ValueStr + FString(TEXT(",")) + GetValue(TextBox_2);
 		}
 		else
 		{
-			/* Z value */
-			OldValue = NewVector.Z;
-			NewVector.Z = NewValue;
+			// Update Z value
+			DefaultValue = GetValue(TextBox_0) + FString(TEXT(",")) + GetValue(TextBox_1) + FString(TEXT(",")) + ValueStr;
 		}
 
-		SetNewValueHelper(OldValue, NewValue, NewVector);
-	}
-
-	void SetNewValueHelper(NumericType OldValue, NumericType NewValue, FVectorType const& NewVector)
-	{
-		if (OldValue == NewValue)
+		if (GraphPinObj->GetDefaultAsString() != DefaultValue)
 		{
-			return;
-		}
-
-		const FScopedTransaction Transaction(NSLOCTEXT("GraphEditor", "ChangeVectorPinValue", "Change Vector Pin Value"));
-		GraphPinObj->Modify();
-
-		// Create the new value string
-		FString DefaultValue = FString::Format(TEXT("{0},{1},{2}"), { NewVector.X, NewVector.Y, NewVector.Z });
-
-		//Set new default value
-		GraphPinObj->GetSchema()->TrySetDefaultValue(*GraphPinObj, DefaultValue);
-	}
-
-	/*
-	 * @brief Converts the default string value to a value of VectorType.
-	 * 
-	 * Example: it converts the string "(2.00,3.00,4.00)" to the corresponding 3D vector.
-	 */
-	FVectorType ConvertDefaultValueStringToVector() const
-	{
-		FString DefaultString = GraphPinObj->GetDefaultAsString();
-		DefaultString.TrimStartInline();
-		DefaultString.TrimEndInline();
+			const FScopedTransaction Transaction(NSLOCTEXT("GraphEditor", "ChangeVectorPinValue", "Change Vector Pin Value"));
+			GraphPinObj->Modify();
 
-		// Parse string to split its contents separated by ','
-		TArray<FString> VecComponentStrings;
-		DefaultString.ParseIntoArray(VecComponentStrings, TEXT(","), true);
-
-		check(VecComponentStrings.Num() == 3);
-
-		TDefaultNumericTypeInterface<NumericType> NumericTypeInterface{};
-
-		// Construct the vector from the string parts.
-		FVectorType Vec = FVectorType::ZeroVector;
-		Vec.X = NumericTypeInterface.FromString(VecComponentStrings[0], 0).Get(0);
-		Vec.Y = NumericTypeInterface.FromString(VecComponentStrings[1], 0).Get(0);
-		Vec.Z = NumericTypeInterface.FromString(VecComponentStrings[2], 0).Get(0);
-
-		return Vec;
+			// Set new default value
+			GraphPinObj->GetSchema()->TrySetDefaultValue(*GraphPinObj, DefaultValue);
+		}
 	}
 
 private:
diff --git a/Engine/Source/Editor/GraphEditor/Public/KismetPins/SGraphPinVector2D.h b/Engine/Source/Editor/GraphEditor/Public/KismetPins/SGraphPinVector2D.h
--- a/Engine/Source/Editor/GraphEditor/Public/KismetPins/SGraphPinVector2D.h
+++ b/Engine/Source/Editor/GraphEditor/Public/KismetPins/SGraphPinVector2D.h
@@ -56,7 +56,7 @@ private:
 	 */
 	FString GetCurrentValue_X() const
 	{
-		return FString::Printf(TEXT("%f"), GetValue().X);
+		return GetValue(TextBox_X);
 	}
 
 	/*
@@ -66,7 +66,7 @@ private:
 	 */
 	FString GetCurrentValue_Y() const
 	{
-		return FString::Printf(TEXT("%f"), GetValue().Y);
+		return GetValue(TextBox_Y);
 	}
 
 	/*
@@ -74,21 +74,29 @@ private:
 	 *
 	 *	@param: Text box index
 	 *
-	 *	@return current value
+	 *	@return current string value
 	 */
-	FVector2D GetValue() const
+	FString GetValue(ETextBoxIndex Index) const
 	{
-		const FString& DefaultString = GraphPinObj->GetDefaultAsString();
+		FString DefaultString = GraphPinObj->GetDefaultAsString();
+		TArray<FString> ResultString;
 
 		FVector2D Value;
 		Value.InitFromString(DefaultString);
 
-		return Value;
+		if (Index == TextBox_X)
+		{
+			return FString::Printf(TEXT("%f"), Value.X);
+		}
+		else
+		{
+			return FString::Printf(TEXT("%f"), Value.Y);
+		}
 	}
 
-	static FString MakeVector2DString(NumericType X, NumericType Y)
+	FString MakeVector2DString(const FString& X, const FString& Y)
 	{
-		return FString::Format(TEXT("(X={0},Y={1})"), { X, Y });
+		return FString(TEXT("(X=")) + X + FString(TEXT(",Y=")) + Y + FString(TEXT(")"));
 	}
 
 	/*
@@ -103,13 +111,13 @@ private:
 			return;
 		}
 
-		const FVector2D OldValue = GetValue();
-		if (NewValue != OldValue.X)
+		const FString ValueStr = FString::Printf(TEXT("%f"), NewValue);
+		const FString Vector2DString = MakeVector2DString(ValueStr, GetValue(TextBox_Y));
+
+		if (GraphPinObj->GetDefaultAsString() != Vector2DString)
 		{
 			const FScopedTransaction Transaction(NSLOCTEXT("GraphEditor", "ChangeVectorPinValue", "Change Vector Pin Value"));
 			GraphPinObj->Modify();
-
-			const FString Vector2DString = MakeVector2DString(NewValue, OldValue.Y);
 			GraphPinObj->GetSchema()->TrySetDefaultValue(*GraphPinObj, Vector2DString);
 		}
 	}
@@ -126,13 +134,13 @@ private:
 			return;
 		}
 
-		const FVector2D OldValue = GetValue();
-		if (NewValue != OldValue.Y)
+		const FString ValueStr = FString::Printf(TEXT("%f"), NewValue);
+		const FString Vector2DString = MakeVector2DString(GetValue(TextBox_X), ValueStr);
+
+		if (GraphPinObj->GetDefaultAsString() != Vector2DString)
 		{
 			const FScopedTransaction Transaction(NSLOCTEXT("GraphEditor", "ChangeVectorPinValue", "Change Vector Pin Value"));
 			GraphPinObj->Modify();
-
-			const FString Vector2DString = MakeVector2DString(OldValue.X, NewValue);
 			GraphPinObj->GetSchema()->TrySetDefaultValue(*GraphPinObj, Vector2DString);
 		}
 	}
