Fix SGraphPinNum bug causing value not to change in the Niagara editor when the user enters a new value.

diff --git a/Engine/Source/Editor/GraphEditor/Public/KismetPins/SGraphPinNum.h b/Engine/Source/Editor/GraphEditor/Public/KismetPins/SGraphPinNum.h
--- a/Engine/Source/Editor/GraphEditor/Public/KismetPins/SGraphPinNum.h
+++ b/Engine/Source/Editor/GraphEditor/Public/KismetPins/SGraphPinNum.h
@@ -56,7 +56,7 @@ protected:
 		}
 
 		TOptional<NumericType> OldValue = TDefaultNumericTypeInterface<NumericType>{}.FromString(GraphPinObj->GetDefaultAsString(), InValue);
-		if (OldValue && InValue != *OldValue)
+		if (!OldValue.IsSet() || InValue != *OldValue)
 		{
 			const FScopedTransaction Transaction(NSLOCTEXT("GraphEditor", "ChangeNumberPinValue", "Change Number Pin Value"));
 			GraphPinObj->Modify();
