[BugFix] Blueprints incorrectly diff cleanly when you link a pin to a different pin on the same node

diff --git a/Engine/Source/Editor/GraphEditor/Private/GraphDiffControl.cpp b/Engine/Source/Editor/GraphEditor/Private/GraphDiffControl.cpp
--- a/Engine/Source/Editor/GraphEditor/Private/GraphDiffControl.cpp
+++ b/Engine/Source/Editor/GraphEditor/Private/GraphDiffControl.cpp
@@ -263,6 +263,32 @@ static void DiffR_LinkedToNode(FDiffResults& Results, UEdGraphPin* OldPin, UEdGr
 	Results.Add(Diff);
 }
 
+/** Diff result when a pin to relinked to a different node */
+static void DiffR_LinkedToPin(FDiffResults& Results, UEdGraphPin* OldPin, UEdGraphPin* NewPin, const UEdGraphPin* OldLinkedPin, const UEdGraphPin* NewLinkedPin)
+{
+	FDiffSingleResult Diff;
+	Diff.Diff = EDiffType::PIN_LINKEDTO_NODE;
+	Diff.Pin1 = OldPin;
+	Diff.Pin2 = NewPin;
+	Diff.Node1 = nullptr;
+	Diff.Node2 = nullptr;
+
+	// Only bother setting up the display data if we're storing the result
+	if(Results.CanStoreResults())
+	{
+		FFormatNamedArguments Args;
+		Args.Add(TEXT("PinNameForNode1"), FText::FromName(OldPin->PinName));
+		Args.Add(TEXT("OldLinkedPinName"), FText::FromName(OldLinkedPin->PinName));
+		Args.Add(TEXT("NewLinkedPinName"), FText::FromName(NewLinkedPin->PinName));
+		Diff.ToolTip = FText::Format(LOCTEXT("DIF_PinLinkMovedToolTip", "Pin '{PinNameForNode1}' was linked to Pin '{OldLinkedPinName}', but is now linked to Pin '{NewLinkedPinName}'"), Args);
+
+		Diff.Category = EDiffType::MODIFICATION;
+		Diff.DisplayString = FText::Format(LOCTEXT("DIF_PinLinkMoved", "Link Moved  '{PinNameForNode1}' ['{OldLinkedPinName}' -> '{NewLinkedPinName}']"), Args);
+	}
+
+	Results.Add(Diff);
+}
+
 /** Diff result when a pin default value was changed, and is in use*/
 static void DiffR_PinDefaultValueChanged(FDiffResults& Results, UEdGraphPin* NewPin, UEdGraphPin* OldPin)
 {
@@ -470,6 +496,11 @@ static bool LinkedToDifferent(UEdGraphPin* OldPin, UEdGraphPin* NewPin, const TA
 			DiffR_LinkedToNode(Results, OldPin, NewPin, OldNode, NewNode);
 			KEEP_GOING_IF_RESULTS()
 		}
+		else if (OldLinkedPin->PinId != NewLinkedPin->PinId)
+		{
+			DiffR_LinkedToPin(Results, OldPin, NewPin, OldLinkedPin, NewLinkedPin);
+			KEEP_GOING_IF_RESULTS()
+		}
 	}
 	return bHasResult;
 }
diff --git a/Engine/Source/Runtime/Engine/Public/DiffResults.h b/Engine/Source/Runtime/Engine/Public/DiffResults.h
--- a/Engine/Source/Runtime/Engine/Public/DiffResults.h
+++ b/Engine/Source/Runtime/Engine/Public/DiffResults.h
@@ -36,6 +36,7 @@ namespace EDiffType
 		PIN_TYPE_IS_ARRAY,
 		PIN_TYPE_IS_REF,
 		PIN_LINKEDTO_NODE,
+		PIN_LINKEDTO_PIN, // only used when a pin was relinked to a different pin on the same node
 		NODE_MOVED,
 		TIMELINE_LENGTH,
 		TIMELINE_AUTOPLAY,
