[BugFix] //UE5/Main - Build Engine Localization - GraphDiffControl.cpp Warning: Text conflict from LOCTEXT macro for namespace and key

diff --git a/Engine/Source/Editor/GraphEditor/Private/GraphDiffControl.cpp b/Engine/Source/Editor/GraphEditor/Private/GraphDiffControl.cpp
--- a/Engine/Source/Editor/GraphEditor/Private/GraphDiffControl.cpp
+++ b/Engine/Source/Editor/GraphEditor/Private/GraphDiffControl.cpp
@@ -280,10 +280,10 @@ static void DiffR_LinkedToPin(FDiffResults& Results, UEdGraphPin* OldPin, UEdGra
 		Args.Add(TEXT("PinNameForNode1"), FText::FromName(OldPin->PinName));
 		Args.Add(TEXT("OldLinkedPinName"), FText::FromName(OldLinkedPin->PinName));
 		Args.Add(TEXT("NewLinkedPinName"), FText::FromName(NewLinkedPin->PinName));
-		Diff.ToolTip = FText::Format(LOCTEXT("DIF_PinLinkMovedToolTip", "Pin '{PinNameForNode1}' was linked to Pin '{OldLinkedPinName}', but is now linked to Pin '{NewLinkedPinName}'"), Args);
+		Diff.ToolTip = FText::Format(LOCTEXT("DIF_PinLinkMovedToPinToolTip", "Pin '{PinNameForNode1}' was linked to Pin '{OldLinkedPinName}', but is now linked to Pin '{NewLinkedPinName}'"), Args);
 
 		Diff.Category = EDiffType::MODIFICATION;
-		Diff.DisplayString = FText::Format(LOCTEXT("DIF_PinLinkMoved", "Link Moved  '{PinNameForNode1}' ['{OldLinkedPinName}' -> '{NewLinkedPinName}']"), Args);
+		Diff.DisplayString = FText::Format(LOCTEXT("DIF_PinLinkMovedToPin", "Link Moved  '{PinNameForNode1}' ['{OldLinkedPinName}' -> '{NewLinkedPinName}']"), Args);
 	}
 
 	Results.Add(Diff);
