small comment change

diff --git a/Engine/Source/Editor/GraphEditor/Private/GraphDiffControl.cpp b/Engine/Source/Editor/GraphEditor/Private/GraphDiffControl.cpp
--- a/Engine/Source/Editor/GraphEditor/Private/GraphDiffControl.cpp
+++ b/Engine/Source/Editor/GraphEditor/Private/GraphDiffControl.cpp
@@ -263,7 +263,7 @@ static void DiffR_LinkedToNode(FDiffResults& Results, UEdGraphPin* OldPin, UEdGr
 	Results.Add(Diff);
 }
 
-/** Diff result when a pin to relinked to a different node */
+/** Diff result when a pin to relinked to a different pin on the same node */
 static void DiffR_LinkedToPin(FDiffResults& Results, UEdGraphPin* OldPin, UEdGraphPin* NewPin, const UEdGraphPin* OldLinkedPin, const UEdGraphPin* NewLinkedPin)
 {
 	FDiffSingleResult Diff;
