Fix for regional issue (when using ',' as decimal separator) in the numeric and vector input in blueprint.
The regression was introduced a few weeks ago when fixing another bug.
The issue was that conversion functions (between float and string) using the current regionnal settings were introduced in the graph pins, but this kind of conversion is done in a lot of places and does not all support the different type of decimal separator ('.' or ','). So those conversion functions were not supporting the new format of decimal they were receiving.
Also, while the bug is mentionning issue with float, it was existing also for vector.
Bonus: remove a bit of duplicated code.

#test In blueprint/niagara/material, played with node using float, vector, rotator, to ensure that everything was behaving as expected with 2 regional settings (english with '.' decimal separator, french with ',' separator).
Validated that the initial bug (selecting default value with the mouse but not changing it was adding an entry in the undo stack), was still fix.

diff --git a/Engine/Source/Editor/GraphEditor/Public/KismetPins/SGraphPinNum.h b/Engine/Source/Editor/GraphEditor/Public/KismetPins/SGraphPinNum.h
--- a/Engine/Source/Editor/GraphEditor/Public/KismetPins/SGraphPinNum.h
+++ b/Engine/Source/Editor/GraphEditor/Public/KismetPins/SGraphPinNum.h
@@ -55,12 +55,12 @@ protected:
 			return;
 		}
 
-		TOptional<NumericType> OldValue = TDefaultNumericTypeInterface<NumericType>{}.FromString(GraphPinObj->GetDefaultAsString(), InValue);
-		if (!OldValue.IsSet() || InValue != *OldValue)
+		if (GetNumericValue() != InValue)
 		{
 			const FScopedTransaction Transaction(NSLOCTEXT("GraphEditor", "ChangeNumberPinValue", "Change Number Pin Value"));
 			GraphPinObj->Modify();
-			GraphPinObj->GetSchema()->TrySetDefaultValue(*GraphPinObj, *TDefaultNumericTypeInterface<NumericType>{}.ToString(InValue));
+
+			GraphPinObj->GetSchema()->TrySetDefaultValue(*GraphPinObj, *LexToString(InValue));
 		}
 	}
 };
diff --git a/Engine/Source/Editor/GraphEditor/Public/KismetPins/SGraphPinVector.h b/Engine/Source/Editor/GraphEditor/Public/KismetPins/SGraphPinVector.h
--- a/Engine/Source/Editor/GraphEditor/Public/KismetPins/SGraphPinVector.h
+++ b/Engine/Source/Editor/GraphEditor/Public/KismetPins/SGraphPinVector.h
@@ -72,8 +72,7 @@ private:
 	FString GetCurrentValue_0() const
 	{
 		// Text box 0: Rotator->Roll, Vector->X
-		FVectorType Vector = ConvertDefaultValueStringToVector();
-		return TDefaultNumericTypeInterface<NumericType>{}.ToString(bIsRotator ? Vector.Z : Vector.X);
+		return GetValue(bIsRotator ? TextBox_2 : TextBox_0);
 	}
 
 	/*
@@ -84,8 +83,7 @@ private:
 	FString GetCurrentValue_1() const
 	{
 		// Text box 1: Rotator->Pitch, Vector->Y
-		FVectorType Vector = ConvertDefaultValueStringToVector();
-		return TDefaultNumericTypeInterface<NumericType>{}.ToString(bIsRotator ? Vector.X : Vector.Y);
+		return GetValue(bIsRotator ? TextBox_0 : TextBox_1);
 	}
 
 	/*
@@ -96,40 +94,58 @@ private:
 	FString GetCurrentValue_2() const
 	{
 		// Text box 2: Rotator->Yaw, Vector->Z
-		FVectorType Vector = ConvertDefaultValueStringToVector();
-		return TDefaultNumericTypeInterface<NumericType>{}.ToString(bIsRotator ? Vector.Y : Vector.Z);
+		return GetValue(bIsRotator ? TextBox_1 : TextBox_2);
 	}
 
 	/*
-	 *	Function to store value when text box 0 value in modified
+	 *	Function to get a string array of the components composing the default string.
+	 *
+	 *	@return The array filled with each component as string
 	 *
-	 *	@param 0: Updated numeric Value
 	 */
-	void OnChangedValueTextBox_0(NumericType NewValue, ETextCommit::Type CommitInfo)
+	TArray<FString> GetComponentArray() const
 	{
-		if (GraphPinObj->IsPendingKill())
-		{
-			return;
-		}
+		TArray<FString> VecComponentStrings;
 
-		FVectorType NewVector = ConvertDefaultValueStringToVector();
-		NumericType OldValue;
+		FString DefaultString = GraphPinObj->GetDefaultAsString();
+		// Parse string to split its contents separated by ','
+		DefaultString.TrimStartInline();
+		DefaultString.TrimEndInline();
+		DefaultString.ParseIntoArray(VecComponentStrings, TEXT(","), true);
+
+		return VecComponentStrings;
+	}
+	
+	/*
+	 *	Function to get current value based on text box index value
+	 *
+	 *	@param: Text box index
+	 *
+	 *	@return current string value
+	 */
+	FString GetValue(ETextBoxIndex Index) const
+	{
+		const TArray<FString> VecComponentStrings = GetComponentArray();
 
-		if (bIsRotator)
+		if (Index < VecComponentStrings.Num())
 		{
-			/* update roll */
-			OldValue = NewVector.Z;
-			NewVector.Z = NewValue;
+			return VecComponentStrings[Index];
 		}
 		else
 		{
-			/* X value */
-			OldValue = NewVector.X;
-			NewVector.X = NewValue;
+			return FString(TEXT("0"));
 		}
+	}
 
-		SetNewValueHelper(OldValue, NewValue, NewVector);
-
+	/*
+	 *	Function to store value when text box 0 value in modified
+	 *
+	 *	@param 0: Updated numeric Value
+	 */
+	void OnChangedValueTextBox_0(NumericType NewValue, ETextCommit::Type CommitInfo)
+	{
+		const EAxis::Type Axis = bIsRotator ? /* update roll */ EAxis::Z : EAxis::X;
+		SetNewValueHelper(Axis, NewValue);
 	}
 
 	/*
@@ -139,28 +155,8 @@ private:
 	 */
 	void OnChangedValueTextBox_1(NumericType NewValue, ETextCommit::Type CommitInfo)
 	{
-		if (GraphPinObj->IsPendingKill())
-		{
-			return;
-		}
-
-		FVectorType NewVector = ConvertDefaultValueStringToVector();
-		NumericType OldValue;
-
-		if (bIsRotator)
-		{
-			/* update pitch */
-			OldValue = NewVector.X;
-			NewVector.X = NewValue;
-		}
-		else
-		{
-			/* Y value */
-			OldValue = NewVector.Y;
-			NewVector.Y = NewValue;
-		}
-
-		SetNewValueHelper(OldValue, NewValue, NewVector);
+		const EAxis::Type Axis = bIsRotator ? /* update pitch */ EAxis::X : EAxis::Y;
+		SetNewValueHelper(Axis, NewValue);
 	}
 
 	/*
@@ -169,38 +165,28 @@ private:
 	 *	@param 0: Updated numeric Value
 	 */
 	void OnChangedValueTextBox_2(NumericType NewValue, ETextCommit::Type CommitInfo)
+	{
+		const EAxis::Type Axis = bIsRotator ? /* update yaw */ EAxis::Y : EAxis::Z;
+		SetNewValueHelper(Axis, NewValue);
+	}
+
+	void SetNewValueHelper(EAxis::Type Axis, NumericType NewValue)
 	{
 		if (GraphPinObj->IsPendingKill())
 		{
 			return;
 		}
-		
-		FVectorType NewVector = ConvertDefaultValueStringToVector();
-		NumericType OldValue;
 
-		if (bIsRotator)
-		{
-			/* update yaw */
-			OldValue = NewVector.Y;
-			NewVector.Y = NewValue;
-		}
-		else
-		{
-			/* Z value */
-			OldValue = NewVector.Z;
-			NewVector.Z = NewValue;
-		}
-
-		SetNewValueHelper(OldValue, NewValue, NewVector);
-	}
+		FVectorType NewVector = ConvertDefaultValueStringToVector();
+		const NumericType OldValue = NewVector.GetComponentForAxis(Axis);
 
-	void SetNewValueHelper(NumericType OldValue, NumericType NewValue, FVectorType const& NewVector)
-	{
 		if (OldValue == NewValue)
 		{
 			return;
 		}
 
+		NewVector.SetComponentForAxis(Axis, NewValue);
+
 		const FScopedTransaction Transaction(NSLOCTEXT("GraphEditor", "ChangeVectorPinValue", "Change Vector Pin Value"));
 		GraphPinObj->Modify();
 
@@ -218,13 +204,7 @@ private:
 	 */
 	FVectorType ConvertDefaultValueStringToVector() const
 	{
-		FString DefaultString = GraphPinObj->GetDefaultAsString();
-		DefaultString.TrimStartInline();
-		DefaultString.TrimEndInline();
-
-		// Parse string to split its contents separated by ','
-		TArray<FString> VecComponentStrings;
-		DefaultString.ParseIntoArray(VecComponentStrings, TEXT(","), true);
+		const TArray<FString> VecComponentStrings = GetComponentArray();
 
 		// Construct the vector from the string parts.
 		FVectorType Vec = FVectorType::ZeroVector;
