Check that LastOuter is valid before dereferencing

diff --git a/Engine/Source/Editor/GraphEditor/Private/SGraphNode.cpp b/Engine/Source/Editor/GraphEditor/Private/SGraphNode.cpp
--- a/Engine/Source/Editor/GraphEditor/Private/SGraphNode.cpp
+++ b/Engine/Source/Editor/GraphEditor/Private/SGraphNode.cpp
@@ -1646,7 +1646,7 @@ void SGraphNode::PopulateMetaTag(FGraphNodeMetaData* TagMeta) const
 		// We want the name of the blueprint as our name - we can find the node from the GUID
 		UObject* Package = GraphNode->GetOutermost();
 		UObject* LastOuter = GraphNode->GetOuter();
-		while (LastOuter->GetOuter() != Package)
+		while (LastOuter && (LastOuter->GetOuter() != Package))
 		{
 			LastOuter = LastOuter->GetOuter();
 		}
