Additional conditions before dereferencing pointers

diff --git a/Engine/Source/Editor/GraphEditor/Private/SGraphNode.cpp b/Engine/Source/Editor/GraphEditor/Private/SGraphNode.cpp
--- a/Engine/Source/Editor/GraphEditor/Private/SGraphNode.cpp
+++ b/Engine/Source/Editor/GraphEditor/Private/SGraphNode.cpp
@@ -1641,7 +1641,7 @@ EVisibility SGraphNode::IsAddPinButtonVisible() const
 
 void SGraphNode::PopulateMetaTag(FGraphNodeMetaData* TagMeta) const
 {
-	if (GraphNode != nullptr)
+	if (GraphNode && TagMeta)
 	{
 		// We want the name of the blueprint as our name - we can find the node from the GUID
 		UObject* Package = GraphNode->GetOutermost();
@@ -1650,10 +1650,14 @@ void SGraphNode::PopulateMetaTag(FGraphNodeMetaData* TagMeta) const
 		{
 			LastOuter = LastOuter->GetOuter();
 		}
-		TagMeta->Tag = FName(*FString::Printf(TEXT("GraphNode_%s_%s"), *LastOuter->GetFullName(), *GraphNode->NodeGuid.ToString()));
-		TagMeta->OuterName = LastOuter->GetFullName();
-		TagMeta->GUID = GraphNode->NodeGuid;
-		TagMeta->FriendlyName = FString::Printf(TEXT("%s in %s"), *GraphNode->GetNodeTitle(ENodeTitleType::ListView).ToString(), *TagMeta->OuterName);		
+
+		if(LastOuter)
+		{
+			TagMeta->Tag = FName(*FString::Printf(TEXT("GraphNode_%s_%s"), *LastOuter->GetFullName(), *GraphNode->NodeGuid.ToString()));
+			TagMeta->OuterName = LastOuter->GetFullName();
+			TagMeta->GUID = GraphNode->NodeGuid;
+			TagMeta->FriendlyName = FString::Printf(TEXT("%s in %s"), *GraphNode->GetNodeTitle(ENodeTitleType::ListView).ToString(), *TagMeta->OuterName);	
+		}	
 	}
 }
 
