[UE-136096] Relinking transitions in anim graph state machines
[UE-168360] User can create a transition from a state into the entry state

* Relinking multiple transitions (transition nodes at once) when none of the transition nodes is selected.
* Selective relinking in case a subset of the transition nodes on a transitions are selected.
* Special case handling for transitions from entry state (they don't have a transition node connected).
* Drawing circle indicator around the arrow head in case the mouse cursor is approaching the transition.
* Hovering the arrow head of a transition highlights it by filling the circle with solid orange color and turning the arrow black.
* Ending the relink operation in empty space or the entry node will cancel the relink operation.
* Fixed bug that prevents creating new transitions ending in an entry state.

diff --git a/Engine/Source/Editor/AnimGraph/Private/AnimStateTransitionNode.cpp b/Engine/Source/Editor/AnimGraph/Private/AnimStateTransitionNode.cpp
--- a/Engine/Source/Editor/AnimGraph/Private/AnimStateTransitionNode.cpp
+++ b/Engine/Source/Editor/AnimGraph/Private/AnimStateTransitionNode.cpp
@@ -304,6 +304,86 @@ void UAnimStateTransitionNode::CreateConnections(UAnimStateNodeBase* PreviousSta
 	Pins[1]->MakeLinkTo(NextState->GetInputPin());
 }
 
+void UAnimStateTransitionNode::RelinkHead(UAnimStateNodeBase* NewTargetState)
+{
+	UAnimStateNodeBase* SourceState = GetPreviousState();
+	UAnimStateNodeBase* TargetStateBeforeRelinking = GetNextState();
+
+	// Remove the incoming transition from the previous target state
+	TargetStateBeforeRelinking->GetInputPin()->Modify();
+	TargetStateBeforeRelinking->GetInputPin()->BreakLinkTo(SourceState->GetOutputPin());
+
+	// Add the new incoming transition to the new target state
+	NewTargetState->GetInputPin()->Modify();
+	NewTargetState->GetInputPin()->MakeLinkTo(SourceState->GetOutputPin());
+
+	// Relink the target state of the transition node
+	Pins[1]->Modify();
+	Pins[1]->BreakLinkTo(TargetStateBeforeRelinking->GetInputPin());
+	Pins[1]->MakeLinkTo(NewTargetState->GetInputPin());
+}
+
+TArray<UAnimStateTransitionNode*> UAnimStateTransitionNode::GetListTransitionNodesToRelink(UEdGraphPin* SourcePin, UEdGraphPin* OldTargetPin, const TArray<UEdGraphNode*>& InSelectedGraphNodes)
+{
+	UAnimStateNodeBase* SourceState = Cast<UAnimStateNodeBase>(SourcePin->GetOwningNode());
+	if (SourceState == nullptr || SourceState->GetInputPin() == nullptr || SourceState->GetOutputPin() == nullptr)
+	{
+		return {};
+	}
+
+	// Collect all transition nodes starting at the source state
+	TArray<UAnimStateTransitionNode*> TransitionNodeCandidates;
+	SourceState->GetTransitionList(TransitionNodeCandidates);
+
+	// Remove the transition nodes from the candidates that are linked to a different target state.
+	for (int i = TransitionNodeCandidates.Num() - 1; i >= 0; i--)
+	{
+		UAnimStateTransitionNode* CurrentTransition = TransitionNodeCandidates[i];
+
+		// Get the actual target states from the transition nodes
+		UEdGraphNode* TransitionTargetNode = CurrentTransition->GetNextState();
+		UAnimStateTransitionNode* CastedOldTarget = Cast<UAnimStateTransitionNode>(OldTargetPin->GetOwningNode());
+		UEdGraphNode* OldTargetNode = CastedOldTarget->GetNextState();
+
+		// Compare the target states rather than comparing against the transition nodes
+		if (TransitionTargetNode != OldTargetNode)
+		{
+			TransitionNodeCandidates.Remove(CurrentTransition);
+		}
+	}
+
+	// Collect the subset of selected transitions from the list of possible transitions to be relinked
+	TSet<UAnimStateTransitionNode*> SelectedTransitionNodes;
+	for (UEdGraphNode* GraphNode : InSelectedGraphNodes)
+	{
+		UAnimStateTransitionNode* TransitionNode = Cast<UAnimStateTransitionNode>(GraphNode);
+		if (!TransitionNode)
+		{
+			continue;
+		}
+
+		if (TransitionNodeCandidates.Find(TransitionNode) != INDEX_NONE)
+		{
+			SelectedTransitionNodes.Add(TransitionNode);
+		}
+	}
+
+	TArray<UAnimStateTransitionNode*> Result;
+	Result.Reserve(TransitionNodeCandidates.Num());
+	for (UAnimStateTransitionNode* TransitionNode : TransitionNodeCandidates)
+	{
+		// Only relink the selected transitions. If none are selected, relink them all.
+		if (!SelectedTransitionNodes.IsEmpty() && SelectedTransitionNodes.Find(TransitionNode) == nullptr)
+		{
+			continue;
+		}
+
+		Result.Add(TransitionNode);
+	}
+
+	return Result;
+}
+
 void UAnimStateTransitionNode::PrepareForCopying()
 {
 	Super::PrepareForCopying();
diff --git a/Engine/Source/Editor/AnimGraph/Private/AnimationStateMachineSchema.cpp b/Engine/Source/Editor/AnimGraph/Private/AnimationStateMachineSchema.cpp
--- a/Engine/Source/Editor/AnimGraph/Private/AnimationStateMachineSchema.cpp
+++ b/Engine/Source/Editor/AnimGraph/Private/AnimationStateMachineSchema.cpp
@@ -151,6 +151,7 @@ const FPinConnectionResponse UAnimationStateMachineSchema::CanCreateConnection(c
 	const bool bPinAIsStateNode = PinA->GetOwningNode()->IsA(UAnimStateNodeBase::StaticClass());
 	const bool bPinBIsStateNode = PinB->GetOwningNode()->IsA(UAnimStateNodeBase::StaticClass());
 
+	// Special case handling for entry states: Only allow creating connections starting at the entry state.
 	if (bPinAIsEntry || bPinBIsEntry)
 	{
 		if (bPinAIsEntry && bPinBIsStateNode)
@@ -158,15 +159,9 @@ const FPinConnectionResponse UAnimationStateMachineSchema::CanCreateConnection(c
 			return FPinConnectionResponse(CONNECT_RESPONSE_BREAK_OTHERS_A, TEXT(""));
 		}
 
-		if (bPinBIsEntry && bPinAIsStateNode)
-		{
-			return FPinConnectionResponse(CONNECT_RESPONSE_BREAK_OTHERS_B, TEXT(""));
-		}
-
 		return FPinConnectionResponse(CONNECT_RESPONSE_DISALLOW, TEXT("Entry must connect to a state node"));
 	}
 
-
 	const bool bPinAIsTransition = PinA->GetOwningNode()->IsA(UAnimStateTransitionNode::StaticClass());
 	const bool bPinBIsTransition = PinB->GetOwningNode()->IsA(UAnimStateTransitionNode::StaticClass());
 
@@ -270,6 +265,87 @@ bool UAnimationStateMachineSchema::CreateAutomaticConversionNodeAndConnections(U
 	return false;
 }
 
+bool UAnimationStateMachineSchema::TryRelinkConnectionTarget(UEdGraphPin* SourcePin, UEdGraphPin* OldTargetPin, UEdGraphPin* NewTargetPin, const TArray<UEdGraphNode*>& InSelectedGraphNodes) const
+{
+	const FPinConnectionResponse Response = CanCreateConnection(SourcePin, NewTargetPin);
+	if (Response.Response == ECanCreateConnectionResponse::CONNECT_RESPONSE_DISALLOW)
+	{
+		return false;
+	}
+
+	UAnimStateNodeBase* OldTargetState = Cast<UAnimStateNodeBase>(OldTargetPin->GetOwningNode());
+	UAnimStateNodeBase* NewTargetState = Cast<UAnimStateNodeBase>(NewTargetPin->GetOwningNode());
+	if (OldTargetState == nullptr || OldTargetState->GetInputPin() == nullptr || OldTargetState->GetOutputPin() == nullptr ||
+		NewTargetState == nullptr || NewTargetState->GetInputPin() == nullptr || NewTargetState->GetOutputPin() == nullptr)
+	{
+		return false;
+	}
+
+	// In the case we are relinking the transition starting at the entry state, the SourceState is nullptr. Special case handling.
+	UAnimStateEntryNode* EntryState = Cast<UAnimStateEntryNode>(SourcePin->GetOwningNode());
+	if (EntryState)
+	{
+		// Remove the incoming transition from the previous target state
+		OldTargetPin->Modify();
+		OldTargetPin->LinkedTo.Remove(SourcePin);
+		SourcePin->Modify();
+		SourcePin->LinkedTo.Remove(OldTargetPin);
+
+		// Add the new incoming transition to the new target state
+		TryCreateConnection(SourcePin, NewTargetPin);
+
+		UBlueprint* Blueprint = FBlueprintEditorUtils::FindBlueprintForGraphChecked(EntryState->GetGraph());
+		FBlueprintEditorUtils::MarkBlueprintAsStructurallyModified(Blueprint);
+
+		return true;
+	}
+
+	// Collect all transition nodes starting at the source state, filter them by the transitions and perform the actual relink operation.
+	const TArray<UAnimStateTransitionNode*> TransitionNodes = UAnimStateTransitionNode::GetListTransitionNodesToRelink(SourcePin, OldTargetPin, InSelectedGraphNodes);
+	for (UAnimStateTransitionNode* TransitionNode : TransitionNodes)
+	{
+		TransitionNode->RelinkHead(NewTargetState);
+	}
+
+	// In case one or more transitions got relinked, inform the blueprint about the changes
+#if WITH_EDITOR
+	if (!TransitionNodes.IsEmpty())
+	{
+		//UBlueprint* Blueprint = FBlueprintEditorUtils::FindBlueprintForGraphChecked(OneRelinkedTransition->GetBoundGraph());
+		UBlueprint* Blueprint = FBlueprintEditorUtils::FindBlueprintForGraphChecked(TransitionNodes[0]->GetBoundGraph());
+		FBlueprintEditorUtils::MarkBlueprintAsStructurallyModified(Blueprint);
+
+
+		SourcePin->GetOwningNode()->PinConnectionListChanged(SourcePin);
+		OldTargetPin->GetOwningNode()->PinConnectionListChanged(OldTargetPin);
+		NewTargetPin->GetOwningNode()->PinConnectionListChanged(NewTargetPin);
+	}
+#endif//#if WITH_EDITOR
+
+	return true;
+}
+
+bool UAnimationStateMachineSchema::IsConnectionRelinkingAllowed(UEdGraphPin* InPin) const
+{
+	if (!InPin || !InPin->GetOwningNode())
+	{
+		return false;
+	}
+
+	return true;
+}
+
+const FPinConnectionResponse UAnimationStateMachineSchema::CanRelinkConnectionToPin(const UEdGraphPin* OldSourcePin, const UEdGraphPin* TargetPinCandidate) const
+{
+	FPinConnectionResponse Response = CanCreateConnection(OldSourcePin, TargetPinCandidate);
+	if (Response.Response != CONNECT_RESPONSE_DISALLOW)
+	{
+		Response.Message = FText::FromString("Relink transition");
+	}
+
+	return Response;
+}
+
 void UAnimationStateMachineSchema::GetGraphContextActions(FGraphContextMenuBuilder& ContextMenuBuilder) const
 {
 	// Add state node
diff --git a/Engine/Source/Editor/AnimGraph/Public/AnimStateTransitionNode.h b/Engine/Source/Editor/AnimGraph/Public/AnimStateTransitionNode.h
--- a/Engine/Source/Editor/AnimGraph/Public/AnimStateTransitionNode.h
+++ b/Engine/Source/Editor/AnimGraph/Public/AnimStateTransitionNode.h
@@ -141,6 +141,18 @@ class UAnimStateTransitionNode : public UAnimStateNodeBase
 	ANIMGRAPH_API UAnimStateNodeBase* GetPreviousState() const;
 	ANIMGRAPH_API UAnimStateNodeBase* GetNextState() const;
 	ANIMGRAPH_API void CreateConnections(UAnimStateNodeBase* PreviousState, UAnimStateNodeBase* NextState);
+
+	/**
+	 * Relink transition head (where the arrow is of a state transition) to a new state.
+	 * @param[in] NewTargetState The new transition target.
+	 */
+	ANIMGRAPH_API void RelinkHead(UAnimStateNodeBase* NewTargetState);
+
+	/**
+	 * Helper function to gather the transition nodes to be relinked by taking the graph selection into account as well.
+	 * For example when relinking a transition holding several transition nodes but only a few are selected to be relinked.
+	 */
+	ANIMGRAPH_API static TArray<UAnimStateTransitionNode*> GetListTransitionNodesToRelink(UEdGraphPin* SourcePin, UEdGraphPin* OldTargetPin, const TArray<UEdGraphNode*>& InSelectedGraphNodes);
 	
 	ANIMGRAPH_API bool IsBoundGraphShared() const;
 
diff --git a/Engine/Source/Editor/AnimGraph/Public/AnimationStateMachineSchema.h b/Engine/Source/Editor/AnimGraph/Public/AnimationStateMachineSchema.h
--- a/Engine/Source/Editor/AnimGraph/Public/AnimationStateMachineSchema.h
+++ b/Engine/Source/Editor/AnimGraph/Public/AnimationStateMachineSchema.h
@@ -75,6 +75,9 @@ class UAnimationStateMachineSchema : public UEdGraphSchema
 	virtual const FPinConnectionResponse CanCreateConnection(const UEdGraphPin* A, const UEdGraphPin* B) const override;
 	virtual bool TryCreateConnection(UEdGraphPin* A, UEdGraphPin* B) const override;
 	virtual bool CreateAutomaticConversionNodeAndConnections(UEdGraphPin* PinA, UEdGraphPin* PinB) const override;
+	virtual bool TryRelinkConnectionTarget(UEdGraphPin* SourcePin, UEdGraphPin* OldTargetPin, UEdGraphPin* NewTargetPin, const TArray<UEdGraphNode*>& InSelectedGraphNodes) const override;
+	virtual bool IsConnectionRelinkingAllowed(UEdGraphPin* InPin) const override;
+	virtual const FPinConnectionResponse CanRelinkConnectionToPin(const UEdGraphPin* OldSourcePin, const UEdGraphPin* TargetPinCandidate) const override;
 	virtual void GetGraphContextActions(FGraphContextMenuBuilder& ContextMenuBuilder) const override;
 	virtual EGraphType GetGraphType(const UEdGraph* TestEdGraph) const override;
 	virtual void GetContextMenuActions(class UToolMenu* Menu, class UGraphNodeContextMenuContext* Context) const override;
diff --git a/Engine/Source/Editor/AnimationBlueprintEditor/Private/StateMachineConnectionDrawingPolicy.cpp b/Engine/Source/Editor/AnimationBlueprintEditor/Private/StateMachineConnectionDrawingPolicy.cpp
--- a/Engine/Source/Editor/AnimationBlueprintEditor/Private/StateMachineConnectionDrawingPolicy.cpp
+++ b/Engine/Source/Editor/AnimationBlueprintEditor/Private/StateMachineConnectionDrawingPolicy.cpp
@@ -20,6 +20,9 @@
 #include "SGraphPin.h"
 #include "Styling/SlateBrush.h"
 #include "Templates/Casts.h"
+#include "Math/UnrealMathUtility.h"
+#include "Brushes/SlateRoundedBoxBrush.h"
+#include "Styling/StyleColors.h"
 
 class FSlateRect;
 class SWidget;
@@ -36,6 +39,7 @@ FStateMachineConnectionDrawingPolicy::FStateMachineConnectionDrawingPolicy(int32
 
 void FStateMachineConnectionDrawingPolicy::DetermineWiringStyle(UEdGraphPin* OutputPin, UEdGraphPin* InputPin, /*inout*/ FConnectionParams& Params)
 {
+	Params.bUserFlag2 = true;
 	Params.AssociatedPin1 = OutputPin;
 	Params.AssociatedPin2 = InputPin;
 	Params.WireThickness = 1.5f;
@@ -44,8 +48,8 @@ void FStateMachineConnectionDrawingPolicy::DetermineWiringStyle(UEdGraphPin* Out
 	{
 		if (UAnimStateTransitionNode* TransNode = Cast<UAnimStateTransitionNode>(InputPin->GetOwningNode()))
 		{
-			Params.WireColor = SGraphNodeAnimTransition::StaticGetTransitionColor(TransNode, HoveredPins.Contains(InputPin));
-
+			const bool IsInputPinHovered = HoveredPins.Contains(InputPin);
+			Params.WireColor = SGraphNodeAnimTransition::StaticGetTransitionColor(TransNode, IsInputPinHovered);
 			Params.bUserFlag1 = TransNode->Bidirectional;
 		}
 	}
@@ -55,6 +59,33 @@ void FStateMachineConnectionDrawingPolicy::DetermineWiringStyle(UEdGraphPin* Out
 	{
 		ApplyHoverDeemphasis(OutputPin, InputPin, /*inout*/ Params.WireThickness, /*inout*/ Params.WireColor);
 	}
+
+	// Make the transition that is currently relinked, semi-transparent.
+	for (const FRelinkConnection& Connection : RelinkConnections)
+	{
+		if (InputPin == nullptr || OutputPin == nullptr)
+		{
+			continue;
+		}
+		const FGraphPinHandle SourcePinHandle = Connection.SourcePin;
+		const FGraphPinHandle TargetPinHandle = Connection.TargetPin;
+
+		// Skip all transitions that don't start at the node our dragged and relink transition starts from
+		if (OutputPin->GetOwningNode()->NodeGuid == SourcePinHandle.NodeGuid)
+		{
+			// Safety check to verify if the node is a transition node
+			if (UAnimStateTransitionNode* TransitonNode = Cast<UAnimStateTransitionNode>(InputPin->GetOwningNode()))
+			{
+				if (UEdGraphPin* TransitionOutputPin = TransitonNode->GetOutputPin())
+				{
+					if (TargetPinHandle.NodeGuid == TransitionOutputPin->GetOwningNode()->NodeGuid)
+					{
+						Params.WireColor.A *= 0.2f;
+					}
+				}
+			}
+		}
+	}
 }
 
 void FStateMachineConnectionDrawingPolicy::DetermineLinkGeometry(
@@ -124,6 +155,7 @@ void FStateMachineConnectionDrawingPolicy::DrawPreviewConnector(const FGeometry&
 	const FVector2D SeedPoint = EndPoint;
 	const FVector2D AdjustedStartPoint = FGeometryHelper::FindClosestPointOnGeom(PinGeometry, SeedPoint);
 
+	Params.bUserFlag2 = false; // bUserFlag2 is used to indicate whether the drawn arrow is a preview transition (the temporary transition when creating or relinking).
 	DrawSplineWithArrow(AdjustedStartPoint, EndPoint, Params);
 }
 
@@ -153,14 +185,93 @@ void FStateMachineConnectionDrawingPolicy::Internal_DrawLineWithArrow(const FVec
 	const FVector2D LengthBias = ArrowRadius.X * UnitDelta;
 	const FVector2D StartPoint = StartAnchorPoint + DirectionBias + LengthBias;
 	const FVector2D EndPoint = EndAnchorPoint + DirectionBias - LengthBias;
+	FLinearColor ArrowHeadColor = Params.WireColor;
 
 	// Draw a line/spline
 	DrawConnection(WireLayerID, StartPoint, EndPoint, Params);
 
-	// Draw the arrow
 	const FVector2D ArrowDrawPos = EndPoint - ArrowRadius;
 	const double AngleInRadians = FMath::Atan2(DeltaPos.Y, DeltaPos.X);
 
+	// Draw the transition grab handles in case the mouse is hovering the transition
+	bool StartHovered = false;
+	bool EndHovered = false;
+	const FVector FVecMousePos = FVector(LocalMousePosition.X, LocalMousePosition.Y, 0.0f);
+	const FVector ClosestPoint = FMath::ClosestPointOnSegment(FVecMousePos,
+		FVector(StartPoint.X, StartPoint.Y, 0.0f),
+		FVector(EndPoint.X, EndPoint.Y, 0.0f));
+	if ((ClosestPoint - FVecMousePos).Length() < RelinkHandleHoverRadius)
+	{
+		StartHovered = FVector2D(StartPoint - LocalMousePosition).Length() < RelinkHandleHoverRadius;
+		EndHovered = FVector2D(EndPoint - LocalMousePosition).Length() < RelinkHandleHoverRadius;
+		FVector2D HoverIndicatorPosition = StartHovered ? StartPoint : EndPoint;
+
+		// Set the hovered pin results. This will be used by the SGraphPanel again.
+		const float SquaredDistToPin1 = (Params.AssociatedPin1 != nullptr) ? (StartPoint - LocalMousePosition).SizeSquared() : FLT_MAX;
+		const float SquaredDistToPin2 = (Params.AssociatedPin2 != nullptr) ? (EndPoint - LocalMousePosition).SizeSquared() : FLT_MAX;
+		if (StartHovered)
+		{
+			SplineOverlapResult = FGraphSplineOverlapResult(Params.AssociatedPin2, Params.AssociatedPin1, SquaredDistToPin2, SquaredDistToPin2, SquaredDistToPin1, true);
+		}
+		else if (EndHovered)
+		{
+			SplineOverlapResult = FGraphSplineOverlapResult(Params.AssociatedPin1, Params.AssociatedPin2, SquaredDistToPin1, SquaredDistToPin1, SquaredDistToPin2, true);
+		}
+
+		// Draw grab handles only in case no relinking operation is performed
+		if (RelinkConnections.IsEmpty() && Params.bUserFlag2)
+		{
+			if (EndHovered)
+			{
+				// Draw solid orange circle behind the arrow head in case the arrow head is hovered (area that enables a relink).
+				FSlateRoundedBoxBrush RoundedBoxBrush = FSlateRoundedBoxBrush(
+					FLinearColor(0.0f, 0.0f, 0.0f, 0.0f), 9.0f, FStyleColors::AccentOrange, 100.0f);
+
+				FSlateDrawElement::MakeBox(DrawElementsList,
+					ArrowLayerID-1, // Draw behind the arrow
+					FPaintGeometry(EndPoint - ArrowRadius, BubbleImage->ImageSize * ZoomFactor, ZoomFactor),
+					&RoundedBoxBrush);
+
+				ArrowHeadColor = FLinearColor::Black;
+			}
+			else
+			{
+				// Draw circle around the arrow in case the transition is hovered (mouse close or over transition line or arrow head).
+				const int CircleLineSegments = 16;
+				const float CircleRadius = 10.0f;
+				const FVector2D CircleCenter = EndPoint - UnitDelta * 2.0f;
+				DrawCircle(CircleCenter, CircleRadius, Params.WireColor, CircleLineSegments);
+			}
+		}
+	}
+
+	// Draw the number of relinked transitions on the preview transition.
+	if (!RelinkConnections.IsEmpty() && !Params.bUserFlag2)
+	{
+		// Get the number of actually relinked transitions.
+		int32 NumRelinkedTransitions = 0;
+		for (const FRelinkConnection& Connection : RelinkConnections)
+		{
+			NumRelinkedTransitions += UAnimStateTransitionNode::GetListTransitionNodesToRelink(Connection.SourcePin, Connection.TargetPin, SelectedGraphNodes).Num();
+
+			if (UAnimStateEntryNode* EntryNode = Cast<UAnimStateEntryNode>(Connection.SourcePin->GetOwningNode()))
+			{
+				NumRelinkedTransitions += 1;
+			}
+		}
+
+		const FVector2D TransitionCenter = StartAnchorPoint + DeltaPos * 0.5f;
+		const FVector2D TextPosition = TransitionCenter + Normal * 15.0f * ZoomFactor;
+
+		FSlateDrawElement::MakeText(
+			DrawElementsList,
+			ArrowLayerID,
+			FPaintGeometry(TextPosition, ArrowImage->ImageSize * ZoomFactor, ZoomFactor),
+			FText::AsNumber(NumRelinkedTransitions),
+			FCoreStyle::Get().GetFontStyle("SmallFont"));
+	}
+
+	// Draw the transition arrow triangle
 	FSlateDrawElement::MakeRotatedBox(
 		DrawElementsList,
 		ArrowLayerID,
@@ -170,8 +281,33 @@ void FStateMachineConnectionDrawingPolicy::Internal_DrawLineWithArrow(const FVec
 		static_cast<float>(AngleInRadians),
 		TOptional<FVector2D>(),
 		FSlateDrawElement::RelativeToElement,
-		Params.WireColor
-		);
+		ArrowHeadColor
+	);
+}
+
+void FStateMachineConnectionDrawingPolicy::DrawCircle(const FVector2D& Center, float Radius, const FLinearColor& Color, const int NumLineSegments)
+{
+	TempPoints.Empty();
+
+	const float NumFloatLineSegments = (float)NumLineSegments;
+	for (int i = 0; i <= NumLineSegments; i++)
+	{
+		const float Angle = (i / NumFloatLineSegments) * TWO_PI;
+
+		FVector2D PointOnCircle;
+		PointOnCircle.X = cosf(Angle) * Radius;
+		PointOnCircle.Y = sinf(Angle) * Radius;
+		TempPoints.Add(PointOnCircle);
+	}
+
+	FSlateDrawElement::MakeLines(
+		DrawElementsList,
+		ArrowLayerID + 1,
+		FPaintGeometry(Center, FVector2D(Radius, Radius) * ZoomFactor, ZoomFactor),
+		TempPoints,
+		ESlateDrawEffect::None,
+		Color
+	);
 }
 
 void FStateMachineConnectionDrawingPolicy::DrawSplineWithArrow(const FGeometry& StartGeom, const FGeometry& EndGeom, const FConnectionParams& Params)
diff --git a/Engine/Source/Editor/AnimationBlueprintEditor/Private/StateMachineConnectionDrawingPolicy.h b/Engine/Source/Editor/AnimationBlueprintEditor/Private/StateMachineConnectionDrawingPolicy.h
--- a/Engine/Source/Editor/AnimationBlueprintEditor/Private/StateMachineConnectionDrawingPolicy.h
+++ b/Engine/Source/Editor/AnimationBlueprintEditor/Private/StateMachineConnectionDrawingPolicy.h
@@ -29,6 +29,8 @@ protected:
 
 	TMap<UEdGraphNode*, int32> NodeWidgetMap;
 public:
+	// @Note FConnectionParams.bUserFlag2: Is used here to indicate whether the drawn arrow is a preview transition or a real one.
+
 	//
 	FStateMachineConnectionDrawingPolicy(int32 InBackLayerID, int32 InFrontLayerID, float ZoomFactor, const FSlateRect& InClippingRect, FSlateWindowElementList& InDrawElements, UEdGraph* InGraphObj);
 
@@ -51,4 +53,10 @@ public:
 
 protected:
 	void Internal_DrawLineWithArrow(const FVector2D& StartAnchorPoint, const FVector2D& EndAnchorPoint, const FConnectionParams& Params);
+
+	// Draw line-based circle (no solid filling)
+	void DrawCircle(const FVector2D& Center, float Radius, const FLinearColor& Color, const int NumLineSegments);
+	TArray<FVector2D> TempPoints;
+
+	static constexpr float RelinkHandleHoverRadius = 20.0f;
 };
diff --git a/Engine/Source/Editor/GraphEditor/Private/DragConnection.cpp b/Engine/Source/Editor/GraphEditor/Private/DragConnection.cpp
--- a/Engine/Source/Editor/GraphEditor/Private/DragConnection.cpp
+++ b/Engine/Source/Editor/GraphEditor/Private/DragConnection.cpp
@@ -43,9 +43,10 @@ TSharedRef<FDragConnection> FDragConnection::New(const TSharedRef<SGraphPanel>&
 	return Operation;
 }
 
-void FDragConnection::OnDrop( bool bDropWasHandled, const FPointerEvent& MouseEvent )
+void FDragConnection::OnDrop(bool bDropWasHandled, const FPointerEvent& MouseEvent)
 {
 	GraphPanel->OnStopMakingConnection();
+	GraphPanel->OnEndRelinkConnection();
 
 	Super::OnDrop(bDropWasHandled, MouseEvent);
 }
@@ -77,8 +78,17 @@ void FDragConnection::HoverTargetChanged()
 				// The Graph object in which the pins reside.
 				UEdGraph* GraphObj = StartingPinObj->GetOwningNode()->GetGraph();
 
-				// Determine what the schema thinks about the wiring action
-				const FPinConnectionResponse Response = GraphObj->GetSchema()->CanCreateConnection( StartingPinObj, TargetPinObj );
+				// Determine what the schema thinks
+				FPinConnectionResponse Response;
+				switch (DragMode)
+				{
+				case EDragMode::CreateConnection:
+					Response = GraphObj->GetSchema()->CanCreateConnection( StartingPinObj, TargetPinObj );
+					break;
+				case EDragMode::RelinkConnection:
+					Response = GraphObj->GetSchema()->CanRelinkConnectionToPin(StartingPinObj, TargetPinObj);
+					break;
+				}
 
 				if (Response.Response == ECanCreateConnectionResponse::CONNECT_RESPONSE_DISALLOW)
 				{
@@ -93,7 +103,7 @@ void FDragConnection::HoverTargetChanged()
 			}
 		}
 	}
-	else if(UEdGraphNode* TargetNodeObj = GetHoveredNode())
+	else if(GetHoveredNode() && DragMode == EDragMode::CreateConnection)
 	{
 		TArray<UEdGraphPin*> ValidSourcePins;
 		ValidateGraphPinList(/*out*/ ValidSourcePins);
@@ -103,7 +113,7 @@ void FDragConnection::HoverTargetChanged()
 		{
 			FPinConnectionResponse Response;
 			FText ResponseText;
-			if (StartingPinObj->GetOwningNode() != TargetNodeObj && StartingPinObj->GetSchema()->SupportsDropPinOnNode(TargetNodeObj, StartingPinObj->PinType, StartingPinObj->Direction, ResponseText))
+			if (StartingPinObj->GetOwningNode() != GetHoveredNode() && StartingPinObj->GetSchema()->SupportsDropPinOnNode(GetHoveredNode(), StartingPinObj->PinType, StartingPinObj->Direction, ResponseText))
 			{
 				Response.Response = ECanCreateConnectionResponse::CONNECT_RESPONSE_MAKE;
 			}
@@ -120,7 +130,7 @@ void FDragConnection::HoverTargetChanged()
 			}
 		}
 	}
-	else if(UEdGraph* CurrentHoveredGraph = GetHoveredGraph())
+	else if(GetHoveredGraph() && DragMode == EDragMode::CreateConnection)
 	{
 		TArray<UEdGraphPin*> ValidSourcePins;
 		ValidateGraphPinList(/*out*/ ValidSourcePins);
@@ -128,7 +138,7 @@ void FDragConnection::HoverTargetChanged()
 		for (UEdGraphPin* StartingPinObj : ValidSourcePins)
 		{
 			// Let the schema describe the connection we might make
-			FPinConnectionResponse Response = CurrentHoveredGraph->GetSchema()->CanCreateNewNodes(StartingPinObj);
+			FPinConnectionResponse Response = GetHoveredGraph()->GetSchema()->CanCreateNewNodes(StartingPinObj);
 			if(!Response.Message.IsEmpty())
 			{
 				UniqueMessages.AddUnique(Response);
@@ -137,7 +147,7 @@ void FDragConnection::HoverTargetChanged()
 	}
 
 	// Let the user know the status of dropping now
-	if (UniqueMessages.Num() == 0)
+	if (UniqueMessages.Num() == 0 && DragMode == EDragMode::CreateConnection)
 	{
 		// Display the place a new node icon, we're not over a valid pin and have no message from the schema
 		SetSimpleFeedbackMessage(
@@ -206,64 +216,154 @@ FDragConnection::FDragConnection(const TSharedRef<SGraphPanel>& GraphPanelIn, co
 	: GraphPanel(GraphPanelIn)
 	, DraggingPins(DraggedPinsIn)
 	, DecoratorAdjust(FSlateApplication::Get().GetCursorSize())
-{	
-	if (DraggingPins.Num() > 0)
+{
+	// Handle connection relinking, if supported by the schema
+	const FGraphSplineOverlapResult& PreviousFrameSplineOverlap = GraphPanel->GetPreviousFrameSplineOverlap();
+	FGraphPinHandle Pin1Handle = PreviousFrameSplineOverlap.GetPin1Handle();
+	if (Pin1Handle.IsValid())
 	{
-		const UEdGraphPin* PinObj = FDraggedPinTable::TConstIterator(DraggedPinsIn)->GetPinObj(*GraphPanelIn);
-		if (PinObj && PinObj->Direction == EGPD_Input)
+		UEdGraphPin* SourcePin = Pin1Handle.GetPinObj(*GraphPanel.Get());
+		if (SourcePin)
 		{
-			DecoratorAdjust *= FVector2D(-1.0f, 1.0f);
+			const UEdGraphSchema* Schema = SourcePin->GetSchema();
+			if (Schema->IsConnectionRelinkingAllowed(SourcePin))
+			{
+				FGraphPinHandle Pin2Handle = PreviousFrameSplineOverlap.GetPin2Handle();
+				UEdGraphPin* TargetPin = Pin2Handle.GetPinObj(*GraphPanel.Get());
+				if (TargetPin)
+				{
+					SourcePinHandle = SourcePin;
+					TargetPinHandle = TargetPin;
+					DragMode = RelinkConnection;
+				}
+			}
 		}
 	}
 
-	for (const FGraphPinHandle& DraggedPin : DraggedPinsIn)
+	switch (DragMode)
 	{
-		GraphPanelIn->OnBeginMakingConnection(DraggedPin);
-	}
-}
+	case EDragMode::CreateConnection:
+	{
+		if (DraggingPins.Num() > 0)
+		{
+			const UEdGraphPin* PinObj = FDraggedPinTable::TConstIterator(DraggedPinsIn)->GetPinObj(*GraphPanelIn);
+			if (PinObj && PinObj->Direction == EGPD_Input)
+			{
+				DecoratorAdjust *= FVector2D(-1.0f, 1.0f);
+			}
+		}
 
-FReply FDragConnection::DroppedOnPin(FVector2D ScreenPosition, FVector2D GraphPosition)
-{
-	TArray<UEdGraphPin*> ValidSourcePins;
-	ValidateGraphPinList(/*out*/ ValidSourcePins);
+		for (const FGraphPinHandle& DraggedPin : DraggedPinsIn)
+		{
+			GraphPanelIn->OnBeginMakingConnection(DraggedPin);
+		}
+		break;
+	}
 
-	// store the pins as pin tuples since the structure of the
-	// graph may change during the creation of a connection
-	TArray<FEdGraphPinHandle> ValidSourcePinHandles;
-	for(const UEdGraphPin* ValidSourcePin : ValidSourcePins)
+	case EDragMode::RelinkConnection:
 	{
-		ValidSourcePinHandles.Add(ValidSourcePin);
-	}
+		if (SourcePinHandle.IsValid())
+		{
+			const UEdGraphPin* PinObj = SourcePinHandle.GetPinObj(*GraphPanel.Get());
+			if (PinObj && PinObj->Direction == EGPD_Input)
+			{
+				DecoratorAdjust *= FVector2D(-1.0f, 1.0f);
+			}
+		}
 
-	const FScopedTransaction Transaction( NSLOCTEXT("UnrealEd", "GraphEd_CreateConnection", "Create Pin Link") );
+		GraphPanel->OnBeginRelinkConnection(SourcePinHandle, TargetPinHandle);
+		break;
+	}
+	}
+}
 
-	FEdGraphPinHandle PinB(GetHoveredPin());
+FReply FDragConnection::DroppedOnPin(FVector2D ScreenPosition, FVector2D GraphPosition)
+{
 	bool bError = false;
 	TSet<FEdGraphNodeHandle> NodeList;
 
-	for (const FEdGraphPinHandle& PinA : ValidSourcePinHandles)
+	switch (DragMode)
 	{
-		if ((PinA.GetPin() != NULL) && (PinB.GetPin() != NULL))
+	case EDragMode::CreateConnection:
+	{
+		TArray<UEdGraphPin*> ValidSourcePins;
+		ValidateGraphPinList(/*out*/ ValidSourcePins);
+
+		// store the pins as pin tuples since the structure of the
+		// graph may change during the creation of a connection
+		TArray<FEdGraphPinHandle> ValidSourcePinHandles;
+		for (const UEdGraphPin* ValidSourcePin : ValidSourcePins)
 		{
-			const UEdGraph* MyGraphObj = PinA.GetGraph();
+			ValidSourcePinHandles.Add(ValidSourcePin);
+		}
+
+		const FScopedTransaction Transaction(NSLOCTEXT("UnrealEd", "GraphEd_CreateConnection", "Create Pin Link"));
 
-			// the pin may change during the creation of the link
-			if (MyGraphObj->GetSchema()->TryCreateConnection(PinA.GetPin(), PinB.GetPin()))
+		FEdGraphPinHandle PinB(GetHoveredPin());
+
+		for (const FEdGraphPinHandle& PinA : ValidSourcePinHandles)
+		{
+			if ((PinA.GetPin() != NULL) && (PinB.GetPin() != NULL))
 			{
-				if (PinA.GetPin() && !PinA.GetPin()->IsPendingKill())
-				{
-					NodeList.Add(PinA.GetPin()->GetOwningNode());
-				}
-				if (PinB.GetPin() && !PinB.GetPin()->IsPendingKill())
+				const UEdGraph* MyGraphObj = PinA.GetGraph();
+
+				// the pin may change during the creation of the link
+				if (MyGraphObj->GetSchema()->TryCreateConnection(PinA.GetPin(), PinB.GetPin()))
 				{
-					NodeList.Add(PinB.GetNode());
+					if (PinA.GetPin() && !PinA.GetPin()->IsPendingKill())
+					{
+						NodeList.Add(PinA.GetPin()->GetOwningNode());
+					}
+					if (PinB.GetPin() && !PinB.GetPin()->IsPendingKill())
+					{
+						NodeList.Add(PinB.GetNode());
+					}
 				}
 			}
+			else
+			{
+				bError = true;
+			}
 		}
-		else
+		break;
+	}
+
+	case EDragMode::RelinkConnection:
+	{
+		const UEdGraphPin* SourceGraphPin = SourcePinHandle.GetPinObj(*GraphPanel);
+		FEdGraphPinHandle SourceGraphPinHandle = SourceGraphPin;
+
+		UEdGraphPin* TargetGraphPin = TargetPinHandle.GetPinObj(*GraphPanel);
+		FEdGraphPinHandle TargetGraphPinHandle = SourceGraphPin;
+
+		FEdGraphPinHandle HoveredGraphPinHandle(GetHoveredPin());
+		UEdGraphPin* HoveredGraphPin = HoveredGraphPinHandle.GetPin();
+		
+		if (SourceGraphPin && TargetGraphPin && HoveredGraphPin)
 		{
-			bError = true;
+			const UEdGraph* Graph = SourceGraphPinHandle.GetGraph();
+			if (Graph)
+			{
+				const FScopedTransaction Transaction(NSLOCTEXT("UnrealEd", "GraphEd_RelinkConnection", "Relink Connection"));
+
+				const TArray<UEdGraphNode*> SelectedGraphNodes = GraphPanel->GetSelectedGraphNodes();
+				bError = !Graph->GetSchema()->TryRelinkConnectionTarget(SourceGraphPinHandle.GetPin(), TargetGraphPin, HoveredGraphPinHandle.GetPin(), SelectedGraphNodes);				if (!bError)
+				{
+					// Send all nodes that received a new pin connection a notification
+					TArray<FEdGraphPinHandle> ModifiedGraphPinHandles = { SourceGraphPinHandle, TargetGraphPinHandle, HoveredGraphPinHandle };
+					for (const FEdGraphPinHandle& PinHandle : ModifiedGraphPinHandles)
+					{
+						if (PinHandle.GetPin() && !PinHandle.GetPin()->IsPendingKill() && PinHandle.GetNode())
+						{
+							NodeList.Add(PinHandle.GetNode());
+						}
+					}
+				}
+			}
 		}
+
+		break;
+	}
 	}
 
 	// Send all nodes that received a new pin connection a notification
@@ -285,6 +385,11 @@ FReply FDragConnection::DroppedOnPin(FVector2D ScreenPosition, FVector2D GraphPo
 
 FReply FDragConnection::DroppedOnNode(FVector2D ScreenPosition, FVector2D GraphPosition)
 {
+	if (DragMode != EDragMode::CreateConnection)
+	{
+		return FReply::Unhandled();
+	}
+
 	bool bHandledPinDropOnNode = false;
 	UEdGraphNode* NodeOver = GetHoveredNode();
 
@@ -339,6 +444,11 @@ FReply FDragConnection::DroppedOnNode(FVector2D ScreenPosition, FVector2D GraphP
 
 FReply FDragConnection::DroppedOnPanel( const TSharedRef< SWidget >& Panel, FVector2D ScreenPosition, FVector2D GraphPosition, UEdGraph& Graph)
 {
+	if (DragMode != EDragMode::CreateConnection)
+	{
+		return FReply::Unhandled();
+	}
+
 	// Gather any source drag pins
 	TArray<UEdGraphPin*> PinObjects;
 	ValidateGraphPinList(/*out*/ PinObjects);
diff --git a/Engine/Source/Editor/GraphEditor/Private/DragConnection.h b/Engine/Source/Editor/GraphEditor/Private/DragConnection.h
--- a/Engine/Source/Editor/GraphEditor/Private/DragConnection.h
+++ b/Engine/Source/Editor/GraphEditor/Private/DragConnection.h
@@ -22,6 +22,12 @@ class FDragConnection : public FGraphEditorDragDropAction
 public:
 	DRAG_DROP_OPERATOR_TYPE(FDragConnection, FGraphEditorDragDropAction)
 
+	enum EDragMode : uint8
+	{
+		CreateConnection = 0,
+		RelinkConnection
+	};
+
 	typedef TArray<FGraphPinHandle> FDraggedPinTable;
 	static TSharedRef<FDragConnection> New(const TSharedRef<SGraphPanel>& InGraphPanel, const FDraggedPinTable& InStartingPins);
 
@@ -51,7 +57,11 @@ protected:
 protected:
 	TSharedPtr<SGraphPanel> GraphPanel;
 	FDraggedPinTable DraggingPins;
+	EDragMode DragMode = EDragMode::CreateConnection;
 
 	/** Offset information for the decorator widget */
 	FVector2D DecoratorAdjust;
+
+	FGraphPinHandle SourcePinHandle;
+	FGraphPinHandle TargetPinHandle;
 };
diff --git a/Engine/Source/Editor/GraphEditor/Private/SGraphPanel.cpp b/Engine/Source/Editor/GraphEditor/Private/SGraphPanel.cpp
--- a/Engine/Source/Editor/GraphEditor/Private/SGraphPanel.cpp
+++ b/Engine/Source/Editor/GraphEditor/Private/SGraphPanel.cpp
@@ -468,7 +468,6 @@ int32 SGraphPanel::OnPaint( const FPaintArgs& Args, const FGeometry& AllottedGeo
 
 	MaxLayerId += 1;
 
-
 	// Draw connections between pins 
 	if (Children.Num() > 0 )
 	{
@@ -495,6 +494,12 @@ int32 SGraphPanel::OnPaint( const FPaintArgs& Args, const FGeometry& AllottedGeo
 		ConnectionDrawingPolicy->SetMarkedPin(MarkedPin);
 		ConnectionDrawingPolicy->SetMousePosition(AllottedGeometry.LocalToAbsolute(SavedMousePosForOnPaintEventLocalSpace));
 
+		if (IsRelinkingConnection())
+		{
+			ConnectionDrawingPolicy->SetRelinkConnections(RelinkConnections);
+			ConnectionDrawingPolicy->SetSelectedNodes(GetSelectedGraphNodes());
+		}
+
 		// Get the set of pins for all children and synthesize geometry for culled out pins so lines can be drawn to them.
 		TMap<TSharedRef<SWidget>, FArrangedWidget> PinGeometries;
 		TSet< TSharedRef<SWidget> > VisiblePins;
@@ -810,8 +815,11 @@ FReply SGraphPanel::OnKeyDown( const FGeometry& MyGeometry, const FKeyEvent& InK
 
 FReply SGraphPanel::OnMouseButtonDown(const FGeometry& MyGeometry, const FPointerEvent& MouseEvent)
 {
-	if ((MouseEvent.GetEffectingButton() == EKeys::LeftMouseButton) && (MouseEvent.IsAltDown() || MouseEvent.IsControlDown()))
+	if ((MouseEvent.GetEffectingButton() == EKeys::LeftMouseButton))
 	{
+		SGraphPin* BestPinFromHoveredSpline = GetBestPinFromHoveredSpline();
+		FGraphPinHandle Pin1Handle = PreviousFrameSplineOverlap.GetPin1Handle();
+
 		// Intercept alt-left clicking on the hovered spline for targeted break link
 		UEdGraphPin* Pin1;
 		UEdGraphPin* Pin2;
@@ -820,15 +828,44 @@ FReply SGraphPanel::OnMouseButtonDown(const FGeometry& MyGeometry, const FPointe
 			const UEdGraphSchema* Schema = GraphObj->GetSchema();
 			Schema->BreakSinglePinLink(Pin1, Pin2);
 		}
-		else if (SGraphPin* BestPinFromHoveredSpline = GetBestPinFromHoveredSpline())
+		else if (BestPinFromHoveredSpline && MouseEvent.IsControlDown())
 		{
 			return BestPinFromHoveredSpline->OnPinMouseDown(MyGeometry, MouseEvent);
 		}
+		else if (Pin1Handle.IsValid())
+		{
+			TSharedPtr<class SGraphPin> SourcePin = Pin1Handle.FindInGraphPanel(*this);
+			if (SourcePin.IsValid())
+			{
+				return SourcePin->OnPinMouseDown(MyGeometry, MouseEvent);
+			}
+		}
 	}
 
 	return SNodePanel::OnMouseButtonDown(MyGeometry, MouseEvent);
 }
 
+TArray<UEdGraphNode*> SGraphPanel::GetSelectedGraphNodes() const
+{
+	TArray<UEdGraphNode*> SelectedGraphNodes;
+	SelectedGraphNodes.Reserve(SelectionManager.SelectedNodes.Num());
+
+	for (FGraphPanelSelectionSet::TConstIterator NodeIt(SelectionManager.SelectedNodes); NodeIt; ++NodeIt)
+	{
+		const TSharedRef<SNode>* SelectedNode = NodeToWidgetLookup.Find(*NodeIt);
+		if (SelectedNode)
+		{
+			UEdGraphNode* SelectedGraphNode = Cast<UEdGraphNode>(SelectedNode->Get().GetObjectBeingDisplayed());
+			if (SelectedGraphNode)
+			{
+				SelectedGraphNodes.Add(SelectedGraphNode);
+			}
+		}
+	}
+
+	return SelectedGraphNodes;
+}
+
 FReply SGraphPanel::OnMouseButtonUp(const FGeometry& MyGeometry, const FPointerEvent& MouseEvent)
 {
 	if (!NodeUnderMousePtr.IsValid() && !Marquee.IsValid() && (MouseEvent.GetEffectingButton() == EKeys::LeftMouseButton) && (MouseEvent.IsShiftDown()))
@@ -1264,6 +1301,23 @@ void SGraphPanel::OnStopMakingConnection(bool bForceStop)
 	}
 }
 
+void SGraphPanel::OnBeginRelinkConnection(const FGraphPinHandle& InSourcePinHandle, const FGraphPinHandle& InTargetPinHandle)
+{
+	RelinkConnections.Add({ InSourcePinHandle.GetPinObj(*this), InTargetPinHandle.GetPinObj(*this) });
+	OnBeginMakingConnection(InSourcePinHandle);
+}
+
+void SGraphPanel::OnEndRelinkConnection(bool bForceStop)
+{
+	OnStopMakingConnection(bForceStop);
+	RelinkConnections.Empty();
+}
+
+bool SGraphPanel::IsRelinkingConnection() const
+{
+	return (RelinkConnections.IsEmpty() == false);
+}
+
 void SGraphPanel::PreservePinPreviewUntilForced()
 {
 	bPreservePinPreviewConnection = true;
diff --git a/Engine/Source/Editor/GraphEditor/Public/ConnectionDrawingPolicy.h b/Engine/Source/Editor/GraphEditor/Public/ConnectionDrawingPolicy.h
--- a/Engine/Source/Editor/GraphEditor/Public/ConnectionDrawingPolicy.h
+++ b/Engine/Source/Editor/GraphEditor/Public/ConnectionDrawingPolicy.h
@@ -109,6 +109,13 @@ public:
 
 	FGraphSplineOverlapResult SplineOverlapResult;
 
+	/** Handle for a currently relinked connection. */
+	struct FRelinkConnection
+	{
+		UEdGraphPin* SourcePin;
+		UEdGraphPin* TargetPin;
+	};
+
 protected:
 	float ZoomFactor; 
 	float HoverDeemphasisDarkFraction;
@@ -119,6 +126,12 @@ protected:
 	TMap<TSharedRef<SWidget>, FArrangedWidget>* PinGeometries;
 	double LastHoverTimeEvent;
 	FVector2D LocalMousePosition;
+
+	/** List of currently relinked connections. */
+	TArray<FRelinkConnection> RelinkConnections;
+
+	/** Selected nodes in the graph panel. */
+	TArray<UEdGraphNode*> SelectedGraphNodes;
 public:
 	virtual ~FConnectionDrawingPolicy() {}
 
@@ -132,6 +145,12 @@ public:
 	// Update the drawing policy with the marked pin (which may not be valid)
 	void SetMarkedPin(TWeakPtr<SGraphPin> InMarkedPin);
 
+	// Set the selected nodes from the graph panel.
+	void SetSelectedNodes(const TArray<UEdGraphNode*>& InSelectedNodes) { SelectedGraphNodes = InSelectedNodes; }
+
+	// Set the list of currently relinked connections.
+	void SetRelinkConnections(const TArray<FRelinkConnection>& Connections) { RelinkConnections = Connections; }
+
 	static float MakeSplineReparamTable(const FVector2D& P0, const FVector2D& P0Tangent, const FVector2D& P1, const FVector2D& P1Tangent, FInterpCurve<float>& OutReparamTable);
 
 	virtual void DrawSplineWithArrow(const FVector2D& StartPoint, const FVector2D& EndPoint, const FConnectionParams& Params);
diff --git a/Engine/Source/Editor/GraphEditor/Public/GraphSplineOverlapResult.h b/Engine/Source/Editor/GraphEditor/Public/GraphSplineOverlapResult.h
--- a/Engine/Source/Editor/GraphEditor/Public/GraphSplineOverlapResult.h
+++ b/Engine/Source/Editor/GraphEditor/Public/GraphSplineOverlapResult.h
@@ -69,6 +69,14 @@ public:
 		return Result;
 	}
 
+	FGraphPinHandle GetBestPinHandle() const
+	{
+		return BestPinHandle;
+	}
+
+	FGraphPinHandle GetPin1Handle() const { return Pin1Handle; }
+	FGraphPinHandle GetPin2Handle() const { return Pin2Handle; }
+
 	bool GetPins(const class SGraphPanel& InGraphPanel, UEdGraphPin*& OutPin1, UEdGraphPin*& OutPin2) const;
 	void GetPinWidgets(const class SGraphPanel& InGraphPanel, TSharedPtr<class SGraphPin>& OutPin1, TSharedPtr<class SGraphPin>& OutPin2) const;
 
diff --git a/Engine/Source/Editor/GraphEditor/Public/SGraphPanel.h b/Engine/Source/Editor/GraphEditor/Public/SGraphPanel.h
--- a/Engine/Source/Editor/GraphEditor/Public/SGraphPanel.h
+++ b/Engine/Source/Editor/GraphEditor/Public/SGraphPanel.h
@@ -4,6 +4,7 @@
 
 #include "Animation/CurveSequence.h"
 #include "BlueprintUtilities.h"
+#include "ConnectionDrawingPolicy.h"
 #include "Containers/Array.h"
 #include "Containers/Map.h"
 #include "Containers/Set.h"
@@ -160,6 +161,15 @@ public:
 	void OnStopMakingConnection(bool bForceStop = false);
 	void PreservePinPreviewUntilForced();
 
+	/** Indicate that the connection from the given start to the given end pins is being relinked. A preview connection is being drawn for the relinked connection. */
+	void OnBeginRelinkConnection(const FGraphPinHandle& InSourcePinHandle, const FGraphPinHandle& InTargetPinHandle);
+
+	/** The relink connection operation either got cancelled or has successfully been executed. Preview connection won't be drawn anymore. */
+	void OnEndRelinkConnection(bool bForceStop = false);
+
+	/** True in case a connection is currently being relinked, false if not. */
+	bool IsRelinkingConnection() const;
+
 	/** Update this GraphPanel to match the data that it is observing. Expected to be called during ticking. */
 	void Update();
 
@@ -239,11 +249,19 @@ public:
 	/** Get a graph node widget from the specified GUID, if it applies to any nodes in this graph */
 	TSharedPtr<SGraphNode> GetNodeWidgetFromGuid(FGuid Guid) const;
 
+	/** Get a list of selected editor graph nodes from the selection manager. */
+	TArray<UEdGraphNode*> GetSelectedGraphNodes() const;
+
+	const FGraphSplineOverlapResult& GetPreviousFrameSplineOverlap() const { return PreviousFrameSplineOverlap; }
+
 private:
 
 	/** A map of guid -> graph nodes */
 	TMap<FGuid, TWeakPtr<SGraphNode>> NodeGuidMap;
 
+	/** List of currently relinked connections. */
+	TArray<FConnectionDrawingPolicy::FRelinkConnection> RelinkConnections;
+
 protected:
 	UEdGraph* GraphObj;
 	
diff --git a/Engine/Source/Editor/GraphEditor/Public/SGraphPin.h b/Engine/Source/Editor/GraphEditor/Public/SGraphPin.h
--- a/Engine/Source/Editor/GraphEditor/Public/SGraphPin.h
+++ b/Engine/Source/Editor/GraphEditor/Public/SGraphPin.h
@@ -58,6 +58,12 @@ struct GRAPHEDITOR_API FGraphPinHandle
 	/** The GUID of the pin we are referencing */
 	FGuid PinId;
 
+	/**
+	 * Default constructor
+	 * Will contain a invalid node and pin GUID and IsValid() will return false.
+	 */
+	FGraphPinHandle() = default;
+
 	/** Constructor */
 	FGraphPinHandle(UEdGraphPin* InPin);
 
diff --git a/Engine/Source/Runtime/Engine/Classes/EdGraph/EdGraphSchema.h b/Engine/Source/Runtime/Engine/Classes/EdGraph/EdGraphSchema.h
--- a/Engine/Source/Runtime/Engine/Classes/EdGraph/EdGraphSchema.h
+++ b/Engine/Source/Runtime/Engine/Classes/EdGraph/EdGraphSchema.h
@@ -775,6 +775,20 @@ class ENGINE_API UEdGraphSchema : public UObject
 	 */
 	virtual bool TryCreateConnection(UEdGraphPin* A, UEdGraphPin* B) const;
 
+	/** Is this schema supporting connection relinking for the given pin? */
+	virtual bool IsConnectionRelinkingAllowed(UEdGraphPin* InPin) const;
+
+	/**
+	 * Determine if a connection can be relinked to the given pin.
+	 * @param[in] OldSourcePin The current source pin of the connection.
+	 * @param[in] TargetPinCandidate The target pin of the relink.
+	 * @return A message describing if the operation can succeed or why the relink operation would fail.
+	 */
+	virtual const FPinConnectionResponse CanRelinkConnectionToPin(const UEdGraphPin* OldSourcePin, const UEdGraphPin* TargetPinCandidate) const;
+
+	/** Try relinking the connection starting at the old source and target pins and relink it to the new target pin. */
+	virtual bool TryRelinkConnectionTarget(UEdGraphPin* SourcePin, UEdGraphPin* OldTargetPin, UEdGraphPin* NewTargetPin, const TArray<UEdGraphNode*>& InSelectedGraphNodes) const;
+
 	/**
 	 * Try to create an automatic cast or other conversion node node to facilitate a connection between two pins.
 	 * It makes the cast node, a connection between A and the cast node, and a connection from the cast node to B.two 
diff --git a/Engine/Source/Runtime/Engine/Private/EdGraph/EdGraphSchema.cpp b/Engine/Source/Runtime/Engine/Private/EdGraph/EdGraphSchema.cpp
--- a/Engine/Source/Runtime/Engine/Private/EdGraph/EdGraphSchema.cpp
+++ b/Engine/Source/Runtime/Engine/Private/EdGraph/EdGraphSchema.cpp
@@ -463,6 +463,21 @@ bool UEdGraphSchema::TryCreateConnection(UEdGraphPin* PinA, UEdGraphPin* PinB) c
 	return bModified;
 }
 
+bool UEdGraphSchema::IsConnectionRelinkingAllowed(UEdGraphPin* InPin) const
+{
+	return false;
+}
+
+const FPinConnectionResponse UEdGraphSchema::CanRelinkConnectionToPin(const UEdGraphPin* OldSourcePin, const UEdGraphPin* TargetPinCandidate) const
+{
+	return FPinConnectionResponse(CONNECT_RESPONSE_DISALLOW, TEXT("Not implemented by this schema"));
+}
+
+bool UEdGraphSchema::TryRelinkConnectionTarget(UEdGraphPin* SourcePin, UEdGraphPin* OldTargetPin, UEdGraphPin* NewTargetPin, const TArray<UEdGraphNode*>& InSelectedGraphNodes) const
+{
+	return false;
+}
+
 bool UEdGraphSchema::CreateAutomaticConversionNodeAndConnections(UEdGraphPin* PinA, UEdGraphPin* PinB) const
 {
 	return false;
