[UE-169576] Can't insert reroute nodes via double click

diff --git a/Engine/Source/Editor/AnimGraph/Private/AnimationStateMachineSchema.cpp b/Engine/Source/Editor/AnimGraph/Private/AnimationStateMachineSchema.cpp
--- a/Engine/Source/Editor/AnimGraph/Private/AnimationStateMachineSchema.cpp
+++ b/Engine/Source/Editor/AnimGraph/Private/AnimationStateMachineSchema.cpp
@@ -327,12 +327,18 @@ bool UAnimationStateMachineSchema::TryRelinkConnectionTarget(UEdGraphPin* Source
 
 bool UAnimationStateMachineSchema::IsConnectionRelinkingAllowed(UEdGraphPin* InPin) const
 {
-	if (!InPin || !InPin->GetOwningNode())
+	if (InPin && InPin->GetOwningNode())
 	{
-		return false;
+		UAnimStateNodeBase* StateNode = Cast<UAnimStateNodeBase>(InPin->GetOwningNode());
+		UAnimStateTransitionNode* TransitionNode = Cast<UAnimStateTransitionNode>(InPin->GetOwningNode());
+		UAnimStateEntryNode* EntryNode = Cast<UAnimStateEntryNode>(InPin->GetOwningNode());
+		if (StateNode || TransitionNode || EntryNode)
+		{
+			return true;
+		}
 	}
 
-	return true;
+	return false;
 }
 
 const FPinConnectionResponse UAnimationStateMachineSchema::CanRelinkConnectionToPin(const UEdGraphPin* OldSourcePin, const UEdGraphPin* TargetPinCandidate) const
diff --git a/Engine/Source/Editor/GraphEditor/Private/SGraphPanel.cpp b/Engine/Source/Editor/GraphEditor/Private/SGraphPanel.cpp
--- a/Engine/Source/Editor/GraphEditor/Private/SGraphPanel.cpp
+++ b/Engine/Source/Editor/GraphEditor/Private/SGraphPanel.cpp
@@ -815,11 +815,8 @@ FReply SGraphPanel::OnKeyDown( const FGeometry& MyGeometry, const FKeyEvent& InK
 
 FReply SGraphPanel::OnMouseButtonDown(const FGeometry& MyGeometry, const FPointerEvent& MouseEvent)
 {
-	if ((MouseEvent.GetEffectingButton() == EKeys::LeftMouseButton))
+	if ((MouseEvent.GetEffectingButton() == EKeys::LeftMouseButton) && (MouseEvent.IsAltDown() || MouseEvent.IsControlDown()))
 	{
-		SGraphPin* BestPinFromHoveredSpline = GetBestPinFromHoveredSpline();
-		FGraphPinHandle Pin1Handle = PreviousFrameSplineOverlap.GetPin1Handle();
-
 		// Intercept alt-left clicking on the hovered spline for targeted break link
 		UEdGraphPin* Pin1;
 		UEdGraphPin* Pin2;
@@ -828,16 +825,25 @@ FReply SGraphPanel::OnMouseButtonDown(const FGeometry& MyGeometry, const FPointe
 			const UEdGraphSchema* Schema = GraphObj->GetSchema();
 			Schema->BreakSinglePinLink(Pin1, Pin2);
 		}
-		else if (BestPinFromHoveredSpline && MouseEvent.IsControlDown())
+		else if (SGraphPin* BestPinFromHoveredSpline = GetBestPinFromHoveredSpline())
 		{
 			return BestPinFromHoveredSpline->OnPinMouseDown(MyGeometry, MouseEvent);
 		}
-		else if (Pin1Handle.IsValid())
+	}
+
+	if (MouseEvent.GetEffectingButton() == EKeys::LeftMouseButton)
+	{
+		FGraphPinHandle Pin1Handle = PreviousFrameSplineOverlap.GetPin1Handle();
+		if (Pin1Handle.IsValid())
 		{
 			TSharedPtr<class SGraphPin> SourcePin = Pin1Handle.FindInGraphPanel(*this);
 			if (SourcePin.IsValid())
 			{
-				return SourcePin->OnPinMouseDown(MyGeometry, MouseEvent);
+				const UEdGraphSchema* Schema = GraphObj->GetSchema();
+				if (Schema->IsConnectionRelinkingAllowed(SourcePin->GetPinObj()))
+				{
+					return SourcePin->OnPinMouseDown(MyGeometry, MouseEvent);
+				}
 			}
 		}
 	}
