Added replaced hardcoded in SGraphNodeSwitchStatement::IsAddPinButtonVisible with a overridable UK2Node_Switch::SupportsAddPinButton function
Intention: allow us to inherit from UK2Node_Switch and hide the "Add Pin" button

diff --git a/Engine/Source/Editor/BlueprintGraph/Classes/K2Node_Switch.h b/Engine/Source/Editor/BlueprintGraph/Classes/K2Node_Switch.h
--- a/Engine/Source/Editor/BlueprintGraph/Classes/K2Node_Switch.h
+++ b/Engine/Source/Editor/BlueprintGraph/Classes/K2Node_Switch.h
@@ -84,6 +84,9 @@ class BLUEPRINTGRAPH_API UK2Node_Switch : public UK2Node
 	/** Whether an execution pin can be removed from the node or not */
 	virtual bool CanRemoveExecutionPin(UEdGraphPin* TargetPin) const;
 
+	/** Determines whether the add button should be enabled */
+	virtual bool SupportsAddPinButton() const { return true; }
+
 	/** Getting pin access */
 	UEdGraphPin* GetSelectionPin() const;
 	UEdGraphPin* GetDefaultPin() const;
diff --git a/Engine/Source/Editor/BlueprintGraph/Classes/K2Node_SwitchEnum.h b/Engine/Source/Editor/BlueprintGraph/Classes/K2Node_SwitchEnum.h
--- a/Engine/Source/Editor/BlueprintGraph/Classes/K2Node_SwitchEnum.h
+++ b/Engine/Source/Editor/BlueprintGraph/Classes/K2Node_SwitchEnum.h
@@ -67,8 +67,9 @@ class UK2Node_SwitchEnum : public UK2Node_Switch, public INodeDependingOnEnumInt
 	virtual void AddPinToSwitchNode() override;
 	virtual void RemovePinFromSwitchNode(UEdGraphPin* TargetPin) override;
 	virtual ERedirectType DoPinsMatchForReconstruction(const UEdGraphPin* NewPin, int32 NewPinIndex, const UEdGraphPin* OldPin, int32 OldPinIndex) const override;
+	virtual bool SupportsAddPinButton() const { return false; }
 	// End of UK2Node_Switch Interface
-
+	
 	/** Bind the switch to a named enum */
 	void SetEnum(UEnum* InEnum);
 
diff --git a/Engine/Source/Editor/GraphEditor/Private/KismetNodes/SGraphNodeSwitchStatement.cpp b/Engine/Source/Editor/GraphEditor/Private/KismetNodes/SGraphNodeSwitchStatement.cpp
--- a/Engine/Source/Editor/GraphEditor/Private/KismetNodes/SGraphNodeSwitchStatement.cpp
+++ b/Engine/Source/Editor/GraphEditor/Private/KismetNodes/SGraphNodeSwitchStatement.cpp
@@ -11,7 +11,6 @@
 #include "HAL/PlatformCrt.h"
 #include "Internationalization/Internationalization.h"
 #include "K2Node_Switch.h"
-#include "K2Node_SwitchEnum.h"
 #include "Kismet2/BlueprintEditorUtils.h"
 #include "KismetPins/SGraphPinExec.h"
 #include "Layout/Margin.h"
@@ -120,7 +119,8 @@ void SGraphNodeSwitchStatement::CreateOutputSideAddButton(TSharedPtr<SVerticalBo
 
 EVisibility SGraphNodeSwitchStatement::IsAddPinButtonVisible() const
 {
-	if(!GraphNode->IsA<UK2Node_SwitchEnum>())
+	if (UK2Node_Switch* SwitchNode = CastChecked<UK2Node_Switch>(GraphNode)
+		; SwitchNode && SwitchNode->SupportsAddPinButton())
 	{
 		return SGraphNode::IsAddPinButtonVisible();
 	}
