Extend SGraphActionMenu to allow owners to directly provide the action menu source list at refresh time.

diff --git a/Engine/Source/Editor/GraphEditor/Private/SGraphActionMenu.cpp b/Engine/Source/Editor/GraphEditor/Private/SGraphActionMenu.cpp
--- a/Engine/Source/Editor/GraphEditor/Private/SGraphActionMenu.cpp
+++ b/Engine/Source/Editor/GraphEditor/Private/SGraphActionMenu.cpp
@@ -320,6 +320,7 @@ void SGraphActionMenu::Construct( const FArguments& InArgs, bool bIsReadOnly/* =
 	this->OnCategoryDragged = InArgs._OnCategoryDragged;
 	this->OnCreateWidgetForAction = InArgs._OnCreateWidgetForAction;
 	this->OnCreateCustomRowExpander = InArgs._OnCreateCustomRowExpander;
+	this->OnGetActionList = InArgs._OnGetActionList;
 	this->OnCollectAllActions = InArgs._OnCollectAllActions;
 	this->OnCollectStaticSections = InArgs._OnCollectStaticSections;
 	this->OnCategoryTextCommitted = InArgs._OnCategoryTextCommitted;
@@ -332,6 +333,18 @@ void SGraphActionMenu::Construct( const FArguments& InArgs, bool bIsReadOnly/* =
 	this->DraggedFromPins = InArgs._DraggedFromPins;
 	this->GraphObj = InArgs._GraphObj;
 
+	// Default graph action list (also provides an empty source list to start with)
+	AllActions = MakeShared<FGraphActionListBuilderBase>();
+	
+	if(OnGetActionList.IsBound())
+	{
+		// If we are obtaining a new action list at refresh time, ensure that the indirect collection delegate is unbound
+		if (!ensureMsgf(!OnCollectAllActions.IsBound(), TEXT("The OnCollectAllActions delegate is bound, but will not be invoked, because OnGetActionList has also been bound and will be used to obtain the action list for this menu. To resolve this, one of these events should be removed from its construction.")))
+		{
+			OnCollectAllActions.Unbind();
+		}
+	}
+
 	// If a delegate for filtering text is passed in, assign it so that it will be used instead of the built-in filter box
 	if(InArgs._OnGetFilterText.IsBound())
 	{
@@ -400,11 +413,18 @@ void SGraphActionMenu::RefreshAllActions(bool bPreserveExpansion, bool bHandleOn
 	TArray< TSharedPtr<FGraphActionNode> > SelectedNodes = TreeView->GetSelectedItems();
 	TSharedPtr<FGraphActionNode> SelectedAction = SelectedNodes.Num() > 0 ? SelectedNodes[0] : nullptr;
 
+	if (OnGetActionList.IsBound())
+	{
+		// Obtain the source action list directly.
+		AllActions = OnGetActionList.Execute();
+	}
+	else
 	{
 		TRACE_CPUPROFILER_EVENT_SCOPE(SGraphActionMenu::CollectAllActions);
 
-		AllActions.Empty();
-		OnCollectAllActions.ExecuteIfBound(AllActions);
+		// Collect actions into our local list context.
+		AllActions->Empty();
+		OnCollectAllActions.ExecuteIfBound(*AllActions);
 	}
 
 	GenerateFilteredItems(bPreserveExpansion);
@@ -752,9 +772,9 @@ void SGraphActionMenu::GenerateFilteredItems(bool bPreserveExpansion)
 	const UEdGraphSchema* ActionSchema = GraphObj ? GraphObj->GetSchema() : GetDefault<UEdGraphSchema>();
 	check(ActionSchema);
 
-	for (int32 CurTypeIndex=0; CurTypeIndex < AllActions.GetNumActions(); ++CurTypeIndex)
+	for (int32 CurTypeIndex = 0; CurTypeIndex < AllActions->GetNumActions(); ++CurTypeIndex)
 	{
-		FGraphActionListBuilderBase::ActionGroup& CurrentAction = AllActions.GetAction( CurTypeIndex );
+		FGraphActionListBuilderBase::ActionGroup& CurrentAction = AllActions->GetAction(CurTypeIndex);
 
 		// If we're filtering, search check to see if we need to show this action
 		bool bShowAction = true;
@@ -807,7 +827,7 @@ void SGraphActionMenu::GenerateFilteredItems(bool bPreserveExpansion)
 	// If theres a BestMatchIndex find it in the actions nodes and select it (maybe this should check the current selected suggestion first ?)
 	if( BestMatchIndex != INDEX_NONE ) 
 	{
-		FGraphActionListBuilderBase::ActionGroup& FilterSelectAction = AllActions.GetAction( BestMatchIndex );
+		FGraphActionListBuilderBase::ActionGroup& FilterSelectAction = AllActions->GetAction( BestMatchIndex );
 		if( FilterSelectAction.Actions[0].IsValid() == true )
 		{
 			for (int32 iNode = 0; iNode < FilteredActionNodes.Num() ; iNode++)
@@ -850,7 +870,7 @@ bool SGraphActionMenu::ShouldExpandNodes() const
 {
 	// Expand all the categories that have filter results, or when there are only a few to show
 	const bool bFilterActive = !GetFilterText().IsEmpty();
-	const bool bOnlyAFewTotal = AllActions.GetNumActions() < 10;
+	const bool bOnlyAFewTotal = AllActions->GetNumActions() < 10;
 
 	return bFilterActive || bOnlyAFewTotal || bAutoExpandActionMenu;
 }
@@ -1280,9 +1300,9 @@ void SGraphActionMenu::MarkActiveSuggestion()
 
 void SGraphActionMenu::AddReferencedObjects( FReferenceCollector& Collector )
 {
-	for (int32 CurTypeIndex=0; CurTypeIndex < AllActions.GetNumActions(); ++CurTypeIndex)
+	for (int32 CurTypeIndex = 0; CurTypeIndex < AllActions->GetNumActions(); ++CurTypeIndex)
 	{
-		FGraphActionListBuilderBase::ActionGroup& Action = AllActions.GetAction( CurTypeIndex );
+		FGraphActionListBuilderBase::ActionGroup& Action = AllActions->GetAction(CurTypeIndex);
 
 		for ( int32 ActionIndex = 0; ActionIndex < Action.Actions.Num(); ActionIndex++ )
 		{
diff --git a/Engine/Source/Editor/GraphEditor/Public/SGraphActionMenu.h b/Engine/Source/Editor/GraphEditor/Public/SGraphActionMenu.h
--- a/Engine/Source/Editor/GraphEditor/Public/SGraphActionMenu.h
+++ b/Engine/Source/Editor/GraphEditor/Public/SGraphActionMenu.h
@@ -11,7 +11,6 @@
 #include "CoreMinimal.h"
 #include "CoreTypes.h"
 #include "Delegates/Delegate.h"
-#include "EdGraph/EdGraphSchema.h"
 #include "Framework/SlateDelegates.h"
 #include "HAL/PlatformCrt.h"
 #include "Input/Reply.h"
@@ -48,6 +47,7 @@ struct FGeometry;
 struct FGraphActionNode;
 struct FKeyEvent;
 struct FPointerEvent;
+struct FGraphActionListBuilderBase;
 
 /** Delegate for hooking up an inline editable text block to be notified that a rename is requested. */
 DECLARE_DELEGATE( FOnRenameRequestActionNode );
@@ -141,6 +141,8 @@ public:
 	DECLARE_DELEGATE_RetVal( FText, FGetFilterText);
 	/** Delegate to check if an action matches a specified name (used for renaming items etc.) */
 	DECLARE_DELEGATE_RetVal_TwoParams( bool, FOnActionMatchesName, FEdGraphSchemaAction*, const FName& );
+	/** Delegate that can be used to create and/or get a custom action list. */
+	DECLARE_DELEGATE_RetVal(TSharedRef<FGraphActionListBuilderBase>, FOnGetActionList);
 
 	SLATE_BEGIN_ARGS(SGraphActionMenu)
 		: _AutoExpandActionMenu(false)
@@ -159,6 +161,7 @@ public:
 		SLATE_EVENT( FOnContextMenuOpening, OnContextMenuOpening )
 		SLATE_EVENT( FOnCreateWidgetForAction, OnCreateWidgetForAction )
 		SLATE_EVENT( FOnCreateCustomRowExpander, OnCreateCustomRowExpander )
+		SLATE_EVENT( FOnGetActionList, OnGetActionList )
 		SLATE_EVENT( FOnCollectAllActions, OnCollectAllActions )
 		SLATE_EVENT( FOnCollectStaticSections, OnCollectStaticSections )
 		SLATE_EVENT( FOnCategoryTextCommitted, OnCategoryTextCommitted )
@@ -205,7 +208,7 @@ protected:
 	TSharedPtr<SSearchBox> FilterTextBox;
 
 	/** List of all actions we can browser */
-	FGraphActionListBuilderBase AllActions;
+	TSharedPtr<FGraphActionListBuilderBase> AllActions;
 
 	/** Flattened list of all actions passing the filter */
 	TArray< TSharedPtr<FGraphActionNode> > FilteredActionNodes; 
@@ -242,6 +245,8 @@ protected:
 	FOnCreateWidgetForAction OnCreateWidgetForAction;
 	/** Delegate to call for creating a custom "expander" widget for indenting a menu row with */
 	FOnCreateCustomRowExpander OnCreateCustomRowExpander;
+	/** Delegate to call to get a custom action list */
+	FOnGetActionList OnGetActionList;
 	/** Delegate to call to collect all actions */
 	FOnCollectAllActions OnCollectAllActions;
 	/** Delegate to call to collect all always visible sections */
