CIS non-unity fix.

diff --git a/Engine/Source/Editor/GraphEditor/Public/SGraphActionMenu.h b/Engine/Source/Editor/GraphEditor/Public/SGraphActionMenu.h
--- a/Engine/Source/Editor/GraphEditor/Public/SGraphActionMenu.h
+++ b/Engine/Source/Editor/GraphEditor/Public/SGraphActionMenu.h
@@ -47,6 +47,7 @@ struct FGeometry;
 struct FGraphActionNode;
 struct FKeyEvent;
 struct FPointerEvent;
+struct FEdGraphSchemaAction;
 struct FGraphActionListBuilderBase;
 
 /** Delegate for hooking up an inline editable text block to be notified that a rename is requested. */
