SGraphNode::PopulateMetaTag : use GetOutermostObject (fixes use case with OFPA Outers)

diff --git a/Engine/Source/Editor/GraphEditor/Private/SGraphNode.cpp b/Engine/Source/Editor/GraphEditor/Private/SGraphNode.cpp
--- a/Engine/Source/Editor/GraphEditor/Private/SGraphNode.cpp
+++ b/Engine/Source/Editor/GraphEditor/Private/SGraphNode.cpp
@@ -1690,21 +1690,13 @@ void SGraphNode::PopulateMetaTag(FGraphNodeMetaData* TagMeta) const
 {
 	if (GraphNode && TagMeta)
 	{
-		// We want the name of the blueprint as our name - we can find the node from the GUID
-		UObject* Package = GraphNode->GetOutermost();
-		UObject* LastOuter = GraphNode->GetOuter();
-		while (LastOuter && (LastOuter->GetOuter() != Package))
-		{
-			LastOuter = LastOuter->GetOuter();
-		}
-
-		if(LastOuter)
-		{
-			TagMeta->Tag = FName(*FString::Printf(TEXT("GraphNode_%s_%s"), *LastOuter->GetFullName(), *GraphNode->NodeGuid.ToString()));
-			TagMeta->OuterName = LastOuter->GetFullName();
-			TagMeta->GUID = GraphNode->NodeGuid;
-			TagMeta->FriendlyName = FString::Printf(TEXT("%s in %s"), *GraphNode->GetNodeTitle(ENodeTitleType::ListView).ToString(), *TagMeta->OuterName);	
-		}	
+		// We want the name of the blueprint/world as our name - we can find the node from the GUID
+		UObject* OutermostObject = GraphNode->GetOutermostObject();
+		check(OutermostObject);
+		TagMeta->Tag = FName(*FString::Printf(TEXT("GraphNode_%s_%s"), *OutermostObject->GetFullName(), *GraphNode->NodeGuid.ToString()));
+		TagMeta->OuterName = OutermostObject->GetFullName();
+		TagMeta->GUID = GraphNode->NodeGuid;
+		TagMeta->FriendlyName = FString::Printf(TEXT("%s in %s"), *GraphNode->GetNodeTitle(ENodeTitleType::ListView).ToString(), *TagMeta->OuterName);		
 	}
 }
 
