Removing GraphEditor private includes via moving SGraphPinEnum.h to the Public folder.

diff --git a/Engine/Plugins/FX/Niagara/Source/NiagaraEditor/NiagaraEditor.Build.cs b/Engine/Plugins/FX/Niagara/Source/NiagaraEditor/NiagaraEditor.Build.cs
--- a/Engine/Plugins/FX/Niagara/Source/NiagaraEditor/NiagaraEditor.Build.cs
+++ b/Engine/Plugins/FX/Niagara/Source/NiagaraEditor/NiagaraEditor.Build.cs
@@ -14,7 +14,6 @@ public class NiagaraEditor : ModuleRules
 			"NiagaraEditor/Private/Sequencer/NiagaraSequence",
 			"NiagaraEditor/Private/ViewModels",
 			"NiagaraEditor/Private/TypeEditorUtilities",
-			Path.Combine(GetModuleDirectory("GraphEditor"), "Private"),
 			Path.Combine(GetModuleDirectory("PropertyEditor"), "Private"),
 			Path.Combine(GetModuleDirectory("Renderer"), "Private"),
 		});
diff --git a/Engine/Plugins/FX/Niagara/Source/NiagaraEditorWidgets/NiagaraEditorWidgets.Build.cs b/Engine/Plugins/FX/Niagara/Source/NiagaraEditorWidgets/NiagaraEditorWidgets.Build.cs
--- a/Engine/Plugins/FX/Niagara/Source/NiagaraEditorWidgets/NiagaraEditorWidgets.Build.cs
+++ b/Engine/Plugins/FX/Niagara/Source/NiagaraEditorWidgets/NiagaraEditorWidgets.Build.cs
@@ -8,7 +8,6 @@ public class NiagaraEditorWidgets : ModuleRules
 	public NiagaraEditorWidgets(ReadOnlyTargetRules Target) : base(Target)
 	{
 		PrivateIncludePaths.AddRange(new string[] {
-			System.IO.Path.Combine(GetModuleDirectory("GraphEditor"), "Private"),
 		});
 
 		PrivateDependencyModuleNames.AddRange(new string[] {
diff --git a/Engine/Source/Editor/GraphEditor/Private/KismetPins/SGraphPinEnum.h b/Engine/Source/Editor/GraphEditor/Private/KismetPins/SGraphPinEnum.h
deleted file mode 100644
diff --git a/Engine/Source/Editor/GraphEditor/Public/KismetPins/SGraphPinEnum.h b/Engine/Source/Editor/GraphEditor/Public/KismetPins/SGraphPinEnum.h
new file mode 100644
--- /dev/null
+++ b/Engine/Source/Editor/GraphEditor/Public/KismetPins/SGraphPinEnum.h
@@ -0,0 +1,73 @@
+// Copyright Epic Games, Inc. All Rights Reserved.
+
+#pragma once
+
+#include "Containers/Array.h"
+#include "Containers/UnrealString.h"
+#include "HAL/Platform.h"
+#include "Internationalization/Text.h"
+#include "SGraphPin.h"
+#include "Templates/SharedPointer.h"
+#include "Types/SlateEnums.h"
+#include "Widgets/DeclarativeSyntaxSupport.h"
+
+class SWidget;
+class UEdGraphPin;
+
+class GRAPHEDITOR_API SGraphPinEnum : public SGraphPin
+{
+public:
+	SLATE_BEGIN_ARGS(SGraphPinEnum) {}
+	SLATE_END_ARGS()
+
+	void Construct(const FArguments& InArgs, UEdGraphPin* InGraphPinObj);
+
+protected:
+
+	/**
+	 *	Function to create class specific widget.
+	 *
+	 *	@return Reference to the newly created widget object
+	 */
+	virtual TSharedRef<SWidget>	GetDefaultValueWidget() override;
+
+private:
+
+	/**
+	 *	Function to get current string associated with the combo box selection
+	 *
+	 *	@return currently selected string
+	 */
+	FString OnGetText() const;
+
+	/**
+	 *	Function to generate the list of indexes from the enum object
+	 *
+	 *	@param OutComboBoxIndexes - Int array reference to store the list of indexes
+	 */
+	void GenerateComboBoxIndexes( TArray< TSharedPtr<int32> >& OutComboBoxIndexes );
+
+	/**
+	 *	Function to set the newly selected index
+	 *
+	 * @param NewSelection The newly selected item in the combo box
+	 * @param SelectInfo Provides context on how the selection changed
+	 */
+	void ComboBoxSelectionChanged( TSharedPtr<int32> NewSelection, ESelectInfo::Type SelectInfo );
+
+	/**
+	 * Returns the friendly name of the enum at index EnumIndex
+	 *
+	 * @param EnumIndex	- The index of the enum to return the friendly name for
+	 */
+	FText OnGetFriendlyName(int32 EnumIndex);
+
+	/**
+	 * Returns the tooltip of the enum at index EnumIndex
+	 *
+	 * @param EnumIndex	- The index of the enum to return the tooltip for
+	 */
+	FText OnGetTooltip(int32 EnumIndex);
+
+	TSharedPtr<class SPinComboBox>	ComboBox;
+};
diff --git a/Engine/Source/Editor/MaterialEditor/MaterialEditor.Build.cs b/Engine/Source/Editor/MaterialEditor/MaterialEditor.Build.cs
--- a/Engine/Source/Editor/MaterialEditor/MaterialEditor.Build.cs
+++ b/Engine/Source/Editor/MaterialEditor/MaterialEditor.Build.cs
@@ -8,7 +8,6 @@ public class MaterialEditor : ModuleRules
 	{
 		PrivateIncludePaths.AddRange(
 			new string[] {
-				System.IO.Path.Combine(GetModuleDirectory("GraphEditor"), "Private"),
 			}
 		);
 
