Removing dependencies on GraphEditor Private by moving GraphActionNode.h to its Public folder.

diff --git a/Engine/Plugins/Animation/DeformerGraph/Source/OptimusEditor/OptimusEditor.Build.cs b/Engine/Plugins/Animation/DeformerGraph/Source/OptimusEditor/OptimusEditor.Build.cs
--- a/Engine/Plugins/Animation/DeformerGraph/Source/OptimusEditor/OptimusEditor.Build.cs
+++ b/Engine/Plugins/Animation/DeformerGraph/Source/OptimusEditor/OptimusEditor.Build.cs
@@ -7,7 +7,6 @@ namespace UnrealBuildTool.Rules
         public OptimusEditor(ReadOnlyTargetRules Target) : base(Target)
         {
 			PrivateIncludePaths.AddRange(new string[] {
-				System.IO.Path.Combine(GetModuleDirectory("GraphEditor"), "Private"),
 				System.IO.Path.Combine(GetModuleDirectory("OptimusCore"), "Private"),
 				System.IO.Path.Combine(GetModuleDirectory("OptimusEditor"), "Private"),
 				System.IO.Path.Combine(GetModuleDirectory("PropertyEditor"), "Private"),
diff --git a/Engine/Plugins/Runtime/GameplayAbilities/Source/GameplayAbilitiesEditor/GameplayAbilitiesEditor.Build.cs b/Engine/Plugins/Runtime/GameplayAbilities/Source/GameplayAbilitiesEditor/GameplayAbilitiesEditor.Build.cs
--- a/Engine/Plugins/Runtime/GameplayAbilities/Source/GameplayAbilitiesEditor/GameplayAbilitiesEditor.Build.cs
+++ b/Engine/Plugins/Runtime/GameplayAbilities/Source/GameplayAbilitiesEditor/GameplayAbilitiesEditor.Build.cs
@@ -13,7 +13,6 @@ namespace UnrealBuildTool.Rules
 				new string[] {
 					Path.Combine(GetModuleDirectory("AssetTools"), "Private"),
 					Path.Combine(GetModuleDirectory("GameplayTagsEditor"), "Private"),
-					Path.Combine(GetModuleDirectory("GraphEditor"), "Private"),
 					Path.Combine(GetModuleDirectory("Kismet"), "Private"),
 				});
 
diff --git a/Engine/Source/Editor/AIGraph/AIGraph.Build.cs b/Engine/Source/Editor/AIGraph/AIGraph.Build.cs
--- a/Engine/Source/Editor/AIGraph/AIGraph.Build.cs
+++ b/Engine/Source/Editor/AIGraph/AIGraph.Build.cs
@@ -11,7 +11,6 @@ public class AIGraph : ModuleRules
 		PrivateIncludePaths.AddRange(
             new string[] {
 				System.IO.Path.Combine(GetModuleDirectory("AIGraph"), "Private"),
-				System.IO.Path.Combine(GetModuleDirectory("GraphEditor"), "Private"),
 				System.IO.Path.Combine(GetModuleDirectory("Kismet"), "Private"),
 			}
         );
diff --git a/Engine/Source/Editor/BehaviorTreeEditor/BehaviorTreeEditor.Build.cs b/Engine/Source/Editor/BehaviorTreeEditor/BehaviorTreeEditor.Build.cs
--- a/Engine/Source/Editor/BehaviorTreeEditor/BehaviorTreeEditor.Build.cs
+++ b/Engine/Source/Editor/BehaviorTreeEditor/BehaviorTreeEditor.Build.cs
@@ -10,7 +10,6 @@ public class BehaviorTreeEditor : ModuleRules
             new string[] {
 				System.IO.Path.Combine(GetModuleDirectory("AIGraph"), "Private"),
 				System.IO.Path.Combine(GetModuleDirectory("BehaviorTreeEditor"), "Private"),
-				System.IO.Path.Combine(GetModuleDirectory("GraphEditor"), "Private"),
 			}
 		);
 
diff --git a/Engine/Source/Editor/GameplayTasksEditor/GameplayTasksEditor.Build.cs b/Engine/Source/Editor/GameplayTasksEditor/GameplayTasksEditor.Build.cs
--- a/Engine/Source/Editor/GameplayTasksEditor/GameplayTasksEditor.Build.cs
+++ b/Engine/Source/Editor/GameplayTasksEditor/GameplayTasksEditor.Build.cs
@@ -12,7 +12,6 @@ namespace UnrealBuildTool.Rules
                 new string[] {
 					System.IO.Path.Combine(GetModuleDirectory("AssetTools"), "Private"),
 					System.IO.Path.Combine(GetModuleDirectory("GameplayTasksEditor"), "Private"),
-					System.IO.Path.Combine(GetModuleDirectory("GraphEditor"), "Private"),
 					System.IO.Path.Combine(GetModuleDirectory("Kismet"), "Private"),
 				}
 			);
diff --git a/Engine/Source/Editor/GraphEditor/Private/GraphActionNode.h b/Engine/Source/Editor/GraphEditor/Private/GraphActionNode.h
deleted file mode 100644
diff --git a/Engine/Source/Editor/GraphEditor/Public/GraphActionNode.h b/Engine/Source/Editor/GraphEditor/Public/GraphActionNode.h
new file mode 100644
--- /dev/null
+++ b/Engine/Source/Editor/GraphEditor/Public/GraphActionNode.h
@@ -0,0 +1,285 @@
+// Copyright Epic Games, Inc. All Rights Reserved.
+
+#pragma once
+
+#include "Containers/Array.h"
+#include "Containers/Map.h"
+#include "Containers/Set.h"
+#include "Containers/UnrealString.h"
+#include "EdGraph/EdGraphSchema.h"
+#include "HAL/Platform.h"
+#include "Internationalization/Text.h"
+#include "SGraphActionMenu.h"
+#include "Templates/SharedPointer.h"
+#include "Templates/TypeHash.h"
+
+template <typename ItemType> class STreeView;
+
+// Utility class for building menus of graph actions
+struct GRAPHEDITOR_API FGraphActionNode : TSharedFromThis<FGraphActionNode>
+{
+public:
+	/** */
+	static const int32 INVALID_SECTION_ID = 0;
+
+	/** Identifies the named section that this node belongs to, if any (defaults to INVALID_SECTION_ID) */
+	int32 const SectionID;
+	/** Identifies the menu group that this node belongs to (defaults to zero) */
+	int32 const Grouping;
+	/** A set of actions to execute when this node is picked from a menu */
+	TArray< TSharedPtr<FEdGraphSchemaAction> > const Actions;
+
+	/** */
+	TArray< TSharedPtr<FGraphActionNode> > Children;
+
+	/** Lookup table for category nodes, used to speed up menu construction */
+	TMap< FString, TSharedPtr<FGraphActionNode> > CategoryNodes;
+
+public:
+	/**
+	 * Static allocator for a new root node (so external users have a starting
+	 * point to build graph action trees from).
+	 *
+	 * @return A newly allocated root node (should not be displayed in the tree view).
+	 */
+	static TSharedPtr<FGraphActionNode> NewRootNode();
+
+	/**
+	 * Inserts a new action node (and any accompanying category nodes) based off
+	 * the provided ActionSet. 
+	 *
+	 * NOTE: This does NOT insert the node in a sorted manner. Call SortChildren() 
+	 *       separately to accomplish that (done for performance reasons).
+	 * 
+	 * @param  ActionSet	A list of actions that you want the node to execute when picked.
+	 * @return The new action node.
+	 */
+	TSharedPtr<FGraphActionNode> AddChild(FGraphActionListBuilderBase::ActionGroup const& ActionSet);
+
+	TSharedPtr<FGraphActionNode> AddSection(int32 Grouping, int32 InSectionID);
+
+	/**
+	 * Sorts all child nodes by section, group, and type (additionally, can
+	 * sort alphabetically if wanted).
+	 * 
+	 * @param  bAlphabetically	Determines if we sort alphabetically on top of section/group/type.
+	 * @param  bRecursive		Determines if we should sort all decendent nodes' children ass well.
+	 */
+	void SortChildren(bool bAlphabetically = true, bool bRecursive = true);
+
+	/**
+	 * Returns a WeakPtr to the Parent Node
+	 */
+	TWeakPtr<FGraphActionNode> GetParentNode() const{ return ParentNode; }
+
+	/**
+	 * Recursively collects all child/grandchild/decendent nodes.
+	 * 
+	 * @param  OutNodeArray	The array to fill out with decendent nodes.
+	 */
+	void GetAllNodes(TArray< TSharedPtr<FGraphActionNode> >& OutNodeArray) const;
+
+	/**
+	 * Recursively collects all decendent action/separator nodes (leaves out 
+	 * branching category-nodes).
+	 * 
+	 * @param  OutLeafArray	The array to fill out with decendent leaf nodes.
+	 */
+	void GetLeafNodes(TArray< TSharedPtr<FGraphActionNode> >& OutLeafArray) const;
+
+	/**
+	 * Takes the tree view and expands its elements for each child.
+	 * 
+	 * @param  TreeView		The tree responsible for visualizing this node hierarchy.
+	 * @param  bRecursive	Determines if you want children/decendents to expand their children as well. 
+	 */
+	void ExpandAllChildren(TSharedPtr< STreeView< TSharedPtr<FGraphActionNode> > > TreeView, bool bRecursive = true);
+
+	/**
+	 * Clears all children (not recursively... the TSharedPtrs should clean up 
+	 * appropriately).
+	 */
+	void ClearChildren();
+
+	/**
+	 * Query to determine this node's type (there are five distinguishable node
+	 * types: root, section heading, category, action, & group-divider).
+	 *
+	 * @return True if this is the type your queried about, otherwise false.
+	 */
+	bool IsRootNode() const;
+	bool IsSectionHeadingNode() const;
+	bool IsCategoryNode() const;
+	bool IsActionNode() const;
+	bool IsGroupDividerNode() const;
+
+	/**
+	 * Determines if this node is a menu separator of some kind (either a
+	 * "group-divider" or a "section heading").
+	 *
+	 * @return True if this is a menu divider, otherwise false.
+	 */
+	bool IsSeparator() const;
+	/**
+	 * Retrieves this node's display name (for category and action nodes). The
+	 * text string will be empty for separator and root nodes.
+	 *
+	 * @return The name to present this node with in the tree view (will be an empty text string if this is a separator node)
+	 */
+	FText const& GetDisplayName() const;
+
+	/**
+	 * Walks the node chain backwards, constructing a category path (delimited
+	 * by '|' characters). This includes this node's category (if it is a
+	 * category node).
+	 *
+	 * @return A category path string, denoting the category hierarchy up to this node.
+	 */
+	FText GetCategoryPath() const;
+
+	/**
+	 * Checks to see if this node contains at least one valid action.
+	 *
+	 * @return True is the Actions array contains a valid entry, otherwise false.
+	 */
+	bool HasValidAction() const;
+
+	/**
+	 * Looks through this node's Actions array, and returns the first valid
+	 * action it finds.
+	 *
+	 * @return This node's first valid action (will be an empty pointer if this is not an action node).
+	 */
+	TSharedPtr<FEdGraphSchemaAction> GetPrimaryAction() const;
+
+	/**
+	 * Accessor to the node's RenameRequestEvent (for binding purposes). Do not
+	 * Execute() the delegate from this function, instead call
+	 * BroadcastRenameRequest() on the node.
+	 *
+	 * @return The node's internal RenameRequestEvent.
+	 */
+	FOnRenameRequestActionNode& OnRenameRequest() { return RenameRequestEvent; }
+
+	/**
+	 * Executes the node's RenameRequestEvent if it is bound. Otherwise, it will
+	 * mark the node as having a pending rename request.
+	 *
+	 * @return True if the broadcast went through, false if the "pending rename request" flag was set.
+	 */
+	bool BroadcastRenameRequest();
+
+	/**
+	 * Sometimes a call to BroadcastRenameRequest() is made before the
+	 * RenameRequestEvent has been bound. When that happens, this node is
+	 * marked with a pending rename request. This method determines if that is
+	 * the case for this node.
+	 *
+	 * @return True if a call to BroadcastRenameRequest() was made without a valid RenameRequestEvent.
+	 */
+	bool IsRenameRequestPending() const;
+
+private:
+	/**
+	 *
+	 *
+	 * @param  Grouping
+	 * @param  SectionID
+	 */
+	FGraphActionNode(int32 Grouping, int32 SectionID);
+
+	/**
+	 * Constructor for action nodes. Private so that users go through AddChild().
+	 *
+	 * @param  ActionList
+	 * @param  Grouping
+	 * @param  SectionID
+	 */
+	FGraphActionNode(TArray< TSharedPtr<FEdGraphSchemaAction> > const& ActionList, int32 Grouping, int32 SectionID);
+
+	/**
+	 *
+	 *
+	 * @param  Parent
+	 * @param  Grouping
+	 * @param  SectionID
+	 * @return
+	 */
+	static TSharedPtr<FGraphActionNode> NewSectionHeadingNode(TWeakPtr<FGraphActionNode> Parent, int32 Grouping, int32 SectionID);
+
+	/**
+	 *
+	 *
+	 * @param  Category
+	 * @param  Grouping
+	 * @param  SectionID
+	 * @return
+	 */
+	static TSharedPtr<FGraphActionNode> NewCategoryNode(FString const& Category, int32 Grouping, int32 SectionID);
+
+	/**
+	 *
+	 *
+	 * @param  ActionList
+	 * @return
+	 */
+	static TSharedPtr<FGraphActionNode> NewActionNode(TArray< TSharedPtr<FEdGraphSchemaAction> > const& ActionList);
+
+	/**
+	 *
+	 *
+	 * @param  Parent
+	 * @param  Grouping
+	 * @return
+	 */
+	static TSharedPtr<FGraphActionNode> NewGroupDividerNode(TWeakPtr<FGraphActionNode> Parent, int32 Grouping);
+
+	/**
+	 * Iterates the CategoryStack, adding category-nodes as needed. The
+	 * last category is what the node will be inserted under.
+	 *
+	 * @param  CategoryStack	A list of categories denoting where to nest the new node (the first element is the highest category)
+	 * @param  Idx				Current point in the category stack that we have iterated to
+	 * @param  NodeToAdd		The node you want inserted.
+	 */
+	void AddChildRecursively(const TArray<FString>& CategoryStack, int32 Idx, TSharedPtr<FGraphActionNode> NodeToAdd);
+
+	/**
+	 * Looks through this node's children to see if a there already exists a 
+	 * node matching one we'd have to spawn (to parent the supplied NodeToAdd).
+	 * 
+	 * @param   ParentName	The name of the category NodeToAdd wants to nest under.
+	 * @param   NodeToAdd	The node that we'll be adding to this child.
+	 * @return  A child node matching the supplied parameters (will be empty if no match was found).
+	 */
+	TSharedPtr<FGraphActionNode> FindMatchingParent(FString const& ParentName, TSharedPtr<FGraphActionNode> NodeToAdd);
+
+	/**
+	 * Adds the specified node directly to this node's Children array. Will
+	 * create and insert separators if needed (if the node has a new group or
+	 * section).
+	 *
+	 * @param  NodeToAdd The node you want inserted.
+	 */
+	void InsertChild(TSharedPtr<FGraphActionNode> NodeToAdd);
+
+private:
+	/** The category or action name (depends on what type of node this is) */
+	FText DisplayText;
+	/** The node that this is a direct child of (empty if this is a root node) */
+	TWeakPtr<FGraphActionNode> ParentNode;
+
+	/** Tracks what groups have already been added (so we can easily determine what group-dividers we need) */
+	TSet<int32> ChildGroupings;
+	/** Tracks what sections have already been added (so we can easily determine what heading we need) */
+	TSet<int32> ChildSections;
+
+	/** When the item is first created, a rename request may occur before everything is setup for it. This toggles to true in those cases */
+	bool bPendingRenameRequest;
+	/** Delegate to trigger when a rename was requested on this node */
+	FOnRenameRequestActionNode RenameRequestEvent;
+
+	friend struct FGraphActionNodeImpl;
+	/** For sorting, when we don't alphabetically sort (so menu items don't jump around). */
+	int32 InsertOrder;	
+};
