PR #10182: Fixed issue where disconnecting an orphan pin would cause ensure to trigger due to invalidated graph pin (Contributed by [REDACTED_PERSON])

diff --git a/Engine/Source/Editor/GraphEditor/Private/SGraphPin.cpp b/Engine/Source/Editor/GraphEditor/Private/SGraphPin.cpp
--- a/Engine/Source/Editor/GraphEditor/Private/SGraphPin.cpp
+++ b/Engine/Source/Editor/GraphEditor/Private/SGraphPin.cpp
@@ -720,14 +720,17 @@ void SGraphPin::OnDragEnter( const FGeometry& MyGeometry, const FDragDropEvent&
 	}
 
 	// Is someone dragging a connection?
-	if (Operation->IsOfType<FGraphEditorDragDropAction>() && GetIsConnectable())
+	if (Operation->IsOfType<FGraphEditorDragDropAction>())
 	{
 		// Ensure that the pin is valid before using it
 		if(GraphPinObj != NULL && !GraphPinObj->IsPendingKill() && GraphPinObj->GetOuter() != NULL && GraphPinObj->GetOuter()->IsA(UEdGraphNode::StaticClass()))
 		{
-			// Inform the Drag and Drop operation that we are hovering over this pin.
-			TSharedPtr<FGraphEditorDragDropAction> DragConnectionOp = StaticCastSharedPtr<FGraphEditorDragDropAction>(Operation);
-			DragConnectionOp->SetHoveredPin(GraphPinObj);
+			if (GetIsConnectable())
+			{
+				// Inform the Drag and Drop operation that we are hovering over this pin.
+				TSharedPtr<FGraphEditorDragDropAction> DragConnectionOp = StaticCastSharedPtr<FGraphEditorDragDropAction>(Operation);
+				DragConnectionOp->SetHoveredPin(GraphPinObj);
+			}
 		}	
 
 		// Pins treat being dragged over the same as being hovered outside of drag and drop if they know how to respond to the drag action.
