Remove needless BroadcastChanged calls, which refresh the entire graph of a blueprint, from various nodes. Add GRAPHACTION_EditNode to synchronously update a node's SGraphNode after a change is made.

Also requires 24595399

diff --git a/Engine/Plugins/Experimental/PythonScriptPlugin/Source/PythonScriptPlugin/Private/K2Node_ExecutePythonScript.cpp b/Engine/Plugins/Experimental/PythonScriptPlugin/Source/PythonScriptPlugin/Private/K2Node_ExecutePythonScript.cpp
--- a/Engine/Plugins/Experimental/PythonScriptPlugin/Source/PythonScriptPlugin/Private/K2Node_ExecutePythonScript.cpp
+++ b/Engine/Plugins/Experimental/PythonScriptPlugin/Source/PythonScriptPlugin/Private/K2Node_ExecutePythonScript.cpp
@@ -136,7 +136,6 @@ void UK2Node_ExecutePythonScript::SynchronizeArgumentPinTypeImpl(UEdGraphPin* Pi
 		if (!Blueprint->bBeingCompiled)
 		{
 			FBlueprintEditorUtils::MarkBlueprintAsModified(Blueprint);
-			Blueprint->BroadcastChanged();
 		}
 	}
 }
diff --git a/Engine/Source/Editor/BlueprintGraph/Private/K2Node_FormatText.cpp b/Engine/Source/Editor/BlueprintGraph/Private/K2Node_FormatText.cpp
--- a/Engine/Source/Editor/BlueprintGraph/Private/K2Node_FormatText.cpp
+++ b/Engine/Source/Editor/BlueprintGraph/Private/K2Node_FormatText.cpp
@@ -115,7 +115,6 @@ void UK2Node_FormatText::SynchronizeArgumentPinType(UEdGraphPin* Pin)
 			if (!Blueprint->bBeingCompiled)
 			{
 				FBlueprintEditorUtils::MarkBlueprintAsModified(Blueprint);
-				Blueprint->BroadcastChanged();
 			}
 		}
 	}
diff --git a/Engine/Source/Editor/BlueprintGraph/Private/K2Node_MacroInstance.cpp b/Engine/Source/Editor/BlueprintGraph/Private/K2Node_MacroInstance.cpp
--- a/Engine/Source/Editor/BlueprintGraph/Private/K2Node_MacroInstance.cpp
+++ b/Engine/Source/Editor/BlueprintGraph/Private/K2Node_MacroInstance.cpp
@@ -336,7 +336,6 @@ void UK2Node_MacroInstance::NodeConnectionListChanged()
 		if (Blueprint && !Blueprint->bBeingCompiled)
 		{
 			FBlueprintEditorUtils::MarkBlueprintAsModified(Blueprint);
-			Blueprint->BroadcastChanged();
 		}
 	}
 }
diff --git a/Engine/Source/Editor/BlueprintGraph/Private/K2Node_MathExpression.cpp b/Engine/Source/Editor/BlueprintGraph/Private/K2Node_MathExpression.cpp
--- a/Engine/Source/Editor/BlueprintGraph/Private/K2Node_MathExpression.cpp
+++ b/Engine/Source/Editor/BlueprintGraph/Private/K2Node_MathExpression.cpp
@@ -2725,8 +2725,6 @@ void UK2Node_MathExpression::RebuildExpression(FString InExpression)
 		// finally, recompile
 		UBlueprint* Blueprint = FBlueprintEditorUtils::FindBlueprintForNodeChecked(this);
 		FBlueprintEditorUtils::MarkBlueprintAsModified(Blueprint);
-		// The UI needs a refresh, so notify any interested parties that the blueprint has changed
-		Blueprint->BroadcastChanged();
 	}
 }
 
diff --git a/Engine/Source/Editor/BlueprintGraph/Private/K2Node_Select.cpp b/Engine/Source/Editor/BlueprintGraph/Private/K2Node_Select.cpp
--- a/Engine/Source/Editor/BlueprintGraph/Private/K2Node_Select.cpp
+++ b/Engine/Source/Editor/BlueprintGraph/Private/K2Node_Select.cpp
@@ -714,7 +714,6 @@ void UK2Node_Select::NodeConnectionListChanged()
 		if(!Blueprint->bBeingCompiled)
 		{
 			FBlueprintEditorUtils::MarkBlueprintAsModified(Blueprint);
-			Blueprint->BroadcastChanged();
 		}
 	}
 }
@@ -764,7 +763,6 @@ void UK2Node_Select::ChangePinType(UEdGraphPin* Pin)
 	if (!Blueprint->bBeingCompiled)
 	{
 		FBlueprintEditorUtils::MarkBlueprintAsModified(Blueprint);
-		Blueprint->BroadcastChanged();
 	}
 }
 
diff --git a/Engine/Source/Editor/GraphEditor/Private/SGraphEditorImpl.cpp b/Engine/Source/Editor/GraphEditor/Private/SGraphEditorImpl.cpp
--- a/Engine/Source/Editor/GraphEditor/Private/SGraphEditorImpl.cpp
+++ b/Engine/Source/Editor/GraphEditor/Private/SGraphEditorImpl.cpp
@@ -192,12 +192,13 @@ void SGraphEditorImpl::OnGraphChanged(const FEdGraphEditAction& InAction)
 		const bool bWasAddAction = (InAction.Action & GRAPHACTION_AddNode) != 0;
 		const bool bWasSelectAction = (InAction.Action & GRAPHACTION_SelectNode) != 0;
 		const bool bWasRemoveAction = (InAction.Action & GRAPHACTION_RemoveNode) != 0;
+		const bool bWasEditAction = (InAction.Action & GRAPHACTION_EditNode) != 0;
 
 		// If we did a 'default action' (or some other action not handled by SGraphPanel::OnGraphChanged
 		// or if we're using a schema that always needs a full refresh, then purge the current nodes
 		// and queue an update:
 		if (bSchemaRequiresFullRefresh || 
-			(!bWasAddAction && !bWasSelectAction && !bWasRemoveAction) )
+			(!bWasAddAction && !bWasSelectAction && !bWasRemoveAction && !bWasEditAction) )
 		{
 			GraphPanel->PurgeVisualRepresentation();
 			// Trigger the refresh
@@ -209,6 +210,13 @@ void SGraphEditorImpl::OnGraphChanged(const FEdGraphEditAction& InAction)
 		{
 			NumNodesAddedSinceLastPointerPosition++;
 		}
+		else if (bWasEditAction && !bSchemaRequiresFullRefresh)
+		{
+			for (const UEdGraphNode* Node : InAction.Nodes)
+			{
+				RefreshNode(const_cast<UEdGraphNode&>(*Node));
+			}
+		}
 	}
 }
 
diff --git a/Engine/Source/Runtime/Engine/Private/EdGraph/EdGraph.cpp b/Engine/Source/Runtime/Engine/Private/EdGraph/EdGraph.cpp
--- a/Engine/Source/Runtime/Engine/Private/EdGraph/EdGraph.cpp
+++ b/Engine/Source/Runtime/Engine/Private/EdGraph/EdGraph.cpp
@@ -285,7 +285,11 @@ void UEdGraph::NotifyGraphChanged()
 
 void UEdGraph::NotifyNodeChanged(const UEdGraphNode* Node)
 {
-	NotifyGraphChanged();
+	FEdGraphEditAction Action;
+	Action.Graph = this;
+	Action.Action = GRAPHACTION_EditNode;
+	Action.Nodes.Add(Node);
+	NotifyGraphChanged(Action);
 }
 
 void UEdGraph::NotifyGraphChanged(const FEdGraphEditAction& InAction)
diff --git a/Engine/Source/Runtime/Engine/Public/GraphEditAction.h b/Engine/Source/Runtime/Engine/Public/GraphEditAction.h
--- a/Engine/Source/Runtime/Engine/Public/GraphEditAction.h
+++ b/Engine/Source/Runtime/Engine/Public/GraphEditAction.h
@@ -18,6 +18,9 @@ enum EEdGraphActionType
 
 	/** A node was removed from the graph at the user's request */
 	GRAPHACTION_RemoveNode = 0x1 << 2,
+
+	/** A node was edited in the graph by a user */
+	GRAPHACTION_EditNode = 0x1 << 3,
 };
 
 
