Make sure we don't dereference GraphPinObj after GraphData has been invalidated

diff --git a/Engine/Source/Editor/GraphEditor/Private/SGraphPin.cpp b/Engine/Source/Editor/GraphEditor/Private/SGraphPin.cpp
--- a/Engine/Source/Editor/GraphEditor/Private/SGraphPin.cpp
+++ b/Engine/Source/Editor/GraphEditor/Private/SGraphPin.cpp
@@ -907,12 +907,12 @@ UEdGraphPin* SGraphPin::GetPinObj() const
 {
 	ensureMsgf(!bGraphDataInvalid, TEXT("The Graph Pin Object has been invalidated. Someone is keeping a hard ref on the SGraphPin (%s). See InvalidateGraphData for more info"), *ToString());
 
-	if(GraphPinObj && GraphPinObj->bWasTrashed)
+	if (bGraphDataInvalid || (GraphPinObj && GraphPinObj->bWasTrashed))
 	{
 		return nullptr;
 	}
 
-	return !bGraphDataInvalid ? GraphPinObj : nullptr;
+	return GraphPinObj;
 }
 
 /** @param OwnerNode  The SGraphNode that this pin belongs to */
