Optimize async blueprint context menu building by avoiding resorting and inserting nodes in the desired order

diff --git a/Engine/Source/Editor/GraphEditor/Private/GraphActionNode.cpp b/Engine/Source/Editor/GraphEditor/Private/GraphActionNode.cpp
--- a/Engine/Source/Editor/GraphEditor/Private/GraphActionNode.cpp
+++ b/Engine/Source/Editor/GraphEditor/Private/GraphActionNode.cpp
@@ -203,6 +203,34 @@ TSharedPtr<FGraphActionNode> FGraphActionNode::AddChild(FGraphActionListBuilderB
 	return ActionNode;
 }
 
+//------------------------------------------------------------------------------
+void FGraphActionNode::AddChildAlphabetical(FGraphActionListBuilderBase::ActionGroup const& ActionSet)
+{
+	TSharedPtr<FGraphActionNode> ActionNode = FGraphActionNode::NewActionNode(ActionSet.Actions);
+	check(ActionNode->SectionID == INVALID_SECTION_ID); // this method does not support sections, those should be built statically
+
+	// if a divider hasn't been created for the grouping, create one:
+	AddChildGrouping(ActionNode, this->AsShared());
+
+	// find or add categories iteratively, inserting as needed:
+	FGraphActionNode* OwningCategory = this;
+	const TArray<FString>& CategoryStack = ActionSet.GetCategoryChain();
+	for (const FString& CategorySection : CategoryStack)
+	{
+		TSharedPtr<FGraphActionNode> CategoryNode = OwningCategory->FindMatchingParent(CategorySection, ActionNode);
+		if (!CategoryNode.IsValid())
+		{
+			CategoryNode = NewCategoryNode(CategorySection, ActionNode->Grouping, ActionNode->SectionID);
+			OwningCategory->InsertChildAlphabetical(CategoryNode);
+		}
+
+		OwningCategory = CategoryNode.Get();
+	}
+
+	// finally insert the leaf:
+	OwningCategory->InsertChildAlphabetical(ActionNode);
+}
+
 //------------------------------------------------------------------------------
 TSharedPtr<FGraphActionNode> FGraphActionNode::AddSection(int32 InGrouping, int32 InSectionID)
 {
@@ -593,33 +621,69 @@ void FGraphActionNode::InsertChild(TSharedPtr<FGraphActionNode> NodeToAdd)
 	}
 	// we don't use group-dividers inside of sections (we use groups to more to
 	// hardcode the order), but if this isn't in a section...
-	else if (!ChildGroupings.Contains(NodeToAdd->Grouping))
+	else
+	{
+		AddChildGrouping(NodeToAdd, NodeToAdd->ParentNode);
+	}
+
+	NodeToAdd->InsertOrder = Children.Num();
+	Children.Add(NodeToAdd);
+	if (NodeToAdd->IsCategoryNode())
 	{
-		// don't need a divider if this is the first group
-		if (ChildGroupings.Num() > 0)
+		CategoryNodes.Add(NodeToAdd->DisplayText.ToString(), NodeToAdd);
+	}
+}
+
+//------------------------------------------------------------------------------
+void FGraphActionNode::AddChildGrouping(TSharedPtr<FGraphActionNode> ActionNode, TWeakPtr<FGraphActionNode> Parent)
+{
+	if (ChildGroupings.Find(ActionNode->Grouping))
+	{
+		return;
+	}
+
+	if (ChildGroupings.Num() > 0)
+	{
+		int32 LowestGrouping = MAX_int32;
+		for (int32 Group : ChildGroupings)
 		{
-			int32 LowestGrouping = MAX_int32;
-			for (int32 Group : ChildGroupings)
-			{
-				LowestGrouping = FMath::Min(LowestGrouping, Group);
-			}
-			// dividers come at the end of a menu group, so it would be 
-			// undesirable to add it for NodeToAdd->Grouping if that group is 
-			// lower than all the others (the lowest group should not have a 
-			// divider associated with it)
-			int32 DividerGrouping = FMath::Max(LowestGrouping, NodeToAdd->Grouping);
-
-			ChildGroupings.Add(NodeToAdd->Grouping); // to avoid recursion, add before we insert
-			InsertChild(NewGroupDividerNode(NodeToAdd->ParentNode, DividerGrouping));
+			LowestGrouping = FMath::Min(LowestGrouping, Group);
 		}
-		else
+		// dividers come at the end of a menu group, so it would be 
+		// undesirable to add it for NodeToAdd->Grouping if that group is 
+		// lower than all the others (the lowest group should not have a 
+		// divider associated with it)
+		int32 DividerGrouping = FMath::Max(LowestGrouping, ActionNode->Grouping);
+
+		ChildGroupings.Add(ActionNode->Grouping); // to avoid recursion, add before we insert
+		InsertChild(NewGroupDividerNode(this->AsShared(), DividerGrouping));
+	}
+	else
+	{
+		ChildGroupings.Add(ActionNode->Grouping);
+	}
+}
+
+//------------------------------------------------------------------------------
+void FGraphActionNode::InsertChildAlphabetical(TSharedPtr<FGraphActionNode> NodeToAdd)
+{
+	check(NodeToAdd->SectionID == INVALID_SECTION_ID);
+	int Idx = Algo::LowerBound(Children, NodeToAdd, FGraphActionNodeImpl::AlphabeticalNodeCompare);
+	if (Idx != INDEX_NONE)
+	{
+		NodeToAdd->InsertOrder = Idx;
+		Children.Insert(NodeToAdd, Idx);
+		for (int32 I = Idx + 1; I < Children.Num(); ++I)
 		{
-			ChildGroupings.Add(NodeToAdd->Grouping);
+			++(Children[I]->InsertOrder);
 		}
 	}
+	else
+	{
+		NodeToAdd->InsertOrder = Children.Num();
+		Children.Add(NodeToAdd);
+	}
 
-	NodeToAdd->InsertOrder = Children.Num();
-	Children.Add(NodeToAdd);
 	if (NodeToAdd->IsCategoryNode())
 	{
 		CategoryNodes.Add(NodeToAdd->DisplayText.ToString(), NodeToAdd);
diff --git a/Engine/Source/Editor/GraphEditor/Private/SGraphActionMenu.cpp b/Engine/Source/Editor/GraphEditor/Private/SGraphActionMenu.cpp
--- a/Engine/Source/Editor/GraphEditor/Private/SGraphActionMenu.cpp
+++ b/Engine/Source/Editor/GraphEditor/Private/SGraphActionMenu.cpp
@@ -305,6 +305,8 @@ FString SGraphActionMenu::LastUsedFilterText;
 
 void SGraphActionMenu::Construct( const FArguments& InArgs, bool bIsReadOnly/* = true*/ )
 {
+	this->SelectedSuggestionScore = TNumericLimits<float>::Lowest();
+	this->SelectedSuggestionSourceIndex = INDEX_NONE;
 	this->SelectedSuggestion = INDEX_NONE;
 	this->bIgnoreUIUpdate = false;
 	this->bUseSectionStyling = InArgs._UseSectionStyling;
@@ -714,10 +716,29 @@ void RestoreExpansionState(TSharedPtr< STreeView<ItemType> > InTree, const TArra
 	}
 }
 
+void SGraphActionMenu::UpdateForNewActions(int32 IdxStart)
+{
+	check(bAlphaSortItems && bSortItemsRecursively);
+
+	FScoreResults Results = ScoreAndAddActions(IdxStart);
+	UpdateActiveSelection(Results);
+
+	if (ShouldExpandNodes())
+	{
+		// Expand all
+		FilteredRootAction->ExpandAllChildren(TreeView);
+	}
+
+	TreeView->RequestTreeRefresh();
+}
+
 void SGraphActionMenu::GenerateFilteredItems(bool bPreserveExpansion)
 {
 	TRACE_CPUPROFILER_EVENT_SCOPE(SGraphActionMenu::GenerateFilteredItems);
 
+	SelectedSuggestionScore = TNumericLimits<float>::Lowest();
+	SelectedSuggestionSourceIndex = INDEX_NONE;
+
 	// First, save off current expansion state
 	TSet< TSharedPtr<FGraphActionNode> > OldExpansionState;
 	if(bPreserveExpansion)
@@ -740,118 +761,13 @@ void SGraphActionMenu::GenerateFilteredItems(bool bPreserveExpansion)
 		}
 	}
 	
-	// Trim and sanitized the filter text (so that it more likely matches the action descriptions)
-	FString TrimmedFilterString = FText::TrimPrecedingAndTrailing(GetFilterText()).ToString();
-
-	// Remember the last filter string to that external clients can access it
-	LastUsedFilterText = TrimmedFilterString;
-
-	// Tokenize the search box text into a set of terms; all of them must be present to pass the filter
-	TArray<FString> FilterTerms;
-	TrimmedFilterString.ParseIntoArray(FilterTerms, TEXT(" "), true);
-	for (FString& String : FilterTerms)
-	{
-		String = String.ToLower();
-	}
-
-	// Generate a list of sanitized versions of the strings
-	TArray<FString> SanitizedFilterTerms;
-	for (int32 iFilters = 0; iFilters < FilterTerms.Num() ; iFilters++)
-	{
-		FString EachString = FName::NameToDisplayString( FilterTerms[iFilters], false );
-		EachString = EachString.Replace( TEXT( " " ), TEXT( "" ) );
-		SanitizedFilterTerms.Add( EachString );
-	}
-	ensure( SanitizedFilterTerms.Num() == FilterTerms.Num() );// Both of these should match !
-
-	const bool bRequiresFiltering = FilterTerms.Num() > 0;
-	float BestMatchCount = TNumericLimits<float>::Lowest();
-	int32 BestMatchIndex = INDEX_NONE;
+	FScoreResults Results = ScoreAndAddActions();
 
-	// Get the schema of the graph that we are in so that we can correctly get the action weight
-	const UEdGraphSchema* ActionSchema = GraphObj ? GraphObj->GetSchema() : GetDefault<UEdGraphSchema>();
-	check(ActionSchema);
-
-	for (int32 CurTypeIndex = 0; CurTypeIndex < AllActions->GetNumActions(); ++CurTypeIndex)
-	{
-		FGraphActionListBuilderBase::ActionGroup& CurrentAction = AllActions->GetAction(CurTypeIndex);
-
-		// If we're filtering, search check to see if we need to show this action
-		bool bShowAction = true;
-		float EachWeight = TNumericLimits<float>::Lowest();
-		if (bRequiresFiltering)
-		{
-			// Combine the actions string, separate with \n so terms don't run into each other, and remove the spaces (incase the user is searching for a variable)
-			// In the case of groups containing multiple actions, they will have been created and added at the same place in the code, using the same description
-			// and keywords, so we only need to use the first one for filtering.
-			const FString& SearchText = CurrentAction.GetSearchTextForFirstAction();
-
-			FString EachTermSanitized;
-			for (int32 FilterIndex = 0; (FilterIndex < FilterTerms.Num()) && bShowAction; ++FilterIndex)
-			{
-				const bool bMatchesTerm = (SearchText.Contains(FilterTerms[FilterIndex], ESearchCase::CaseSensitive) || (SearchText.Contains(SanitizedFilterTerms[FilterIndex], ESearchCase::CaseSensitive) == true));
-				bShowAction = bShowAction && bMatchesTerm;
-			}
-
-			// Only if we are going to show the action do we want to generate the weight of the filter text
-			if (bShowAction)
-			{
-				// Get the 'weight' of this in relation to the filter
-				EachWeight = ActionSchema->GetActionFilteredWeight(CurrentAction, FilterTerms, SanitizedFilterTerms, DraggedFromPins);
-			}
-		}
-
-		if (bShowAction)
-		{
-			// If this action has a greater relevance than others, cache its index.
-			if( EachWeight > BestMatchCount )
-			{
-				BestMatchCount = EachWeight;
-				BestMatchIndex = CurTypeIndex;
-			}
-			FilteredRootAction->AddChild(CurrentAction);
-		}
-	}
 	FilteredRootAction->SortChildren(bAlphaSortItems, bSortItemsRecursively);
 
 	TreeView->RequestTreeRefresh();
 
-	// Update the filtered list (needs to be done in a separate pass because the list is sorted as items are inserted)
-	FilteredActionNodes.Empty();
-	FilteredRootAction->GetLeafNodes(FilteredActionNodes);
-
-	// Get _all_ new nodes (flattened tree basically)
-	TArray< TSharedPtr<FGraphActionNode> > AllNodes;
-	FilteredRootAction->GetAllNodes(AllNodes);
-
-	// If theres a BestMatchIndex find it in the actions nodes and select it (maybe this should check the current selected suggestion first ?)
-	if( BestMatchIndex != INDEX_NONE ) 
-	{
-		FGraphActionListBuilderBase::ActionGroup& FilterSelectAction = AllActions->GetAction( BestMatchIndex );
-		if( FilterSelectAction.Actions[0].IsValid() == true )
-		{
-			for (int32 iNode = 0; iNode < FilteredActionNodes.Num() ; iNode++)
-			{
-				if( FilteredActionNodes[ iNode ].Get()->GetPrimaryAction() == FilterSelectAction.Actions[ 0 ] )
-				{
-					SelectedSuggestion = iNode;
-				}
-			}	
-		}	
-	}
-
-	// Make sure the selected suggestion stays within the filtered list
-	if ((SelectedSuggestion >= 0) && (FilteredActionNodes.Num() > 0))
-	{
-		//@TODO: Should try to actually maintain the highlight on the same item if it survived the filtering
-		SelectedSuggestion = FMath::Clamp<int32>(SelectedSuggestion, 0, FilteredActionNodes.Num() - 1);
-		MarkActiveSuggestion();
-	}
-	else
-	{
-		SelectedSuggestion = INDEX_NONE;
-	}
-
+	UpdateActiveSelection(Results);
 
 	if (ShouldExpandNodes())
 	{
@@ -860,6 +776,10 @@ void SGraphActionMenu::GenerateFilteredItems(bool bPreserveExpansion)
 	}
 	else
 	{
+		// Get _all_ new nodes (flattened tree basically)
+		TArray< TSharedPtr<FGraphActionNode> > AllNodes;
+		FilteredRootAction->GetAllNodes(AllNodes);
+
 		// Expand to match the old state
 		RestoreExpansionState< TSharedPtr<FGraphActionNode> >(TreeView, AllNodes, OldExpansionState, CompareGraphActionNode);
 	}
@@ -1348,6 +1268,127 @@ void SGraphActionMenu::OnSetExpansionRecursive(TSharedPtr<FGraphActionNode> InTr
 	}
 }
 
+SGraphActionMenu::FScoreResults SGraphActionMenu::ScoreAndAddActions(int32 StartingIndex)
+{
+	// Trim and sanitized the filter text (so that it more likely matches the action descriptions)
+	FString TrimmedFilterString = FText::TrimPrecedingAndTrailing(GetFilterText()).ToString();
+
+	// Remember the last filter string to that external clients can access it
+	LastUsedFilterText = TrimmedFilterString;
+
+	// Tokenize the search box text into a set of terms; all of them must be present to pass the filter
+	TArray<FString> FilterTerms;
+	TrimmedFilterString.ParseIntoArray(FilterTerms, TEXT(" "), true);
+	for (FString& String : FilterTerms)
+	{
+		String = String.ToLower();
+	}
+
+	// Generate a list of sanitized versions of the strings
+	TArray<FString> SanitizedFilterTerms;
+	for (int32 iFilters = 0; iFilters < FilterTerms.Num(); iFilters++)
+	{
+		FString EachString = FName::NameToDisplayString(FilterTerms[iFilters], false);
+		EachString = EachString.Replace(TEXT(" "), TEXT(""));
+		SanitizedFilterTerms.Add(EachString);
+	}
+	ensure(SanitizedFilterTerms.Num() == FilterTerms.Num());// Both of these should match !
+
+	const bool bRequiresFiltering = FilterTerms.Num() > 0;
+	float BestMatchCount = SelectedSuggestionScore;
+	int32 BestMatchIndex = SelectedSuggestionSourceIndex;
+
+	// Get the schema of the graph that we are in so that we can correctly get the action weight
+	const UEdGraphSchema* ActionSchema = GraphObj ? GraphObj->GetSchema() : GetDefault<UEdGraphSchema>();
+	check(ActionSchema);
+
+	bool bIsPartialBuild = StartingIndex != INDEX_NONE;
+	const int32 NumActions = AllActions->GetNumActions();
+	for (int32 CurTypeIndex = bIsPartialBuild ? StartingIndex : 0; CurTypeIndex < NumActions; ++CurTypeIndex)
+	{
+		FGraphActionListBuilderBase::ActionGroup& CurrentAction = AllActions->GetAction(CurTypeIndex);
+
+		// If we're filtering, search check to see if we need to show this action
+		bool bShowAction = true;
+		float EachWeight = TNumericLimits<float>::Lowest();
+
+		const FString& SearchText = CurrentAction.GetSearchTextForFirstAction();
+		for (int32 FilterIndex = 0; (FilterIndex < FilterTerms.Num()) && bShowAction; ++FilterIndex)
+		{
+			const bool bMatchesTerm = (SearchText.Contains(FilterTerms[FilterIndex], ESearchCase::CaseSensitive) || (SearchText.Contains(SanitizedFilterTerms[FilterIndex], ESearchCase::CaseSensitive) == true));
+			bShowAction = bMatchesTerm;
+		}
+
+		if (!bShowAction)
+		{
+			continue;
+		}
+
+		if (bRequiresFiltering)
+		{
+			// Get the 'weight' of this in relation to the filter
+			EachWeight = ActionSchema->GetActionFilteredWeight(CurrentAction, FilterTerms, SanitizedFilterTerms, DraggedFromPins);
+			// If this action has a greater relevance than others, cache its index.
+			if (EachWeight > BestMatchCount)
+			{
+				BestMatchCount = EachWeight;
+				BestMatchIndex = CurTypeIndex;
+			}
+		}
+
+		if (bIsPartialBuild)
+		{
+			FilteredRootAction->AddChildAlphabetical(CurrentAction);
+		}
+		else
+		{
+			FilteredRootAction->AddChild(CurrentAction);
+		}
+	}
+
+	return {BestMatchIndex, BestMatchCount};
+}
+
+void SGraphActionMenu::UpdateActiveSelection(SGraphActionMenu::FScoreResults ForResults)
+{
+	int32 BestMatchIndex = ForResults.BestMatchIndex;
+	float BestMatchCount = ForResults.BestMatchScore;
+	// Update the filtered list (needs to be done in a separate pass because the list is sorted as items are inserted)
+	FilteredActionNodes.Reset();
+	FilteredRootAction->GetLeafNodes(FilteredActionNodes);
+
+	// If theres a BestMatchIndex find it in the actions nodes and select it (maybe this should check the current selected suggestion first ?)
+	if (BestMatchIndex != INDEX_NONE)
+	{
+		FGraphActionListBuilderBase::ActionGroup& FilterSelectAction = AllActions->GetAction(BestMatchIndex);
+		if (FilterSelectAction.Actions[0].IsValid() == true)
+		{
+			for (int32 iNode = 0; iNode < FilteredActionNodes.Num(); iNode++)
+			{
+				if (FilteredActionNodes[iNode].Get()->GetPrimaryAction() == FilterSelectAction.Actions[0])
+				{
+					SelectedSuggestion = iNode;
+					SelectedSuggestionScore = BestMatchCount;
+					SelectedSuggestionSourceIndex = BestMatchIndex;
+				}
+			}
+		}
+	}
+
+	// Make sure the selected suggestion stays within the filtered list
+	if ((SelectedSuggestion >= 0) && (FilteredActionNodes.Num() > 0))
+	{
+		//@TODO: Should try to actually maintain the highlight on the same item if it survived the filtering
+		SelectedSuggestion = FMath::Clamp<int32>(SelectedSuggestion, 0, FilteredActionNodes.Num() - 1);
+		MarkActiveSuggestion();
+	}
+	else
+	{
+		SelectedSuggestionScore = TNumericLimits<float>::Lowest();
+		SelectedSuggestionSourceIndex = INDEX_NONE;
+		SelectedSuggestion = INDEX_NONE;
+	}
+}
 /////////////////////////////////////////////////////
 
 #undef LOCTEXT_NAMESPACE
diff --git a/Engine/Source/Editor/GraphEditor/Public/GraphActionNode.h b/Engine/Source/Editor/GraphEditor/Public/GraphActionNode.h
--- a/Engine/Source/Editor/GraphEditor/Public/GraphActionNode.h
+++ b/Engine/Source/Editor/GraphEditor/Public/GraphActionNode.h
@@ -49,13 +49,21 @@ public:
 	 * the provided ActionSet. 
 	 *
 	 * NOTE: This does NOT insert the node in a sorted manner. Call SortChildren() 
-	 *       separately to accomplish that (done for performance reasons).
+	 *       separately or use AddChildAlphabetical
 	 * 
 	 * @param  ActionSet	A list of actions that you want the node to execute when picked.
 	 * @return The new action node.
 	 */
 	TSharedPtr<FGraphActionNode> AddChild(FGraphActionListBuilderBase::ActionGroup const& ActionSet);
 
+	/**
+	 * Inserts a new action node (and any required category nodes) based off
+	 * the provided ActionSet. Inserts in alphabetical order.
+	 *
+	 * @param  ActionSet	A list of actions that you want the node to execute when picked.
+	 */
+	void AddChildAlphabetical(FGraphActionListBuilderBase::ActionGroup const& ActionSet);
+
 	TSharedPtr<FGraphActionNode> AddSection(int32 Grouping, int32 InSectionID);
 
 	/**
@@ -263,6 +271,8 @@ private:
 	 */
 	void InsertChild(TSharedPtr<FGraphActionNode> NodeToAdd);
 
+	void AddChildGrouping(TSharedPtr<FGraphActionNode> ActionNode, TWeakPtr<FGraphActionNode> Parent);
+	void InsertChildAlphabetical(TSharedPtr<FGraphActionNode> NodeToAdd);
 private:
 	/** The category or action name (depends on what type of node this is) */
 	FText DisplayText;
diff --git a/Engine/Source/Editor/GraphEditor/Public/SGraphActionMenu.h b/Engine/Source/Editor/GraphEditor/Public/SGraphActionMenu.h
--- a/Engine/Source/Editor/GraphEditor/Public/SGraphActionMenu.h
+++ b/Engine/Source/Editor/GraphEditor/Public/SGraphActionMenu.h
@@ -217,6 +217,10 @@ protected:
 	/** Root of filtered actions tree */
 	TSharedPtr<FGraphActionNode> FilteredRootAction;
 
+	/** Stored score for our current selection, so that we can quickly maintain selection when building the list asynchronously */
+	float SelectedSuggestionScore;
+	/** Stored index in the AllActions list builder that we think is the best fit */
+	int32 SelectedSuggestionSourceIndex;
 	/** Used to track selected action for keyboard interaction */
 	int32 SelectedSuggestion;
 	/** Allows us to set selection (via keyboard) without triggering action */
@@ -310,6 +314,8 @@ public:
 	/* Handler for mouse button going down */
 	bool OnMouseButtonDownEvent( TWeakPtr<FEdGraphSchemaAction> InAction );
 
+	/** Updates the displayed list starting from IdxStart, useful for async building the display list of actions */
+	void UpdateForNewActions(int32 IdxStart);
 	/** Regenerated filtered results (FilteredRootAction and FilteredActionNodes) based on filter text  */ 
 	void GenerateFilteredItems(bool bPreserveExpansion);
 
@@ -370,7 +376,15 @@ protected:
 	void OnItemScrolledIntoView( TSharedPtr<FGraphActionNode> InActionNode, const TSharedPtr<ITableRow>& InWidget );
 	/** Callback for expanding tree items recursively */
 	void OnSetExpansionRecursive(TSharedPtr<FGraphActionNode> InTreeNode, bool bInIsItemExpanded);
-
+	/** Helper function for adding and scoring actions from our builder */
+	struct FScoreResults
+	{
+		int32 BestMatchIndex;
+		float BestMatchScore;
+	};
+	FScoreResults ScoreAndAddActions(int32 StartingIndex = INDEX_NONE);
+	/** Helper function to update the active selection after updating the displayed tree */
+	void UpdateActiveSelection(FScoreResults ForResults);
 private:
 	/** The pins that have been dragged off of to prompt the creation of this action menu. */
 	TArray<UEdGraphPin*> DraggedFromPins;
diff --git a/Engine/Source/Editor/Kismet/Private/BlueprintActionMenuBuilder.cpp b/Engine/Source/Editor/Kismet/Private/BlueprintActionMenuBuilder.cpp
--- a/Engine/Source/Editor/Kismet/Private/BlueprintActionMenuBuilder.cpp
+++ b/Engine/Source/Editor/Kismet/Private/BlueprintActionMenuBuilder.cpp
@@ -409,7 +409,7 @@ void FBlueprintActionMenuBuilderImpl::FMenuSectionDefinition::AddBoundMenuItems(
 	// UBlueprintNodeSpawner comes with an interface to test/bind through... 
 	for (auto BindingIt = PerspectiveBindings.CreateConstIterator(); BindingIt;)
 	{
-		FFieldVariant BindingObj = *BindingIt;
+		const FFieldVariant& BindingObj = *BindingIt;
 		++BindingIt;
 		bool const bIsLastBinding = !BindingIt;
 
diff --git a/Engine/Source/Editor/Kismet/Private/SBlueprintActionMenu.cpp b/Engine/Source/Editor/Kismet/Private/SBlueprintActionMenu.cpp
--- a/Engine/Source/Editor/Kismet/Private/SBlueprintActionMenu.cpp
+++ b/Engine/Source/Editor/Kismet/Private/SBlueprintActionMenu.cpp
@@ -736,9 +736,10 @@ void SBlueprintActionMenu::Tick(const FGeometry& AllottedGeometry, const double
 {
 	SBorder::Tick(AllottedGeometry, InCurrentTime, InDeltaTime);
 
+	int32 NewIdxStart = ContextMenuBuilder->GetNumActions();
 	if (ContextMenuBuilder.IsValid() && ContextMenuBuilder->ProcessPendingActions())
 	{
-		GraphActionMenu->GenerateFilteredItems(true);
+		GraphActionMenu->UpdateForNewActions(NewIdxStart);
 	}
 }
 
