Fix crash when undoing node creation

diff --git a/Engine/Source/Editor/GraphEditor/Private/SGraphPin.cpp b/Engine/Source/Editor/GraphEditor/Private/SGraphPin.cpp
--- a/Engine/Source/Editor/GraphEditor/Private/SGraphPin.cpp
+++ b/Engine/Source/Editor/GraphEditor/Private/SGraphPin.cpp
@@ -551,7 +551,10 @@ FReply SGraphPin::OnPinNameMouseDown( const FGeometry& SenderGeometry, const FPo
 
 FReply SGraphPin::OnMouseMove( const FGeometry& MyGeometry, const FPointerEvent& MouseEvent )
 {
-	bIsMovingLinks = MouseEvent.IsControlDown() && (GraphPinObj->LinkedTo.Num() > 0);
+	if (UEdGraphPin* GraphPin = GetPinObj())
+	{
+		bIsMovingLinks = MouseEvent.IsControlDown() && (GraphPin->LinkedTo.Num() > 0);
+	}
 
 	return FReply::Unhandled();
 }
