Fixed mismatched copy constructor/assignment operator in FEdGraphNodeHandle and FEdGraphPinHandle.

diff --git a/Engine/Source/Editor/GraphEditor/Public/EdGraphHandleTypes.h b/Engine/Source/Editor/GraphEditor/Public/EdGraphHandleTypes.h
--- a/Engine/Source/Editor/GraphEditor/Public/EdGraphHandleTypes.h
+++ b/Engine/Source/Editor/GraphEditor/Public/EdGraphHandleTypes.h
@@ -17,11 +17,6 @@ public:
 		, NodeName(InNode ? InNode->GetFName() : NAME_None)
 	{}
 
-	FORCEINLINE FEdGraphNodeHandle(const FEdGraphNodeHandle& InOther)
-		: Graph(InOther.Graph)
-		, NodeName(InOther.NodeName)
-	{}
-
 	friend FORCEINLINE uint32 GetTypeHash(const FEdGraphNodeHandle& InHandle)
 	{
 		return HashCombine(GetTypeHash(InHandle.Graph.ToSoftObjectPath()), GetTypeHash(InHandle.NodeName));
@@ -80,14 +75,6 @@ public:
 			PinDirection = EEdGraphPinDirection::EGPD_Input;
 		}
 	}
-		
-	FORCEINLINE FEdGraphPinHandle(const FEdGraphPinHandle& InOther)
-		: FEdGraphNodeHandle(InOther)
-		, PinName(InOther.PinName)
-		, PinDirection(InOther.PinDirection)
-		, PersistentPinGuid(InOther.PersistentPinGuid)
-		, PinIndex(InOther.PinIndex)
-	{}
 
 	friend FORCEINLINE uint32 GetTypeHash(const FEdGraphPinHandle& InHandle)
 	{
