Expose single click action within SGraphPanel to allow hooking custom logic

diff --git a/Engine/Source/Editor/GraphEditor/Private/SGraphEditorImpl.cpp b/Engine/Source/Editor/GraphEditor/Private/SGraphEditorImpl.cpp
--- a/Engine/Source/Editor/GraphEditor/Private/SGraphEditorImpl.cpp
+++ b/Engine/Source/Editor/GraphEditor/Private/SGraphEditorImpl.cpp
@@ -769,6 +769,7 @@ void SGraphEditorImpl::Construct( const FArguments& InArgs )
 			.OnDisallowedPinConnection( InArgs._GraphEvents.OnDisallowedPinConnection )
 			.ShowGraphStateOverlay(InArgs._ShowGraphStateOverlay)
 			.OnDoubleClicked(InArgs._GraphEvents.OnDoubleClicked)
+			.OnMouseButtonDown(InArgs._GraphEvents.OnMouseButtonDown)
 		]
 
 		// Indicator of current zoom level
diff --git a/Engine/Source/Editor/GraphEditor/Private/SGraphPanel.cpp b/Engine/Source/Editor/GraphEditor/Private/SGraphPanel.cpp
--- a/Engine/Source/Editor/GraphEditor/Private/SGraphPanel.cpp
+++ b/Engine/Source/Editor/GraphEditor/Private/SGraphPanel.cpp
@@ -107,7 +107,8 @@ void SGraphPanel::Construct( const SGraphPanel::FArguments& InArgs )
 	this->OnUpdateGraphPanel = InArgs._OnUpdateGraphPanel;
 	this->OnDisallowedPinConnection = InArgs._OnDisallowedPinConnection;
 	this->OnDoubleClicked = InArgs._OnDoubleClicked;
-
+	this->OnClicked = InArgs._OnMouseButtonDown;
+	
 	this->bPreservePinPreviewConnection = false;
 	this->PinVisibility = SGraphEditor::Pin_Show;
 
@@ -854,6 +855,14 @@ FReply SGraphPanel::OnMouseButtonDown(const FGeometry& MyGeometry, const FPointe
 		}
 	}
 
+	if (OnClicked.IsBound())
+	{
+		if (const FReply Reply = OnClicked.Execute(MyGeometry, MouseEvent); Reply.IsEventHandled())
+		{
+			return FReply::Handled();
+		}
+	}
+	
 	return SNodePanel::OnMouseButtonDown(MyGeometry, MouseEvent);
 }
 
diff --git a/Engine/Source/Editor/GraphEditor/Public/SGraphPanel.h b/Engine/Source/Editor/GraphEditor/Public/SGraphPanel.h
--- a/Engine/Source/Editor/GraphEditor/Public/SGraphPanel.h
+++ b/Engine/Source/Editor/GraphEditor/Public/SGraphPanel.h
@@ -110,6 +110,7 @@ public:
 		SLATE_EVENT( FOnUpdateGraphPanel, OnUpdateGraphPanel )
 		SLATE_EVENT( SGraphEditor::FOnDisallowedPinConnection, OnDisallowedPinConnection )
 		SLATE_EVENT( SGraphEditor::FOnDoubleClicked, OnDoubleClicked )
+		SLATE_EVENT( SGraphEditor::FOnMouseButtonDown, OnMouseButtonDown )
 		//SLATE_ATTRIBUTE( FGraphAppearanceInfo, Appearance )
 	SLATE_END_ARGS()
 
@@ -332,6 +333,9 @@ protected:
 	/** Called when the graph itself is double clicked */
 	SGraphEditor::FOnDoubleClicked OnDoubleClicked;
 
+	/** Called when the graph itself is clicked */
+	SGraphEditor::FOnMouseButtonDown OnClicked;
+	
 	/** Whether to draw the overlay indicating we're in PIE */
 	bool bShowPIENotification;
 
diff --git a/Engine/Source/Editor/UnrealEd/Public/GraphEditor.h b/Engine/Source/Editor/UnrealEd/Public/GraphEditor.h
--- a/Engine/Source/Editor/UnrealEd/Public/GraphEditor.h
+++ b/Engine/Source/Editor/UnrealEd/Public/GraphEditor.h
@@ -122,6 +122,8 @@ public:
 
 	DECLARE_DELEGATE( FOnDoubleClicked );
 
+	DECLARE_DELEGATE_RetVal_TwoParams( FReply, FOnMouseButtonDown, const FGeometry&, const FPointerEvent& );
+
 	/** Info about events occurring in/on the graph */
 	struct FGraphEditorEvents
 	{
@@ -151,6 +153,8 @@ public:
 		FOnDisallowedPinConnection OnDisallowedPinConnection;
 		/** Called when the graph itself is double clicked */
 		FOnDoubleClicked OnDoubleClicked;
+		/** Called when the graph is clicked */
+		FOnMouseButtonDown OnMouseButtonDown;
 	};
 
 
