Change Substrate topology node color.

diff --git a/Engine/Source/Editor/GraphEditor/Private/MaterialNodes/SGraphSubstrateMaterial.cpp b/Engine/Source/Editor/GraphEditor/Private/MaterialNodes/SGraphSubstrateMaterial.cpp
--- a/Engine/Source/Editor/GraphEditor/Private/MaterialNodes/SGraphSubstrateMaterial.cpp
+++ b/Engine/Source/Editor/GraphEditor/Private/MaterialNodes/SGraphSubstrateMaterial.cpp
@@ -10,6 +10,7 @@
 #include "Materials/MaterialExpressionStrata.h"
 #include "SGraphPin.h"
 #include "Math/Color.h"
+#include "Styling/StyleColors.h"
 
 enum class ESubstrateWidgetOutputType : uint8
 {
@@ -17,6 +18,9 @@ enum class ESubstrateWidgetOutputType : uint8
 	DetailPanel,
 };
 
+static EStyleColor GetSubstrateWidgetColor0() { return EStyleColor::AccentBlue;  }
+static EStyleColor GetSubstrateWidgetColor1() { return EStyleColor::AccentGreen;  }
+
 static const TSharedRef<SWidget> InternalProcessOperator(
 	const FStrataMaterialCompilationOutput& CompilationOutput, 
 	const FStrataOperator& Op, 
@@ -25,8 +29,8 @@ static const TSharedRef<SWidget> InternalProcessOperator(
 	EStyleColor OverrideColor)
 {
 	const bool bIsCurrent = OutputType == ESubstrateWidgetOutputType::Node ? InGuid == Op.MaterialExpressionGuid : true;
-	const EStyleColor Color0 = bIsCurrent ? EStyleColor::AccentPurple : OverrideColor;
-	const EStyleColor Color1 = bIsCurrent ? EStyleColor::AccentGreen : OverrideColor;
+	const EStyleColor Color0 = bIsCurrent ? GetSubstrateWidgetColor0() : OverrideColor;
+	const EStyleColor Color1 = bIsCurrent ? GetSubstrateWidgetColor1() : OverrideColor;
 	switch (Op.OperatorType)
 	{
 		case STRATA_OPERATOR_WEIGHT:
@@ -137,8 +141,9 @@ static const TSharedRef<SWidget> InternalProcessOperator(
 		break;
 	}
 
+	static FString NoVisualization = OutputType == ESubstrateWidgetOutputType::DetailPanel ? FString(TEXT("Tree Operator Error")) : FString();
 	auto TreeOperatorError = SNew(SErrorText)
-		.ErrorText(FText::FromString(TEXT("Tree Operator Error")))
+		.ErrorText(FText::FromString(NoVisualization))
 		.BackgroundColor(FSlateColor(EStyleColor::AccentRed));
 	return TreeOperatorError->AsShared();
 }
@@ -155,8 +160,8 @@ const TSharedRef<SWidget> FSubstrateWidget::ProcessOperator(const FStrataMateria
 
 void FSubstrateWidget::GetPinColor(TSharedPtr<SGraphPin>& Out, const UMaterialGraphNode* InNode)
 {	
-	const FLinearColor Color0 = FLinearColor(FColor::Purple);
-	const FLinearColor Color1 = FLinearColor(FColor::Green);
+	const FLinearColor Color0 = USlateThemeManager::Get().GetColor(GetSubstrateWidgetColor0());
+	const FLinearColor Color1 = USlateThemeManager::Get().GetColor(GetSubstrateWidgetColor1());
 
 	FLinearColor ColorModifier;
 	bool bHasColorModifier = false;
