Substrate - added support for progressive simplification of material, starting from operators at the bottom.

diff --git a/Engine/Source/Editor/GraphEditor/Private/MaterialNodes/SGraphSubstrateMaterial.cpp b/Engine/Source/Editor/GraphEditor/Private/MaterialNodes/SGraphSubstrateMaterial.cpp
--- a/Engine/Source/Editor/GraphEditor/Private/MaterialNodes/SGraphSubstrateMaterial.cpp
+++ b/Engine/Source/Editor/GraphEditor/Private/MaterialNodes/SGraphSubstrateMaterial.cpp
@@ -28,7 +28,7 @@ static const TSharedRef<SWidget> InternalProcessOperator(
 	const FGuid& InGuid, 
 	EStyleColor OverrideColor)
 {
-	const bool bIsCurrent = OutputType == ESubstrateWidgetOutputType::Node ? InGuid == Op.MaterialExpressionGuid : true;
+	const bool bIsCurrent = OutputType == ESubstrateWidgetOutputType::Node ? InGuid == Op.MaterialExpressionGuid : false;
 	const EStyleColor Color0 = bIsCurrent ? GetSubstrateWidgetColor0() : OverrideColor;
 	const EStyleColor Color1 = bIsCurrent ? GetSubstrateWidgetColor1() : OverrideColor;
 	switch (Op.OperatorType)
@@ -132,9 +132,10 @@ static const TSharedRef<SWidget> InternalProcessOperator(
 			}
 
 			const EStyleColor Color = OverrideColor != EStyleColor::MAX ? OverrideColor : (bIsCurrent ? EStyleColor::AccentGreen : EStyleColor::AccentGray);
+			const FSlateColor SlateColor = OutputType == ESubstrateWidgetOutputType::Node ? FSlateColor(Color) : FSlateColor(FLinearColor(0.16, 0.015, 0.24));
 			auto BSDF = SNew(SErrorText)
 				.ErrorText(FText::FromString(BSDFDesc))
-				.BackgroundColor(FSlateColor(Color))
+				.BackgroundColor(SlateColor)
 				.ToolTipText(FText::FromString(ToolTip));
 			return BSDF->AsShared();
 		}
