Fix CIS issue.

diff --git a/Engine/Source/Editor/GraphEditor/Private/MaterialNodes/SGraphNodeMaterialBase.cpp b/Engine/Source/Editor/GraphEditor/Private/MaterialNodes/SGraphNodeMaterialBase.cpp
--- a/Engine/Source/Editor/GraphEditor/Private/MaterialNodes/SGraphNodeMaterialBase.cpp
+++ b/Engine/Source/Editor/GraphEditor/Private/MaterialNodes/SGraphNodeMaterialBase.cpp
@@ -63,6 +63,7 @@
 #include "Widgets/Text/STextBlock.h"
 #include "Rendering/StrataMaterialShared.h"
 #include "SGraphSubstrateMaterial.h"
+#include "MaterialShared.h"
 
 class FWidgetStyle;
 class SWidget;
diff --git a/Engine/Source/Editor/GraphEditor/Private/MaterialNodes/SGraphSubstrateMaterial.cpp b/Engine/Source/Editor/GraphEditor/Private/MaterialNodes/SGraphSubstrateMaterial.cpp
--- a/Engine/Source/Editor/GraphEditor/Private/MaterialNodes/SGraphSubstrateMaterial.cpp
+++ b/Engine/Source/Editor/GraphEditor/Private/MaterialNodes/SGraphSubstrateMaterial.cpp
@@ -5,12 +5,14 @@
 #include "SGraphSubstrateMaterial.h"
 #include "Internationalization/Text.h"
 #include "Rendering/RenderingCommon.h"
+#include "Rendering/StrataMaterialShared.h"
 #include "MaterialGraph/MaterialGraphNode.h"
 #include "Materials/MaterialExpression.h"
 #include "Materials/MaterialExpressionStrata.h"
 #include "SGraphPin.h"
 #include "Math/Color.h"
 #include "Styling/StyleColors.h"
+#include "StrataDefinitions.h"
 
 enum class ESubstrateWidgetOutputType : uint8
 {
