Fix CIS issue.

diff --git a/Engine/Source/Editor/GraphEditor/Private/MaterialNodes/SGraphSubstrateMaterial.cpp b/Engine/Source/Editor/GraphEditor/Private/MaterialNodes/SGraphSubstrateMaterial.cpp
--- a/Engine/Source/Editor/GraphEditor/Private/MaterialNodes/SGraphSubstrateMaterial.cpp
+++ b/Engine/Source/Editor/GraphEditor/Private/MaterialNodes/SGraphSubstrateMaterial.cpp
@@ -13,6 +13,7 @@
 #include "Math/Color.h"
 #include "Styling/StyleColors.h"
 #include "StrataDefinitions.h"
+#include "Widgets/SBoxPanel.h"
 
 enum class ESubstrateWidgetOutputType : uint8
 {
