Add custom styleset support for K2Nodes

diff --git a/Engine/Source/Editor/GraphEditor/Private/KismetNodes/SGraphNodeK2Base.cpp b/Engine/Source/Editor/GraphEditor/Private/KismetNodes/SGraphNodeK2Base.cpp
--- a/Engine/Source/Editor/GraphEditor/Private/KismetNodes/SGraphNodeK2Base.cpp
+++ b/Engine/Source/Editor/GraphEditor/Private/KismetNodes/SGraphNodeK2Base.cpp
@@ -139,7 +139,7 @@ void SGraphNodeK2Base::UpdateCompactNode()
 			.AutoHeight()
 			[
 				SNew(STextBlock)
-					.TextStyle( FAppStyle::Get(), "Graph.CompactNode.Title" )
+					.TextStyle( GetStyleSet(), "Graph.CompactNode.Title" )
 					.Text( NodeTitle.Get(), &SNodeTitle::GetHeadTitle )
 					.WrapTextAt(128.0f)
 			]
@@ -222,12 +222,12 @@ void SGraphNodeK2Base::UpdateCompactNode()
 			+SOverlay::Slot()
 			[
 				SNew(SImage)
-				.Image( FAppStyle::GetBrush("Graph.VarNode.Body") )
+				.Image( GetStyleSet().GetBrush("Graph.VarNode.Body") )
 			]
 			+ SOverlay::Slot()
 			[
 				SNew(SImage)
-				.Image( FAppStyle::GetBrush("Graph.VarNode.Gloss") )
+				.Image( GetStyleSet().GetBrush("Graph.VarNode.Gloss") )
 			]
 			+SOverlay::Slot()
 			.Padding( FMargin(0,3) )
@@ -357,7 +357,7 @@ TSharedPtr<SToolTip> SGraphNodeK2Base::GetComplexTooltip()
 				.AutoHeight()
 			[
 				SNew(STextBlock)
-					.TextStyle( FAppStyle::Get(), "Documentation.SDocumentationTooltipSubdued")
+					.TextStyle( GetStyleSet(), "Documentation.SDocumentationTooltipSubdued")
 					.Text(this, &SGraphNodeK2Base::GetToolTipHeading)
 			]
 			+SVerticalBox::Slot()
@@ -367,7 +367,7 @@ TSharedPtr<SToolTip> SGraphNodeK2Base::GetComplexTooltip()
 				SNew(SBorder)
 				// use the border's padding to actually create the horizontal line
 				.Padding(1.f)
-				.BorderImage(FAppStyle::GetBrush(TEXT("Menu.Separator")))
+				.BorderImage(GetStyleSet().GetBrush(TEXT("Menu.Separator")))
 			]
 		]
 		// tooltip body
@@ -396,7 +396,7 @@ TSharedPtr<SToolTip> SGraphNodeK2Base::GetComplexTooltip()
 				SNew( STextBlock )
 				.ColorAndOpacity( FSlateColor::UseSubduedForeground() )
 				.Text( LOCTEXT( "NativeNodeName", "hold (Alt) for native node name" ) )
-				.TextStyle( &FAppStyle::GetWidgetStyle<FTextBlockStyle>(TEXT("Documentation.SDocumentationTooltip")) )
+				.TextStyle( &GetStyleSet().GetWidgetStyle<FTextBlockStyle>(TEXT("Documentation.SDocumentationTooltip")) )
 				.Visibility_Static(&Local::GetNativeNodeNameVisibility)
 			];
 	}
@@ -412,6 +412,16 @@ FText SGraphNodeK2Base::GetToolTipHeading() const
 	return FText::GetEmpty();
 }
 
+const ISlateStyle& SGraphNodeK2Base::GetStyleSet() const
+{
+	return Style ? *Style : FAppStyle::Get();
+}
+
+void SGraphNodeK2Base::Construct(const FArguments& InArgs)
+{
+	Style = InArgs._Style;
+}
+
 /**
  * Update this GraphNode to match the data that it is observing
  */
@@ -457,11 +467,11 @@ void SGraphNodeK2Base::GetOverlayBrushes(bool bSelected, const FVector2D WidgetS
 
 		if (Breakpoint->IsEnabledByUser())
 		{
-			BreakpointOverlayInfo.Brush = FAppStyle::GetBrush(FKismetDebugUtilities::IsBreakpointValid(*Breakpoint) ? TEXT("Kismet.DebuggerOverlay.Breakpoint.EnabledAndValid") : TEXT("Kismet.DebuggerOverlay.Breakpoint.EnabledAndInvalid"));
+			BreakpointOverlayInfo.Brush = GetStyleSet().GetBrush(FKismetDebugUtilities::IsBreakpointValid(*Breakpoint) ? TEXT("Kismet.DebuggerOverlay.Breakpoint.EnabledAndValid") : TEXT("Kismet.DebuggerOverlay.Breakpoint.EnabledAndInvalid"));
 		}
 		else
 		{
-			BreakpointOverlayInfo.Brush = FAppStyle::GetBrush(TEXT("Kismet.DebuggerOverlay.Breakpoint.Disabled"));
+			BreakpointOverlayInfo.Brush = GetStyleSet().GetBrush(TEXT("Kismet.DebuggerOverlay.Breakpoint.Disabled"));
 		}
 
 		if(BreakpointOverlayInfo.Brush != NULL)
@@ -480,7 +490,7 @@ void SGraphNodeK2Base::GetOverlayBrushes(bool bSelected, const FVector2D WidgetS
 		// Pick icon depending on whether we are on a hit breakpoint
 		const bool bIsOnHitBreakpoint = FKismetDebugUtilities::GetMostRecentBreakpointHit() == GraphNode;
 		
-		IPOverlayInfo.Brush = FAppStyle::GetBrush( bIsOnHitBreakpoint ? TEXT("Kismet.DebuggerOverlay.InstructionPointerBreakpoint") : TEXT("Kismet.DebuggerOverlay.InstructionPointer") );
+		IPOverlayInfo.Brush = GetStyleSet().GetBrush( bIsOnHitBreakpoint ? TEXT("Kismet.DebuggerOverlay.InstructionPointerBreakpoint") : TEXT("Kismet.DebuggerOverlay.InstructionPointer") );
 
 		if (IPOverlayInfo.Brush != NULL)
 		{
@@ -501,7 +511,7 @@ void SGraphNodeK2Base::GetOverlayBrushes(bool bSelected, const FVector2D WidgetS
 		if (Timeline && Timeline->bAutoPlay)
 		{
 			FOverlayBrushInfo IPOverlayInfo;
-			IPOverlayInfo.Brush = FAppStyle::GetBrush( TEXT("Graph.Node.Autoplay") );
+			IPOverlayInfo.Brush = GetStyleSet().GetBrush( TEXT("Graph.Node.Autoplay") );
 
 			if (IPOverlayInfo.Brush != NULL)
 			{
@@ -515,7 +525,7 @@ void SGraphNodeK2Base::GetOverlayBrushes(bool bSelected, const FVector2D WidgetS
 		if (Timeline && Timeline->bLoop)
 		{
 			FOverlayBrushInfo IPOverlayInfo;
-			IPOverlayInfo.Brush = FAppStyle::GetBrush( TEXT("Graph.Node.Loop") );
+			IPOverlayInfo.Brush = GetStyleSet().GetBrush( TEXT("Graph.Node.Loop") );
 
 			if (IPOverlayInfo.Brush != NULL)
 			{
@@ -535,7 +545,7 @@ void SGraphNodeK2Base::GetOverlayBrushes(bool bSelected, const FVector2D WidgetS
 		{
 			FOverlayBrushInfo IPOverlayInfo;
 
-			IPOverlayInfo.Brush = FAppStyle::GetBrush( ClientIcon );
+			IPOverlayInfo.Brush = GetStyleSet().GetBrush( ClientIcon );
 
 			if (IPOverlayInfo.Brush != NULL)
 			{
@@ -702,7 +712,7 @@ const FSlateBrush* SGraphNodeK2Base::GetShadowBrush(bool bSelected) const
 
 	if (bSelected && bCompactMode)
 	{
-		return FAppStyle::GetBrush( "Graph.VarNode.ShadowSelected" );
+		return GetStyleSet().GetBrush( "Graph.VarNode.ShadowSelected" );
 	}
 	else
 	{
@@ -716,13 +726,13 @@ void SGraphNodeK2Base::GetDiffHighlightBrushes(const FSlateBrush*& BackgroundOut
 	
 	if (K2Node->ShouldDrawCompact())
 	{
-		BackgroundOut = FAppStyle::GetBrush(TEXT("Graph.VarNode.DiffHighlight"));
-		ForegroundOut = FAppStyle::GetBrush(TEXT("Graph.VarNode.DiffHighlightShading"));
+		BackgroundOut = GetStyleSet().GetBrush(TEXT("Graph.VarNode.DiffHighlight"));
+		ForegroundOut = GetStyleSet().GetBrush(TEXT("Graph.VarNode.DiffHighlightShading"));
 	}
 	else
 	{
-		BackgroundOut = FAppStyle::GetBrush(TEXT("Graph.Node.DiffHighlight"));
-		ForegroundOut = FAppStyle::GetBrush(TEXT("Graph.Node.DiffHighlightShading"));
+		BackgroundOut = GetStyleSet().GetBrush(TEXT("Graph.Node.DiffHighlight"));
+		ForegroundOut = GetStyleSet().GetBrush(TEXT("Graph.Node.DiffHighlightShading"));
 	}
 }
 
diff --git a/Engine/Source/Editor/GraphEditor/Public/KismetNodes/SGraphNodeK2Base.h b/Engine/Source/Editor/GraphEditor/Public/KismetNodes/SGraphNodeK2Base.h
--- a/Engine/Source/Editor/GraphEditor/Public/KismetNodes/SGraphNodeK2Base.h
+++ b/Engine/Source/Editor/GraphEditor/Public/KismetNodes/SGraphNodeK2Base.h
@@ -22,6 +22,11 @@ struct FSlateBrush;
 class GRAPHEDITOR_API SGraphNodeK2Base : public SGraphNode
 {
 public:
+	SLATE_BEGIN_ARGS(SGraphNodeK2Base) { }
+		SLATE_ARGUMENT(TSharedPtr<ISlateStyle>, Style)
+	SLATE_END_ARGS()
+
+	void Construct(const FArguments& InArgs);
 
 	//~ Begin SGraphNode Interface
 	virtual void UpdateGraphNode() override;
@@ -50,6 +55,12 @@ protected :
 	/** Retrieves text to tack on to the top of the tooltip (above the standard text) */
 	FText GetToolTipHeading() const;
 
+	/** Returns the slate style to use for this node */
+	const ISlateStyle& GetStyleSet() const;
+
+	/** The slate style to use. We'll fall back to App Style if it's unset. */
+	TSharedPtr<ISlateStyle> Style;
+
 protected:
 	static const FLinearColor BreakpointHitColor;
 	static const FLinearColor LatentBubbleColor;
