[BugFix] False positive "Link Moved" diff callouts on Blueprint Diff tool

diff --git a/Engine/Source/Editor/GraphEditor/Private/GraphDiffControl.cpp b/Engine/Source/Editor/GraphEditor/Private/GraphDiffControl.cpp
--- a/Engine/Source/Editor/GraphEditor/Private/GraphDiffControl.cpp
+++ b/Engine/Source/Editor/GraphEditor/Private/GraphDiffControl.cpp
@@ -471,7 +471,8 @@ static UEdGraphPin* FindOtherLink(TArray<UEdGraphPin*>& Links2, int32 OriginalIn
 	for (UEdGraphPin* Other : Links2)
 	{
 		UEdGraphNode* Node2 = Other->GetOwningNode();
-		if(FGraphDiffControl::IsNodeMatch(Node1, Node2))
+		const bool bPinsMatch = Other->GetName() == PinToFind->GetName() || Other->PinId == PinToFind->PinId;
+		if(FGraphDiffControl::IsNodeMatch(Node1, Node2) && bPinsMatch)
 		{
 			return Other;
 		}
