This changelist contains both a refactor of the material graph system, as well as an improved custom HLSL node.

The first series of changes relate to modifying the way that SGraphNode* inheritants are created. Previously, NodeFactory was responsible of creating the SGraphNode* inheritants based on some runtime type checking (which basically means a series of if statements checking if the passed in pointer is of a certain child class type). Now, the UMaterialGraph* inheritants (the controller classes in the MVC paradigm of the material graph) are responsible for creating their own UI components (SGraphNode* inheritants). This just means that we now polymorphically create the proper SGraphNode* for the given UMaterialGraph* inheritant.

This refactor allows us to more flexibly create a variety of nodes types instead of relying on SGraphNodeMaterialBase as a one size fits all solution, leading to the new HLSL custom node. Previously, this node didn't support showing the syntax-highlighted code inline in the node itself. Because we now use polymorphism to create the nodes, it was very easy to create new UMaterialGraph* and SGraphNode* inheritants to support this new change.

Other changes which relate to propagating changes (to affect previews) in MaterialGraphNode.cpp are needed because there previously wasn't the mechanism to do so (which may have been a bug). For instance, suppose you had a constant vector3 feeding into a custom HLSL node which just outputs the color coming from that constant vector3 node. It would be obvious that changing the vector3, requires a preview update of the HLSL node. However, that wouldn't happen - the preview on the HLSL node just stayed the same. These changes addresses that issue.

There is one last series of changes which relates to collapsing the HLSL code in the node. We want to make sure that this change doesn't mess up the layouts of artists' already made material graphs which use the custom HLSL node. In order to address this, we added the ability to collapse the code in order to hide it. Furthermore, we had to make sure that projects which previously used a custom HLSL node, have the code collapsed by default (because it wasn't there before). However, creating a new HLSL node once this change comes in, would yield an uncollapsed node. The collapsing of nodes uses the "Advanced Pins" chevron/collapser which doesn't save its state after saving and exiting the material. It was crucial, in order to preserve the layout of the material graph, that the saving happens. However, previously, the only way to save state, was to go through a recompile/regeneration of the previews. Therefore, it was necessary to add a special function in the material editor utilities which would just mark the material as dirty such that states like whether the code was collapsed save throughout sessions.

diff --git a/Engine/Source/Editor/EditorStyle/Private/StarshipStyle.cpp b/Engine/Source/Editor/EditorStyle/Private/StarshipStyle.cpp
--- a/Engine/Source/Editor/EditorStyle/Private/StarshipStyle.cpp
+++ b/Engine/Source/Editor/EditorStyle/Private/StarshipStyle.cpp
@@ -250,6 +250,7 @@ void FStarshipEditorStyle::FStyle::Initialize()
 	SetupMyBlueprintStyles();
 	SetupStatusBarStyle();
 	SetupColorPickerStyle();
+	SetupSourceCodeStyles();
 
 //	LogUnusedBrushResources();
 
@@ -7627,6 +7628,27 @@ void FStarshipEditorStyle::FStyle::SetupColorPickerStyle()
 	Set("ColorPicker.ColorThemes", new IMAGE_BRUSH_SVG("Starship/ColorPicker/ColorThemes", Icon16x16));
 }
 
+void FStarshipEditorStyle::FStyle::SetupSourceCodeStyles()
+{
+	constexpr int32 SourceCodeFontSize = 9;
+	FSlateFontInfo SourceCodeFont = DEFAULT_FONT("Mono", SourceCodeFontSize);
+	FTextBlockStyle NormalSourceCodeText = FTextBlockStyle(NormalText)
+		.SetFont(SourceCodeFont);
+	const FTextBlockStyle SourceCodeErrorText = FTextBlockStyle(NormalSourceCodeText)
+		.SetUnderlineBrush(IMAGE_BRUSH("Icons/White", Icon8x8, FLinearColor::Red, ESlateBrushTileType::Both))
+		.SetColorAndOpacity(FLinearColor::Red);
+	
+	Set("SyntaxHighlight.SourceCode.Normal", FTextBlockStyle(NormalSourceCodeText).SetColorAndOpacity(FLinearColor(FColor(189, 183, 107))));
+	Set("SyntaxHighlight.SourceCode.Operator", FTextBlockStyle(NormalSourceCodeText).SetColorAndOpacity(FLinearColor(FColor(220, 220, 220))));
+	Set("SyntaxHighlight.SourceCode.Keyword", FTextBlockStyle(NormalSourceCodeText).SetColorAndOpacity(FLinearColor(FColor(86, 156, 214))));
+	Set("SyntaxHighlight.SourceCode.String", FTextBlockStyle(NormalSourceCodeText).SetColorAndOpacity(FLinearColor(FColor(214, 157, 133))));
+	Set("SyntaxHighlight.SourceCode.Number", FTextBlockStyle(NormalSourceCodeText).SetColorAndOpacity(FLinearColor(FColor(181, 206, 168))));
+	Set("SyntaxHighlight.SourceCode.Comment", FTextBlockStyle(NormalSourceCodeText).SetColorAndOpacity(FLinearColor(FColor(87, 166, 74))));
+	Set("SyntaxHighlight.SourceCode.PreProcessorKeyword", FTextBlockStyle(NormalSourceCodeText).SetColorAndOpacity(FLinearColor(FColor(188, 98, 171))));
+
+	Set("SyntaxHighlight.SourceCode.Error", SourceCodeErrorText); 
+}
+
 #undef DEFAULT_FONT
 #undef ICON_FONT
 #undef LOCTEXT_NAMESPACE
diff --git a/Engine/Source/Editor/EditorStyle/Private/StarshipStyle.h b/Engine/Source/Editor/EditorStyle/Private/StarshipStyle.h
--- a/Engine/Source/Editor/EditorStyle/Private/StarshipStyle.h
+++ b/Engine/Source/Editor/EditorStyle/Private/StarshipStyle.h
@@ -76,6 +76,7 @@ public:
 		void SetupStatusBarStyle();
 		void SetupColorPickerStyle();
 		void SetupDerivedDataStyle();
+		void SetupSourceCodeStyles();
 
 		void SettingsChanged(FName PropertyName);
 		void SyncSettings();
diff --git a/Engine/Source/Editor/GraphEditor/Private/MaterialNodes/SGraphNodeMaterialBase.cpp b/Engine/Source/Editor/GraphEditor/Private/MaterialNodes/SGraphNodeMaterialBase.cpp
--- a/Engine/Source/Editor/GraphEditor/Private/MaterialNodes/SGraphNodeMaterialBase.cpp
+++ b/Engine/Source/Editor/GraphEditor/Private/MaterialNodes/SGraphNodeMaterialBase.cpp
@@ -599,6 +599,8 @@ TSharedRef<SWidget> SGraphNodeMaterialBase::CreatePreviewWidget()
 		return SNew(SBox)
 			.WidthOverride(ExpressionPreviewSize)
 			.HeightOverride(ExpressionPreviewSize)
+			.MaxAspectRatio(1.0f)
+			.MaxDesiredHeight(ExpressionPreviewSize)
 			.Visibility(ExpressionPreviewVisibility())
 			[
 				SNew(SBorder)
diff --git a/Engine/Source/Editor/GraphEditor/Private/MaterialNodes/SGraphNodeMaterialBase.h b/Engine/Source/Editor/GraphEditor/Private/MaterialNodes/SGraphNodeMaterialBase.h
deleted file mode 100644
diff --git a/Engine/Source/Editor/GraphEditor/Private/MaterialNodes/SGraphNodeMaterialComposite.h b/Engine/Source/Editor/GraphEditor/Private/MaterialNodes/SGraphNodeMaterialComposite.h
deleted file mode 100644
diff --git a/Engine/Source/Editor/GraphEditor/Private/MaterialNodes/SGraphNodeMaterialCustom.cpp b/Engine/Source/Editor/GraphEditor/Private/MaterialNodes/SGraphNodeMaterialCustom.cpp
new file mode 100644
--- /dev/null
+++ b/Engine/Source/Editor/GraphEditor/Private/MaterialNodes/SGraphNodeMaterialCustom.cpp
@@ -0,0 +1,159 @@
+#include "MaterialNodes/SGraphNodeMaterialCustom.h"
+
+#include "Templates/Casts.h"
+#include "Widgets/Input/SCheckBox.h"
+#include "Widgets/Images/SImage.h"
+#include "MaterialGraph/MaterialGraphNode_Custom.h"
+#include "GraphEditorSettings.h"
+
+void SGraphNodeMaterialCustom::Construct(const FArguments& InArgs, class UMaterialGraphNode* InNode)
+{
+	this->GraphNode = InNode;
+	this->MaterialNode = InNode;
+
+	// Create the syntax highlighter
+	FHLSLSyntaxHighlighterMarshaller::FSyntaxTextStyle CodeStyle(
+		FAppStyle::Get().GetWidgetStyle<FTextBlockStyle>("SyntaxHighlight.SourceCode.Normal"),
+		FAppStyle::Get().GetWidgetStyle<FTextBlockStyle>("SyntaxHighlight.SourceCode.Operator"),
+		FAppStyle::Get().GetWidgetStyle<FTextBlockStyle>("SyntaxHighlight.SourceCode.Keyword"),
+		FAppStyle::Get().GetWidgetStyle<FTextBlockStyle>("SyntaxHighlight.SourceCode.String"),
+		FAppStyle::Get().GetWidgetStyle<FTextBlockStyle>("SyntaxHighlight.SourceCode.Number"),
+		FAppStyle::Get().GetWidgetStyle<FTextBlockStyle>("SyntaxHighlight.SourceCode.Comment"),
+		FAppStyle::Get().GetWidgetStyle<FTextBlockStyle>("SyntaxHighlight.SourceCode.PreProcessorKeyword"),
+		FAppStyle::Get().GetWidgetStyle<FTextBlockStyle>("SyntaxHighlight.SourceCode.Error"));
+
+	SyntaxHighlighter = FHLSLSyntaxHighlighterMarshaller::Create(CodeStyle);
+
+	this->SetCursor(EMouseCursor::CardinalCross);
+
+	this->UpdateGraphNode();
+}
+
+void SGraphNodeMaterialCustom::CreateBelowPinControls(TSharedPtr<SVerticalBox> MainBox)
+{
+	UMaterialGraphNode_Custom* CustomNode = CastChecked<UMaterialGraphNode_Custom>(GraphNode);
+
+	TAttribute<FText> GetText;
+	GetText.BindUObject(CustomNode, &UMaterialGraphNode_Custom::GetHlslText);
+
+	FOnTextCommitted TextCommit;
+	TextCommit.BindUObject(CustomNode, &UMaterialGraphNode_Custom::OnCustomHlslTextCommitted);
+
+	SAssignNew(ShaderTextBox, SMultiLineEditableTextBox)
+		.AutoWrapText(false)
+		.Margin(FMargin(5, 5, 5, 5))
+		.Text(GetText)
+		.Visibility(this, &SGraphNodeMaterialCustom::CodeVisibility)
+		.OnTextCommitted(TextCommit)
+		.Marshaller(SyntaxHighlighter);
+
+	TSharedPtr<SVerticalBox> PreviewBox;
+	SAssignNew(PreviewBox, SVerticalBox);
+
+	SGraphNodeMaterialBase::CreateBelowPinControls(PreviewBox);
+
+	MainBox->AddSlot()
+	.Padding(Settings->GetNonPinNodeBodyPadding())
+	.AutoHeight()
+	[
+		SNew(SHorizontalBox)
+		+ SHorizontalBox::Slot()
+		.AutoWidth()
+		.Padding(FMargin(10.f, 5.f, 10.f, 10.f))
+		[
+			ShaderTextBox.ToSharedRef()
+		]
+		+ SHorizontalBox::Slot()
+		.AutoWidth()
+		[
+			PreviewBox.ToSharedRef()
+		]
+	];
+}
+
+EVisibility SGraphNodeMaterialCustom::CodeVisibility() const
+{
+	UMaterialGraphNode_Custom* CastedNode = CastChecked<UMaterialGraphNode_Custom>(GraphNode);
+	if (!CastedNode)
+	{
+		return EVisibility::Collapsed;
+	}
+
+	ECheckBoxState State = CastedNode->IsCodeViewChecked();
+
+	return (State == ECheckBoxState::Checked) ? EVisibility::Visible : EVisibility::Collapsed;
+}
+
+void SGraphNodeMaterialCustom::CreateAdvancedViewArrow(TSharedPtr<SVerticalBox> MainBox)
+{
+	UMaterialGraphNode_Custom* CustomNode = CastChecked<UMaterialGraphNode_Custom>(GraphNode);
+
+	if (!CustomNode)
+	{
+		return;
+	}
+
+	MainBox->AddSlot()
+	.AutoHeight()
+	.HAlign(HAlign_Fill)
+	.VAlign(VAlign_Top)
+	.Padding(3, 0, 3, 3)
+	[
+		SNew(SCheckBox)
+		.Visibility(this, &SGraphNodeMaterialCustom::AdvancedViewArrowVisibility)
+		.OnCheckStateChanged( this, &SGraphNodeMaterialCustom::OnAdvancedViewChanged )
+		.IsChecked( this, &SGraphNodeMaterialCustom::IsAdvancedViewChecked )
+		.Cursor(EMouseCursor::Default)
+		.Style(FAppStyle::Get(), "Graph.Node.AdvancedView")
+		[
+			SNew(SHorizontalBox)
+			+SHorizontalBox::Slot()
+			.VAlign(VAlign_Center)
+			.HAlign(HAlign_Center)
+			[
+				SNew(SImage)
+				. Image(this, &SGraphNodeMaterialCustom::GetAdvancedViewArrow)
+			]
+		]
+	];
+}
+
+EVisibility SGraphNodeMaterialCustom::AdvancedViewArrowVisibility() const
+{
+	// Always show the arrow no matter the configuration of the custom hlsl node
+	return GraphNode ? EVisibility::Visible : EVisibility::Collapsed;
+}
+
+void SGraphNodeMaterialCustom::OnAdvancedViewChanged(const ECheckBoxState NewCheckedState)
+{
+	if (UMaterialGraphNode_Custom* CastedNode = Cast<UMaterialGraphNode_Custom>(GraphNode))
+	{
+		CastedNode->OnCodeViewChanged(NewCheckedState);
+	}
+}
+
+ECheckBoxState SGraphNodeMaterialCustom::IsAdvancedViewChecked() const
+{
+	if (UMaterialGraphNode_Custom* CastedNode = Cast<UMaterialGraphNode_Custom>(GraphNode))
+	{
+		return CastedNode->IsCodeViewChecked();
+	}
+
+	return ECheckBoxState::Unchecked;
+}
+
+#define CHEVRON_UP TEXT("Icons.ChevronUp")
+#define CHEVRON_DOWN TEXT("Icons.ChevronDown")
+
+const FSlateBrush* SGraphNodeMaterialCustom::GetAdvancedViewArrow() const
+{
+	UMaterialGraphNode_Custom* CastedNode = Cast<UMaterialGraphNode_Custom>(GraphNode);
+	if (!CastedNode)
+	{
+		return FAppStyle::GetBrush(CHEVRON_UP);
+	}
+
+	ECheckBoxState State = CastedNode->IsCodeViewChecked();
+
+	return FAppStyle::GetBrush((State == ECheckBoxState::Checked) ? CHEVRON_UP : CHEVRON_DOWN);
+}
diff --git a/Engine/Source/Editor/GraphEditor/Private/MaterialNodes/SGraphNodeMaterialResult.h b/Engine/Source/Editor/GraphEditor/Private/MaterialNodes/SGraphNodeMaterialResult.h
deleted file mode 100644
diff --git a/Engine/Source/Editor/GraphEditor/Private/NodeFactory.cpp b/Engine/Source/Editor/GraphEditor/Private/NodeFactory.cpp
--- a/Engine/Source/Editor/GraphEditor/Private/NodeFactory.cpp
+++ b/Engine/Source/Editor/GraphEditor/Private/NodeFactory.cpp
@@ -123,29 +123,6 @@ TSharedPtr<SGraphNode> FNodeFactory::CreateNodeWidget(UEdGraphNode* InNode)
 		}
 	}
 
-	if (UMaterialGraphNode_Base* BaseMaterialNode = Cast<UMaterialGraphNode_Base>(InNode))
-	{
-		if (UMaterialGraphNode_Root* RootMaterialNode = Cast<UMaterialGraphNode_Root>(InNode))
-		{
-			return SNew(SGraphNodeMaterialResult, RootMaterialNode);
-		}
-		else if (UMaterialGraphNode_Knot* MaterialKnot = Cast<UMaterialGraphNode_Knot>(InNode))
-		{
-			return SNew(SGraphNodeKnot, MaterialKnot);
-		}
-		else if (UMaterialGraphNode* MaterialNode = Cast<UMaterialGraphNode>(InNode))
-		{
-			if (UMaterialGraphNode_Composite* MaterialComposite = Cast<UMaterialGraphNode_Composite>(InNode))
-			{
-				return SNew(SGraphNodeMaterialComposite, MaterialComposite);
-			}
-			else
-			{
-				return SNew(SGraphNodeMaterialBase, MaterialNode);
-			}
-		}
-	}
-
 	if (UK2Node* K2Node = Cast<UK2Node>(InNode))
 	{
 		if (UK2Node_Composite* CompositeNode = Cast<UK2Node_Composite>(InNode))
diff --git a/Engine/Source/Editor/GraphEditor/Public/MaterialNodes/SGraphNodeMaterialBase.h b/Engine/Source/Editor/GraphEditor/Public/MaterialNodes/SGraphNodeMaterialBase.h
new file mode 100644
--- /dev/null
+++ b/Engine/Source/Editor/GraphEditor/Public/MaterialNodes/SGraphNodeMaterialBase.h
@@ -0,0 +1,145 @@
+// Copyright Epic Games, Inc. All Rights Reserved.
+
+#pragma once
+
+#include "Internationalization/Text.h"
+#include "Layout/Visibility.h"
+#include "Math/UnrealMathSSE.h"
+#include "Math/Vector2D.h"
+#include "Rendering/RenderingCommon.h"
+#include "SGraphNode.h"
+#include "SNodePanel.h"
+#include "Styling/SlateTypes.h"
+#include "Templates/SharedPointer.h"
+#include "Textures/SlateShaderResource.h"
+#include "Widgets/DeclarativeSyntaxSupport.h"
+
+class FMaterialRenderProxy;
+class FRHICommandListImmediate;
+class FSlateRect;
+class FWidgetStyle;
+class SGraphPin;
+class SOverlay;
+class SVerticalBox;
+class SWidget;
+class UMaterialGraphNode;
+struct FGeometry;
+struct FSlateBrush;
+
+typedef TSharedPtr<class FPreviewElement, ESPMode::ThreadSafe> FThreadSafePreviewPtr;
+
+class FPreviewViewport : public ISlateViewport
+{
+public:
+	FPreviewViewport(class UMaterialGraphNode* InNode);
+	~FPreviewViewport();
+
+	// ISlateViewport interface
+	virtual void OnDrawViewport( const FGeometry& AllottedGeometry, const FSlateRect& MyCullingRect, class FSlateWindowElementList& OutDrawElements, int32 LayerId, const FWidgetStyle& InWidgetStyle, bool bParentEnabled ) override;
+	virtual FIntPoint GetSize() const override;
+	virtual class FSlateShaderResource* GetViewportRenderTargetTexture() const override {return NULL;}
+	virtual bool RequiresVsync() const override {return false;}
+
+	/** Material node to get expression preview from */
+	UMaterialGraphNode* MaterialNode;
+	/** Custom Slate Element to display the preview */
+	FThreadSafePreviewPtr PreviewElement;
+private:
+	/**
+	 * Updates the expression preview render proxy from a graph node
+	 */
+	void UpdatePreviewNodeRenderProxy();
+};
+
+class FPreviewElement : public ICustomSlateElement
+{
+public:
+	FPreviewElement();
+	~FPreviewElement();
+
+	/**
+	 * Sets up the canvas for rendering
+	 *
+	 * @param	InCanvasRect			Size of the canvas tile
+	 * @param	InClippingRect			How to clip the canvas tile
+	 * @param	InGraphNode				The graph node for the material preview
+	 * @param	bInIsRealtime			Whether preview is using realtime values
+	 *
+	 * @return	Whether there is anything to render
+	 */
+	bool BeginRenderingCanvas( const FIntRect& InCanvasRect, const FIntRect& InClippingRect, UMaterialGraphNode* InGraphNode, bool bInIsRealtime );
+
+	/**
+	 * Updates the expression preview render proxy from a graph node on the render thread
+	 */
+	void UpdateExpressionPreview(UMaterialGraphNode* PreviewNode);
+private:
+	/**
+	 * ICustomSlateElement interface 
+	 */
+	virtual void DrawRenderThread(FRHICommandListImmediate& RHICmdList, const void* InWindowBackBuffer) override;
+
+private:
+	/** Render target that the canvas renders to */
+	class FSlateMaterialPreviewRenderTarget* RenderTarget;
+	/** Render proxy for the expression preview */
+	FMaterialRenderProxy* ExpressionPreview;
+	/** Whether preview is using realtime values */
+	bool bIsRealtime;
+};
+
+class GRAPHEDITOR_API SGraphNodeMaterialBase : public SGraphNode
+{
+public:
+	SLATE_BEGIN_ARGS(SGraphNodeMaterialBase){}
+	SLATE_END_ARGS()
+
+	void Construct(const FArguments& InArgs, class UMaterialGraphNode* InNode);
+
+	// SGraphNode interface
+	virtual void CreatePinWidgets() override;
+	// End of SGraphNode interface
+
+	// SNodePanel::SNode interface
+	virtual void MoveTo(const FVector2D& NewPosition, FNodeSet& NodeFilter, bool bMarkDirty = true) override;
+	// End of SNodePanel::SNode interface
+
+	UMaterialGraphNode* GetMaterialGraphNode() const {return MaterialNode;}
+
+	/* Populate a meta data tag with information about this graph node */
+	virtual void PopulateMetaTag(class FGraphNodeMetaData* TagMeta) const override;
+
+protected:
+	// SGraphNode interface
+	virtual void AddPin( const TSharedRef<SGraphPin>& PinToAdd ) override;
+	virtual void CreateBelowPinControls(TSharedPtr<SVerticalBox> MainBox) override;
+	virtual void SetDefaultTitleAreaWidget(TSharedRef<SOverlay> DefaultTitleAreaWidget) override;
+	virtual TSharedRef<SWidget> CreateNodeContentArea() override;
+	virtual void OnAdvancedViewChanged(const ECheckBoxState NewCheckedState) override;
+	// End of SGraphNode interface
+
+	/** Creates a preview viewport if necessary */
+	TSharedRef<SWidget> CreatePreviewWidget();
+
+	/** Returns visibility of Expression Preview viewport */
+	EVisibility ExpressionPreviewVisibility() const;
+
+	/** Returns text to over lay over the expression preview viewport */
+	FText ExpressionPreviewOverlayText() const;
+
+	/** Show/hide Expression Preview */
+	void OnExpressionPreviewChanged( const ECheckBoxState NewCheckedState );
+
+	/** hidden == unchecked, shown == checked */
+	ECheckBoxState IsExpressionPreviewChecked() const;
+
+	/** Up when shown, down when hidden */
+	const FSlateBrush* GetExpressionPreviewArrow() const;
+
+protected:
+	/** Slate viewport for rendering preview via custom slate element */
+	TSharedPtr<FPreviewViewport> PreviewViewport;
+
+	/** Cached material graph node pointer to avoid casting */
+	UMaterialGraphNode* MaterialNode;
+};
diff --git a/Engine/Source/Editor/GraphEditor/Public/MaterialNodes/SGraphNodeMaterialComposite.h b/Engine/Source/Editor/GraphEditor/Public/MaterialNodes/SGraphNodeMaterialComposite.h
new file mode 100644
--- /dev/null
+++ b/Engine/Source/Editor/GraphEditor/Public/MaterialNodes/SGraphNodeMaterialComposite.h
@@ -0,0 +1,45 @@
+// Copyright Epic Games, Inc. All Rights Reserved.
+
+#pragma once
+
+#include "Internationalization/Text.h"
+#include "MaterialNodes/SGraphNodeMaterialBase.h"
+#include "Math/Vector2D.h"
+#include "SNodePanel.h"
+#include "Templates/SharedPointer.h"
+#include "Widgets/DeclarativeSyntaxSupport.h"
+
+class SToolTip;
+class SWidget;
+class UEdGraph;
+class UMaterialGraphNode_Composite;
+
+//@TODO: This class is mostly C&P from UK2Node_Composite, consolidate composites to not be BP specific.
+class GRAPHEDITOR_API SGraphNodeMaterialComposite : public SGraphNodeMaterialBase
+{
+public:
+	SLATE_BEGIN_ARGS(SGraphNodeMaterialComposite){}
+	SLATE_END_ARGS()
+
+	void Construct(const FArguments& InArgs, UMaterialGraphNode_Composite* InNode);
+
+	// SGraphNode interface
+	virtual void UpdateGraphNode() override;
+	virtual TSharedPtr<SToolTip> GetComplexTooltip() override;
+	// End of SGraphNode interface
+
+	// SNodePanel::SNode interface
+	virtual void MoveTo(const FVector2D& NewPosition, FNodeSet& NodeFilter, bool bMarkDirty = true) override;
+	// End of SNodePanel::SNode interface
+protected:
+	virtual UEdGraph* GetInnerGraph() const;
+
+private:
+	FText GetPreviewCornerText() const;
+	FText GetTooltipTextForNode() const;
+
+	TSharedRef<SWidget> CreateNodeBody();
+
+	/** Cached material graph node pointer to avoid casting */
+	UMaterialGraphNode_Composite* CompositeNode;
+};
diff --git a/Engine/Source/Editor/GraphEditor/Public/MaterialNodes/SGraphNodeMaterialCustom.h b/Engine/Source/Editor/GraphEditor/Public/MaterialNodes/SGraphNodeMaterialCustom.h
new file mode 100644
--- /dev/null
+++ b/Engine/Source/Editor/GraphEditor/Public/MaterialNodes/SGraphNodeMaterialCustom.h
@@ -0,0 +1,39 @@
+// Copyright Epic Games, Inc. All Rights Reserved.
+
+#pragma once
+
+#include "MaterialNodes/SGraphNodeMaterialBase.h"
+#include "Text/HLSLSyntaxHighlighterMarshaller.h"
+#include "Widgets/Input/SMultiLineEditableTextBox.h"
+
+class GRAPHEDITOR_API SGraphNodeMaterialCustom : public SGraphNodeMaterialBase
+{
+public:
+	SLATE_BEGIN_ARGS(SGraphNodeMaterialCustom){}
+	SLATE_END_ARGS()
+
+	void Construct(const FArguments& InArgs, class UMaterialGraphNode* InNode);
+
+	UMaterialGraphNode* GetMaterialGraphNode() const {return MaterialNode;}
+
+	FText GetHlslText() const;
+	void OnCustomHlslTextCommited(const FText& InText, ETextCommit::Type InType);
+
+protected:
+	// SGraphNode interface
+	virtual void CreateBelowPinControls(TSharedPtr<SVerticalBox> MainBox) override;
+	virtual void CreateAdvancedViewArrow(TSharedPtr<SVerticalBox> MainBox) override;
+	virtual EVisibility AdvancedViewArrowVisibility() const override;
+	virtual void OnAdvancedViewChanged( const ECheckBoxState NewCheckedState ) override;
+	virtual ECheckBoxState IsAdvancedViewChecked() const override;
+	virtual const FSlateBrush* GetAdvancedViewArrow() const override;
+	// End of SGraphNode interface
+
+private:
+	EVisibility CodeVisibility() const;
+
+protected:
+	/** Shader text box */
+	TSharedPtr<FHLSLSyntaxHighlighterMarshaller> SyntaxHighlighter;
+	TSharedPtr<SMultiLineEditableTextBox> ShaderTextBox;
+};
diff --git a/Engine/Source/Editor/GraphEditor/Public/MaterialNodes/SGraphNodeMaterialResult.h b/Engine/Source/Editor/GraphEditor/Public/MaterialNodes/SGraphNodeMaterialResult.h
new file mode 100644
--- /dev/null
+++ b/Engine/Source/Editor/GraphEditor/Public/MaterialNodes/SGraphNodeMaterialResult.h
@@ -0,0 +1,32 @@
+// Copyright Epic Games, Inc. All Rights Reserved.
+
+#pragma once
+
+#include "Math/Vector2D.h"
+#include "SGraphNode.h"
+#include "SNodePanel.h"
+#include "Widgets/DeclarativeSyntaxSupport.h"
+
+class FGraphNodeMetaData;
+class UMaterialGraphNode_Root;
+
+class GRAPHEDITOR_API SGraphNodeMaterialResult : public SGraphNode
+{
+public:
+	SLATE_BEGIN_ARGS(SGraphNodeMaterialResult){}
+	SLATE_END_ARGS()
+
+	void Construct(const FArguments& InArgs, class UMaterialGraphNode_Root* InNode);
+
+	// SGraphNode interface
+	virtual void CreatePinWidgets() override;
+	virtual void PopulateMetaTag(FGraphNodeMetaData* TagMeta) const override;
+	// End of SGraphNode interface
+
+	// SNodePanel::SNode interface
+	virtual void MoveTo(const FVector2D& NewPosition, FNodeSet& NodeFilter, bool bMarkDirty = true) override;
+	// End of SNodePanel::SNode interface
+
+private:
+	UMaterialGraphNode_Root* RootNode;
+};
diff --git a/Engine/Source/Editor/GraphEditor/Public/SGraphNode.h b/Engine/Source/Editor/GraphEditor/Public/SGraphNode.h
--- a/Engine/Source/Editor/GraphEditor/Public/SGraphNode.h
+++ b/Engine/Source/Editor/GraphEditor/Public/SGraphNode.h
@@ -364,19 +364,19 @@ protected:
 	///// ADVANCED VIEW FUNCTIONS /////
 
 	/** Create button to show/hide advanced pins */
-	void CreateAdvancedViewArrow(TSharedPtr<SVerticalBox> MainBox);
+	virtual void CreateAdvancedViewArrow(TSharedPtr<SVerticalBox> MainBox);
 
 	/** Returns visibility of AdvancedViewButton */
-	EVisibility AdvancedViewArrowVisibility() const;
+	virtual EVisibility AdvancedViewArrowVisibility() const;
 
 	/** Show/hide advanced view */
 	virtual void OnAdvancedViewChanged( const ECheckBoxState NewCheckedState );
 
 	/** hidden == unchecked, shown == checked */
-	ECheckBoxState IsAdvancedViewChecked() const;
+	virtual ECheckBoxState IsAdvancedViewChecked() const;
 
 	/** Up when shown, down when hidden */
-	const FSlateBrush* GetAdvancedViewArrow() const;
+	virtual const FSlateBrush* GetAdvancedViewArrow() const;
 
 	/** Checks if the node is the only node selected */
 	bool IsSelectedExclusively() const;
diff --git a/Engine/Source/Editor/MaterialEditor/Private/MaterialEditor.cpp b/Engine/Source/Editor/MaterialEditor/Private/MaterialEditor.cpp
--- a/Engine/Source/Editor/MaterialEditor/Private/MaterialEditor.cpp
+++ b/Engine/Source/Editor/MaterialEditor/Private/MaterialEditor.cpp
@@ -6178,6 +6178,11 @@ void FMaterialEditor::UpdateMaterialAfterGraphChange()
 	UpdateStrataTopologyPreview();
 }
 
+void FMaterialEditor::MarkMaterialDirty()
+{
+	SetMaterialDirty();
+}
+
 void FMaterialEditor::JumpToHyperlink(const UObject* ObjectReference)
 {
 	if (const UEdGraphNode* Node = Cast<const UEdGraphNode>(ObjectReference))
diff --git a/Engine/Source/Editor/MaterialEditor/Private/MaterialEditor.h b/Engine/Source/Editor/MaterialEditor/Private/MaterialEditor.h
--- a/Engine/Source/Editor/MaterialEditor/Private/MaterialEditor.h
+++ b/Engine/Source/Editor/MaterialEditor/Private/MaterialEditor.h
@@ -405,6 +405,7 @@ public:
 	virtual void DeleteSelectedNodes() override;
 	virtual FText GetOriginalObjectName() const override;
 	virtual void UpdateMaterialAfterGraphChange() override;
+	virtual void MarkMaterialDirty() override;
 	virtual void JumpToHyperlink(const UObject* ObjectReference) override; 
 	virtual bool CanPasteNodes() const override;
 	virtual void PasteNodesHere(const FVector2D& Location, const class UEdGraph* Graph = nullptr) override;
diff --git a/Engine/Source/Editor/MaterialEditor/Private/MaterialEditorUtilities.cpp b/Engine/Source/Editor/MaterialEditor/Private/MaterialEditorUtilities.cpp
--- a/Engine/Source/Editor/MaterialEditor/Private/MaterialEditorUtilities.cpp
+++ b/Engine/Source/Editor/MaterialEditor/Private/MaterialEditorUtilities.cpp
@@ -132,6 +132,15 @@ void FMaterialEditorUtilities::UpdateMaterialAfterGraphChange(const class UEdGra
 	}
 }
 
+void FMaterialEditorUtilities::MarkMaterialDirty(const class UEdGraph* Graph)
+{
+	TSharedPtr<class IMaterialEditor> MaterialEditor = GetIMaterialEditorForObject(Graph);
+	if (MaterialEditor.IsValid())
+	{
+		MaterialEditor->MarkMaterialDirty();
+	}
+}
+
 void FMaterialEditorUtilities::UpdateDetailView(const class UEdGraph* Graph)
 {
 	TSharedPtr<class IMaterialEditor> MaterialEditor = GetIMaterialEditorForObject(Graph);
diff --git a/Engine/Source/Editor/MaterialEditor/Public/IMaterialEditor.h b/Engine/Source/Editor/MaterialEditor/Public/IMaterialEditor.h
--- a/Engine/Source/Editor/MaterialEditor/Public/IMaterialEditor.h
+++ b/Engine/Source/Editor/MaterialEditor/Public/IMaterialEditor.h
@@ -90,6 +90,9 @@ public:
 	 */
 	virtual void UpdateMaterialAfterGraphChange() {};
 
+	/** Marks the material as dirty. */
+	virtual void MarkMaterialDirty() {};
+
 	virtual void UpdateDetailView() {}
 
 	/**
diff --git a/Engine/Source/Editor/MaterialEditor/Public/MaterialEditorUtilities.h b/Engine/Source/Editor/MaterialEditor/Public/MaterialEditorUtilities.h
--- a/Engine/Source/Editor/MaterialEditor/Public/MaterialEditorUtilities.h
+++ b/Engine/Source/Editor/MaterialEditor/Public/MaterialEditorUtilities.h
@@ -133,6 +133,9 @@ public:
 	 */
 	static void UpdateMaterialAfterGraphChange(const class UEdGraph* Graph);
 
+	/** Mark the material as dirty. */
+	static void MarkMaterialDirty(const class UEdGraph* Graph);
+
 	static void UpdateDetailView(const class UEdGraph* Graph);
 
 	/** Can we paste to this graph? */
diff --git a/Engine/Source/Editor/UnrealEd/Classes/MaterialGraph/MaterialGraphNode.h b/Engine/Source/Editor/UnrealEd/Classes/MaterialGraph/MaterialGraphNode.h
--- a/Engine/Source/Editor/UnrealEd/Classes/MaterialGraph/MaterialGraphNode.h
+++ b/Engine/Source/Editor/UnrealEd/Classes/MaterialGraph/MaterialGraphNode.h
@@ -73,6 +73,7 @@ public:
 	virtual bool CanUserDeleteNode() const override;
 	virtual uint32 GetPinMaterialType(const UEdGraphPin* Pin) const override;
 	virtual void PinDefaultValueChanged(UEdGraphPin* Pin) override;
+	virtual TSharedPtr<SGraphNode> CreateVisualWidget() override;
 	//~ End UEdGraphNode Interface.
 
 	//~ Begin UMaterialGraphNode_Base Interface
@@ -84,6 +85,8 @@ public:
 	/** Will return the shorten pin name to use based on long pin name */
 	static FName UNREALED_API GetShortenPinName(const FName PinName);
 
+	virtual void PropagatePropertyChange() override;
+
 private:
 	/** Make sure the MaterialExpression is owned by the Material */
 	void ResetMaterialExpressionOwner();
diff --git a/Engine/Source/Editor/UnrealEd/Classes/MaterialGraph/MaterialGraphNode_Base.h b/Engine/Source/Editor/UnrealEd/Classes/MaterialGraph/MaterialGraphNode_Base.h
--- a/Engine/Source/Editor/UnrealEd/Classes/MaterialGraph/MaterialGraphNode_Base.h
+++ b/Engine/Source/Editor/UnrealEd/Classes/MaterialGraph/MaterialGraphNode_Base.h
@@ -70,4 +70,7 @@ protected:
 
 	/** Return the first input pin matching the name */
 	class UEdGraphPin* GetInputPin(const FName& PinName) const;
+
+	/** Used when changing one property may affect the preview of subsequent properties */
+	virtual void PropagatePropertyChange() {}
 };
diff --git a/Engine/Source/Editor/UnrealEd/Classes/MaterialGraph/MaterialGraphNode_Composite.h b/Engine/Source/Editor/UnrealEd/Classes/MaterialGraph/MaterialGraphNode_Composite.h
--- a/Engine/Source/Editor/UnrealEd/Classes/MaterialGraph/MaterialGraphNode_Composite.h
+++ b/Engine/Source/Editor/UnrealEd/Classes/MaterialGraph/MaterialGraphNode_Composite.h
@@ -35,6 +35,7 @@ class UMaterialGraphNode_Composite : public UMaterialGraphNode
 	virtual void PostPasteNode() override; 
 	virtual void OnRenameNode(const FString& NewName) override; 
 	virtual TArray<UEdGraph*> GetSubGraphs() const override { return TArray<UEdGraph*>( { BoundGraph } ); }
+	virtual TSharedPtr<SGraphNode> CreateVisualWidget() override;
 	//~ End UEdGraphNode Interface.
 
 private:
diff --git a/Engine/Source/Editor/UnrealEd/Classes/MaterialGraph/MaterialGraphNode_Custom.h b/Engine/Source/Editor/UnrealEd/Classes/MaterialGraph/MaterialGraphNode_Custom.h
new file mode 100644
--- /dev/null
+++ b/Engine/Source/Editor/UnrealEd/Classes/MaterialGraph/MaterialGraphNode_Custom.h
@@ -0,0 +1,35 @@
+// Copyright Epic Games, Inc. All Rights Reserved.
+
+#pragma once
+
+#include "CoreMinimal.h"
+#include "Types/SlateEnums.h"
+#include "Styling/SlateTypes.h"
+#include "UObject/ObjectMacros.h"
+#include "Widgets/SBoxPanel.h"
+#include "MaterialGraph/MaterialGraph.h"
+#include "MaterialGraph/MaterialGraphNode_Base.h"
+#include "MaterialGraph/MaterialGraphNode.h"
+
+#include "MaterialGraphNode_Custom.generated.h"
+
+class UEdGraphPin;
+
+UCLASS(MinimalAPI)
+class UMaterialGraphNode_Custom : public UMaterialGraphNode
+{
+	GENERATED_UCLASS_BODY()
+
+public:
+	//~ Begin UEdGraphNode Interface.
+	virtual TSharedPtr<SGraphNode> CreateVisualWidget() override;
+	//~ End UEdGraphNode Interface.
+
+	UNREALED_API FText GetHlslText() const;
+	UNREALED_API void OnCustomHlslTextCommitted(const FText& InText, ETextCommit::Type InType);
+	UNREALED_API void OnCodeViewChanged(const ECheckBoxState NewCheckedState);
+	UNREALED_API ECheckBoxState IsCodeViewChecked() const;
+
+private:
+	void ChangeProperty(FString PropertyName, bool bUpdatePreview = true);
+};
diff --git a/Engine/Source/Editor/UnrealEd/Classes/MaterialGraph/MaterialGraphNode_Root.h b/Engine/Source/Editor/UnrealEd/Classes/MaterialGraph/MaterialGraphNode_Root.h
--- a/Engine/Source/Editor/UnrealEd/Classes/MaterialGraph/MaterialGraphNode_Root.h
+++ b/Engine/Source/Editor/UnrealEd/Classes/MaterialGraph/MaterialGraphNode_Root.h
@@ -25,6 +25,7 @@ class UMaterialGraphNode_Root : public UMaterialGraphNode_Base
 	virtual bool CanUserDeleteNode() const override { return false; }
 	virtual bool CanDuplicateNode() const override { return false; }
 	virtual void PostPlacedNewNode() override;
+	virtual TSharedPtr<SGraphNode> CreateVisualWidget() override;
 	//~ End UEdGraphNode Interface.
 
 	//~ Begin UMaterialGraphNode_Base Interface
diff --git a/Engine/Source/Editor/UnrealEd/Classes/MaterialGraph/MaterialGraphSchema.h b/Engine/Source/Editor/UnrealEd/Classes/MaterialGraph/MaterialGraphSchema.h
--- a/Engine/Source/Editor/UnrealEd/Classes/MaterialGraph/MaterialGraphSchema.h
+++ b/Engine/Source/Editor/UnrealEd/Classes/MaterialGraph/MaterialGraphSchema.h
@@ -252,6 +252,9 @@ class UMaterialGraphSchema : public UEdGraphSchema
 	/** update material when the default value of a material node's pin has changed */
 	UNREALED_API void UpdateMaterialOnDefaultValueChanged(const UEdGraph* Graph) const;
 
+	/** Mark the material as dirty (because of a change that shouldn't trigger recompile or preview update) */ 
+	UNREALED_API void MarkMaterialDirty(const UEdGraph* Graph) const;
+
 	/** Update the detail view */
 	UNREALED_API void UpdateDetailView(const UEdGraph* Graph) const;
 
diff --git a/Engine/Source/Editor/UnrealEd/Private/MaterialGraph.cpp b/Engine/Source/Editor/UnrealEd/Private/MaterialGraph.cpp
--- a/Engine/Source/Editor/UnrealEd/Private/MaterialGraph.cpp
+++ b/Engine/Source/Editor/UnrealEd/Private/MaterialGraph.cpp
@@ -5,6 +5,7 @@
 
 #include "MaterialGraph/MaterialGraph.h"
 #include "MaterialGraph/MaterialGraphNode_Comment.h"
+#include "MaterialGraph/MaterialGraphNode_Custom.h"
 #include "MaterialGraph/MaterialGraphNode_Composite.h"
 #include "MaterialGraph/MaterialGraphNode_PinBase.h"
 #include "MaterialGraph/MaterialGraphNode.h"
@@ -14,6 +15,7 @@
 #include "Materials/MaterialAttributeDefinitionMap.h"
 #include "Materials/MaterialExpressionComment.h"
 #include "Materials/MaterialExpressionComposite.h"
+#include "Materials/MaterialExpressionCustom.h"
 #include "Materials/MaterialExpressionPinBase.h"
 #include "Materials/MaterialExpressionFunctionOutput.h"
 #include "Materials/MaterialExpressionCustomOutput.h"
@@ -368,6 +370,17 @@ UMaterialGraphNode* UMaterialGraph::AddExpression(UMaterialExpression* Expressio
 		{
 			Node = InitExpressionNewNode<UMaterialGraphNode_PinBase>(this, Expression, false);
 		}
+		else if (Expression->IsA(UMaterialExpressionCustom::StaticClass()))
+		{
+			// This is for backward compatibility. We don't want people's custom HLSL nodes to suddenly explode
+			// when they open their materials for the first time with inline HLSL custom nodes working in the engine.
+			// This ensures that only new nodes or nodes that are subsequently saved as uncollapsed show the code.
+			UMaterialExpressionCustom* CustomNode = Cast<UMaterialExpressionCustom>(Expression);
+			if (bUserInvoked)
+				CustomNode->ShowCode = 1;
+
+			Node = InitExpressionNewNode<UMaterialGraphNode_Custom>(this, Expression, bUserInvoked);
+		}
 		else 
 		{
 			Node = InitExpressionNewNode<UMaterialGraphNode>(this, Expression, bUserInvoked);
diff --git a/Engine/Source/Editor/UnrealEd/Private/MaterialGraphNode.cpp b/Engine/Source/Editor/UnrealEd/Private/MaterialGraphNode.cpp
--- a/Engine/Source/Editor/UnrealEd/Private/MaterialGraphNode.cpp
+++ b/Engine/Source/Editor/UnrealEd/Private/MaterialGraphNode.cpp
@@ -52,6 +52,7 @@
 #include "GraphEditorSettings.h"
 #include "Framework/Commands/GenericCommands.h"
 #include "ScopedTransaction.h"
+#include "MaterialNodes/SGraphNodeMaterialBase.h"
 
 #define LOCTEXT_NAMESPACE "MaterialGraphNode"
 
@@ -660,6 +661,11 @@ void UMaterialGraphNode::PinDefaultValueChanged(UEdGraphPin* Pin)
 	MaterialExpression->PinDefaultValueChanged(Pin->SourceIndex, Pin->DefaultValue);
 }
 
+TSharedPtr<SGraphNode> UMaterialGraphNode::CreateVisualWidget()
+{
+	return SNew(SGraphNodeMaterialBase, this);
+}
+
 void UMaterialGraphNode::CreateInputPins()
 {
 	if (MaterialExpression->HasExecInput())
@@ -820,6 +826,37 @@ void UMaterialGraphNode::CreateOutputPins()
 	}
 }
 
+void UMaterialGraphNode::PropagatePropertyChange()
+{
+	MaterialExpression->bNeedToUpdatePreview = 1;
+
+	// Loop through all the output pins.
+	for (UEdGraphPin* Pin : Pins)
+	{
+		if (Pin->PinType.PinCategory == UMaterialGraphSchema::PC_Exec ||
+			Pin->Direction != EGPD_Output)
+		{
+			continue;
+		}
+
+		// We only care about pins which are outputs. We call PropagatePropertyChange
+		// on the nodes that are linked to this particular output pin.
+		for (UEdGraphPin* LinkedPins : Pin->LinkedTo)
+		{
+			// Retrieve the linked node. Check that it is of type UMaterialGraphNode.
+			UEdGraphNode* Node = LinkedPins->GetOwningNode();
+			UMaterialGraphNode* CastedNode = Cast<UMaterialGraphNode>(Node);
+			
+			if (!CastedNode)
+			{
+				continue;
+			}
+
+			CastedNode->PropagatePropertyChange();
+		}
+	}
+}
+
 void UMaterialGraphNode::ResetMaterialExpressionOwner()
 {
 	if (MaterialExpression)
diff --git a/Engine/Source/Editor/UnrealEd/Private/MaterialGraphNode_Composite.cpp b/Engine/Source/Editor/UnrealEd/Private/MaterialGraphNode_Composite.cpp
--- a/Engine/Source/Editor/UnrealEd/Private/MaterialGraphNode_Composite.cpp
+++ b/Engine/Source/Editor/UnrealEd/Private/MaterialGraphNode_Composite.cpp
@@ -16,6 +16,7 @@
 #include "EdGraphUtilities.h"
 #include "Kismet2/BlueprintEditorUtils.h"
 #include "Framework/Commands/GenericCommands.h"
+#include "MaterialNodes/SGraphNodeMaterialComposite.h"
 
 #define LOCTEXT_NAMESPACE "MaterialGraphNode_Composite"
 
@@ -138,6 +139,11 @@ void UMaterialGraphNode_Composite::OnRenameNode(const FString& NewName)
 	MaterialDirtyDelegate.ExecuteIfBound();
 }
 
+TSharedPtr<SGraphNode> UMaterialGraphNode_Composite::CreateVisualWidget()
+{
+	return SNew(SGraphNodeMaterialComposite, this);
+}
+
 void UMaterialGraphNode_Composite::FixupInputAndOutputPinBases()
 {
 	if (BoundGraph)
diff --git a/Engine/Source/Editor/UnrealEd/Private/MaterialGraphNode_Custom.cpp b/Engine/Source/Editor/UnrealEd/Private/MaterialGraphNode_Custom.cpp
new file mode 100644
--- /dev/null
+++ b/Engine/Source/Editor/UnrealEd/Private/MaterialGraphNode_Custom.cpp
@@ -0,0 +1,75 @@
+// Copyright Epic Games, Inc. All Rights Reserved.
+
+#include "MaterialGraph/MaterialGraphNode_Custom.h"
+#include "MaterialNodes/SGraphNodeMaterialCustom.h"
+
+#include "Materials/MaterialExpressionCustom.h"
+
+#define CODE_PROPERTY_NAME TEXT("Code")
+#define SHOW_CODE_PROPERTY_NAME TEXT("ShowCode")
+
+UMaterialGraphNode_Custom::UMaterialGraphNode_Custom(const FObjectInitializer& ObjectInitializer)
+	: Super(ObjectInitializer)
+{
+}
+
+TSharedPtr<SGraphNode> UMaterialGraphNode_Custom::CreateVisualWidget()
+{
+	return SNew(SGraphNodeMaterialCustom, this);
+}
+
+FText UMaterialGraphNode_Custom::GetHlslText() const
+{
+	UMaterialExpressionCustom* CustomExpression = CastChecked<UMaterialExpressionCustom>(MaterialExpression.Get());
+	return FText::FromString(CustomExpression->Code);
+}
+
+void UMaterialGraphNode_Custom::OnCustomHlslTextCommitted(const FText& InText, ETextCommit::Type InType)
+{
+	UMaterialExpressionCustom* CustomExpression = CastChecked<UMaterialExpressionCustom>(MaterialExpression.Get());
+
+	FString NewValue = InText.ToString();
+	if (!NewValue.Equals(CustomExpression->Code, ESearchCase::CaseSensitive))
+	{
+		CustomExpression->Code = NewValue;
+
+		ChangeProperty(CODE_PROPERTY_NAME);
+	}
+}
+
+void UMaterialGraphNode_Custom::OnCodeViewChanged(const ECheckBoxState NewCheckedState)
+{
+	UMaterialExpressionCustom* Expression = Cast<UMaterialExpressionCustom>(MaterialExpression.Get());
+	if (!Expression)
+		return;
+
+	Expression->ShowCode = (NewCheckedState == ECheckBoxState::Checked);
+
+	ChangeProperty(SHOW_CODE_PROPERTY_NAME, false);
+}
+
+ECheckBoxState UMaterialGraphNode_Custom::IsCodeViewChecked() const
+{
+	UMaterialExpressionCustom* Expression = Cast<UMaterialExpressionCustom>(MaterialExpression.Get());
+	if (!Expression)
+		return ECheckBoxState::Checked;
+
+	const bool bIsShown = (bool)Expression->ShowCode;
+	return bIsShown ? ECheckBoxState::Checked : ECheckBoxState::Unchecked;
+}
+
+void UMaterialGraphNode_Custom::ChangeProperty(FString PropertyName, bool bUpdatePreview)
+{
+	UMaterialExpressionCustom* Expression = Cast<UMaterialExpressionCustom>(MaterialExpression.Get());
+	for (TFieldIterator<FProperty> InputIt(Expression->GetClass(), EFieldIteratorFlags::IncludeSuper, EFieldIteratorFlags::ExcludeDeprecated); InputIt; ++InputIt)
+	{
+		FProperty* Property = *InputIt;
+		FString Name = Property->GetName();
+
+		if (Name == PropertyName)
+		{
+			Expression->ForcePropertyValueChanged(Property, bUpdatePreview);
+			break;
+		}
+	}
+}
diff --git a/Engine/Source/Editor/UnrealEd/Private/MaterialGraphNode_Knot.cpp b/Engine/Source/Editor/UnrealEd/Private/MaterialGraphNode_Knot.cpp
--- a/Engine/Source/Editor/UnrealEd/Private/MaterialGraphNode_Knot.cpp
+++ b/Engine/Source/Editor/UnrealEd/Private/MaterialGraphNode_Knot.cpp
@@ -7,6 +7,8 @@
 #include "MaterialEditorUtilities.h"
 #include "EdGraph/EdGraphPin.h"
 
+#include "SGraphNodeKnot.h"
+
 #define LOCTEXT_NAMESPACE "MaterialGraphNode_Knot"
 
 static const char* PC_Wildcard = "wildcard";
@@ -94,6 +96,11 @@ UEdGraphPin* UMaterialGraphNode_Knot::GetPassThroughPin(const UEdGraphPin* FromP
 	return nullptr;
 }
 
+TSharedPtr<SGraphNode> UMaterialGraphNode_Knot::CreateVisualWidget()
+{
+	return SNew(SGraphNodeKnot, this);
+}
+
 
 /////////////////////////////////////////////////////
 
diff --git a/Engine/Source/Editor/UnrealEd/Private/MaterialGraphNode_Root.cpp b/Engine/Source/Editor/UnrealEd/Private/MaterialGraphNode_Root.cpp
--- a/Engine/Source/Editor/UnrealEd/Private/MaterialGraphNode_Root.cpp
+++ b/Engine/Source/Editor/UnrealEd/Private/MaterialGraphNode_Root.cpp
@@ -12,6 +12,8 @@
 #include "Materials/MaterialAttributeDefinitionMap.h"
 #include "GraphEditorSettings.h"
 
+#include "MaterialNodes/SGraphNodeMaterialResult.h"
+
 #define LOCTEXT_NAMESPACE "MaterialGraphNode_Root"
 
 /////////////////////////////////////////////////////
@@ -46,6 +48,11 @@ void UMaterialGraphNode_Root::PostPlacedNewNode()
 	}
 }
 
+TSharedPtr<SGraphNode> UMaterialGraphNode_Root::CreateVisualWidget()
+{
+	return SNew(SGraphNodeMaterialResult, this);
+}
+
 UObject* UMaterialGraphNode_Root::GetMaterialNodeOwner() const
 {
 	return Material;
diff --git a/Engine/Source/Editor/UnrealEd/Private/MaterialGraphSchema.cpp b/Engine/Source/Editor/UnrealEd/Private/MaterialGraphSchema.cpp
--- a/Engine/Source/Editor/UnrealEd/Private/MaterialGraphSchema.cpp
+++ b/Engine/Source/Editor/UnrealEd/Private/MaterialGraphSchema.cpp
@@ -904,6 +904,11 @@ void UMaterialGraphSchema::UpdateMaterialOnDefaultValueChanged(const UEdGraph* G
 	FMaterialEditorUtilities::UpdateMaterialAfterGraphChange(Graph);
 }
 
+void UMaterialGraphSchema::MarkMaterialDirty(const UEdGraph* Graph) const
+{
+	FMaterialEditorUtilities::MarkMaterialDirty(Graph);
+}
+
 void UMaterialGraphSchema::UpdateDetailView(const UEdGraph* Graph) const
 {
 	FMaterialEditorUtilities::UpdateDetailView(Graph);
diff --git a/Engine/Source/Editor/UnrealEd/Public/MaterialGraphNode_Knot.h b/Engine/Source/Editor/UnrealEd/Public/MaterialGraphNode_Knot.h
--- a/Engine/Source/Editor/UnrealEd/Public/MaterialGraphNode_Knot.h
+++ b/Engine/Source/Editor/UnrealEd/Public/MaterialGraphNode_Knot.h
@@ -22,6 +22,7 @@ public:
 	virtual bool IsCompilerRelevant() const override { return false; }
 	virtual UEdGraphPin* GetPassThroughPin(const UEdGraphPin* FromPin) const override;
 	virtual bool ShouldDrawNodeAsControlPointOnly(int32& OutInputPinIndex, int32& OutOutputPinIndex) const override { OutInputPinIndex = 0;  OutOutputPinIndex = 1; return true; }
+	virtual TSharedPtr<SGraphNode> CreateVisualWidget() override;
 	// End of UEdGraphNode interface
 	
 	UEdGraphPin* GetInputPin() const
diff --git a/Engine/Source/Editor/UnrealEd/UnrealEd.Build.cs b/Engine/Source/Editor/UnrealEd/UnrealEd.Build.cs
--- a/Engine/Source/Editor/UnrealEd/UnrealEd.Build.cs
+++ b/Engine/Source/Editor/UnrealEd/UnrealEd.Build.cs
@@ -40,6 +40,7 @@ public class UnrealEd : ModuleRules
 				"Blutility",
 				"SlateReflector",
 				"PackagesDialog",
+				"GraphEditor"
 			}
 		);
 
diff --git a/Engine/Source/Runtime/Engine/Classes/Materials/MaterialExpression.h b/Engine/Source/Runtime/Engine/Classes/Materials/MaterialExpression.h
--- a/Engine/Source/Runtime/Engine/Classes/Materials/MaterialExpression.h
+++ b/Engine/Source/Runtime/Engine/Classes/Materials/MaterialExpression.h
@@ -366,6 +366,8 @@ class UMaterialExpression : public UObject
 	ENGINE_API virtual FName GetInputPinSubCategory(int32 PinIndex);
 	ENGINE_API virtual UObject* GetInputPinSubCategoryObject(int32 PinIndex);
 	ENGINE_API virtual void PinDefaultValueChanged(int32 PinIndex, const FString& DefaultValue);
+	ENGINE_API virtual void ForcePropertyValueChanged(FProperty* Property, bool bUpdatePreview = true);
+	ENGINE_API virtual void RefreshNode(bool bUpdatePreview = true);
 	ENGINE_API virtual FString GetInputPinDefaultValue(int32 PinIndex);
 	ENGINE_API virtual TArray<FProperty*> GetPropertyInputs() const;
 
diff --git a/Engine/Source/Runtime/Engine/Classes/Materials/MaterialExpressionCustom.h b/Engine/Source/Runtime/Engine/Classes/Materials/MaterialExpressionCustom.h
--- a/Engine/Source/Runtime/Engine/Classes/Materials/MaterialExpressionCustom.h
+++ b/Engine/Source/Runtime/Engine/Classes/Materials/MaterialExpressionCustom.h
@@ -85,6 +85,9 @@ class UMaterialExpressionCustom : public UMaterialExpression
 	UPROPERTY(EditAnywhere, Category=MaterialExpressionCustom)
 	TArray<FString> IncludeFilePaths;
 
+	UPROPERTY(VisibleAnywhere, Category=MaterialExpressionCustom)
+	bool ShowCode;
+
 	//~ Begin UObject Interface.
 #if WITH_EDITOR
 	virtual void PostEditChangeProperty(FPropertyChangedEvent& PropertyChangedEvent) override;
diff --git a/Engine/Source/Runtime/Engine/Private/Materials/MaterialExpressions.cpp b/Engine/Source/Runtime/Engine/Private/Materials/MaterialExpressions.cpp
--- a/Engine/Source/Runtime/Engine/Private/Materials/MaterialExpressions.cpp
+++ b/Engine/Source/Runtime/Engine/Private/Materials/MaterialExpressions.cpp
@@ -1272,12 +1272,43 @@ void UMaterialExpression::PinDefaultValueChanged(int32 PinIndex, const FString&
 		PostEditChangeProperty(Event);
 	}
 
-	// Update the expression preview and the material to reflect the change
-	GetDefault<UMaterialGraphSchema>()->ForceVisualizationCacheClear();
+	RefreshNode();
+}
+
+void UMaterialExpression::ForcePropertyValueChanged(FProperty* Property, bool bUpdatePreview)
+{
+	Modify();
+
+	FPropertyChangedEvent Event(Property);
+	PostEditChangeProperty(Event);
+
+	RefreshNode(bUpdatePreview);
+}
+
+void UMaterialExpression::RefreshNode(bool bUpdatePreview)
+{
 	const UMaterialGraphSchema* Schema = CastChecked<const UMaterialGraphSchema>(GraphNode->GetSchema());
-	Schema->UpdateMaterialOnDefaultValueChanged(GraphNode->GetGraph());
-	// There might be other properties affected by this property change (e.g. propertyA determines if propertyB is read-only) so refresh the detail view
-	Schema->UpdateDetailView(GraphNode->GetGraph());
+
+	if (bUpdatePreview)
+	{
+		// Make sure that all other nodes also require a preview update.
+		UMaterialGraphNode* Node = Cast<UMaterialGraphNode>(GraphNode);
+		if (Node)
+		{
+			Node->PropagatePropertyChange();
+		}
+
+		// Update the expression preview and the material to reflect the change
+		GetDefault<UMaterialGraphSchema>()->ForceVisualizationCacheClear();
+
+		Schema->UpdateMaterialOnDefaultValueChanged(GraphNode->GetGraph());
+		// There might be other properties affected by this property change (e.g. propertyA determines if propertyB is read-only) so refresh the detail view
+		Schema->UpdateDetailView(GraphNode->GetGraph());
+	}
+	else
+	{
+		Schema->MarkMaterialDirty(GraphNode->GetGraph());
+	}
 }
 
 FString UMaterialExpression::GetInputPinDefaultValue(int32 PinIndex)
@@ -13483,7 +13514,9 @@ UMaterialExpressionCustom::UMaterialExpressionCustom(const FObjectInitializer& O
 #endif // WITH_EDITORONLY_DATA
 
 	Description = TEXT("Custom");
-	Code = TEXT("1");
+	Code = TEXT("// The below expression will get compiled\n// into the output of this node\nfloat3(1, 1, 1)");
+
+	ShowCode = false;
 
 #if WITH_EDITORONLY_DATA
 	MenuCategories.Add(ConstructorStatics.NAME_Custom);
