Added copyright boilerplate to SGraphNodeMaterialCustom.cpp

diff --git a/Engine/Source/Editor/GraphEditor/Private/MaterialNodes/SGraphNodeMaterialCustom.cpp b/Engine/Source/Editor/GraphEditor/Private/MaterialNodes/SGraphNodeMaterialCustom.cpp
--- a/Engine/Source/Editor/GraphEditor/Private/MaterialNodes/SGraphNodeMaterialCustom.cpp
+++ b/Engine/Source/Editor/GraphEditor/Private/MaterialNodes/SGraphNodeMaterialCustom.cpp
@@ -1,3 +1,5 @@
+// Copyright Epic Games, Inc. All Rights Reserved.
+
 #include "MaterialNodes/SGraphNodeMaterialCustom.h"
 
 #include "Templates/Casts.h"
