Fix ensure after a drag enter handler runs after a mouse down handler has invalidated the graph

diff --git a/Engine/Source/Editor/GraphEditor/Private/SGraphPin.cpp b/Engine/Source/Editor/GraphEditor/Private/SGraphPin.cpp
--- a/Engine/Source/Editor/GraphEditor/Private/SGraphPin.cpp
+++ b/Engine/Source/Editor/GraphEditor/Private/SGraphPin.cpp
@@ -725,8 +725,9 @@ void SGraphPin::OnDragEnter( const FGeometry& MyGeometry, const FDragDropEvent&
 	// Is someone dragging a connection?
 	if (Operation->IsOfType<FGraphEditorDragDropAction>())
 	{
-		// Ensure that the pin is valid before using it
-		if(GraphPinObj != NULL && !GraphPinObj->IsPendingKill() && GraphPinObj->GetOuter() != NULL && GraphPinObj->GetOuter()->IsA(UEdGraphNode::StaticClass()))
+		// Ensure that the pin is valid before using it - in the case of OnDragEnter a previous OnPinNameMouseDown handler
+		// may have invalidated the graph data:
+		if(!bGraphDataInvalid && GraphPinObj != NULL && !GraphPinObj->IsPendingKill() && GraphPinObj->GetOuter() != NULL && GraphPinObj->GetOuter()->IsA(UEdGraphNode::StaticClass()))
 		{
 			if (GetIsConnectable())
 			{
