PR #9389: Allow using "spawn node" keyboard shortcuts while dragging from a pin

diff --git a/Engine/Config/BaseEditorPerProjectUserSettings.ini b/Engine/Config/BaseEditorPerProjectUserSettings.ini
--- a/Engine/Config/BaseEditorPerProjectUserSettings.ini
+++ b/Engine/Config/BaseEditorPerProjectUserSettings.ini
@@ -786,6 +786,7 @@ DefaultBrushRadius=128
 +Node=(Class=/Engine/EditorBlueprintResources/StandardMacros.StandardMacros:ForEachLoop Key=F Shift=false Ctrl=false Alt=false)
 +Node=(Class=/Engine/EditorBlueprintResources/StandardMacros.StandardMacros:DoOnce Key=O Shift=false Ctrl=false Alt=false)
 +Node=(Class=/Script/Engine.KismetArrayLibrary:Array_Get Key=A Shift=false Ctrl=false Alt=false)
++Node=(Class=/Script/BlueprintGraph.K2Node_Knot Key=R Shift=false Ctrl=false Alt=false)
 
 [DefaultEventNodes]
 +Node=(TargetClass=/Script/Engine.Actor TargetEvent="ReceiveBeginPlay")
@@ -837,6 +838,7 @@ DefaultBrushRadius=128
 +Node=(Class=/Script/Engine.MaterialExpressionConstant2Vector Key=Two Shift=false Ctrl=false Alt=false)
 +Node=(Class=/Script/Engine.MaterialExpressionConstant3Vector Key=Three Shift=false Ctrl=false Alt=false)
 +Node=(Class=/Script/Engine.MaterialExpressionConstant4Vector Key=Four Shift=false Ctrl=false Alt=false)
++Node=(Class=/Script/Engine.MaterialExpressionReroute Key=R Shift=true Ctrl=false Alt=false)
 
 [WidgetTemplatesExpanded]
 Common=True
diff --git a/Engine/Source/Editor/BlueprintGraph/Private/K2Node.cpp b/Engine/Source/Editor/BlueprintGraph/Private/K2Node.cpp
--- a/Engine/Source/Editor/BlueprintGraph/Private/K2Node.cpp
+++ b/Engine/Source/Editor/BlueprintGraph/Private/K2Node.cpp
@@ -444,9 +444,17 @@ void UK2Node::AutowireNewNode(UEdGraphPin* FromPin)
 		{
 			UEdGraphNode* FromPinNode = FromPin->GetOwningNode();
 			UEdGraphPin* FromThenPin = FromPinNode->FindPin(UEdGraphSchema_K2::PN_Then);
-
+			
+			// Prefer a pin named "Execute", but also accept "Exec" and "In", since some loop macros use those as their input pin names,
+			// and it's easier to add a workaround here that allows them to be better autowired than add more "special" pin names to the schema
 			UEdGraphPin* ToExecutePin = FindPin(UEdGraphSchema_K2::PN_Execute);
 
+			static const FName PN_Exec = TEXT("exec");
+			if (!ToExecutePin) ToExecutePin = FindPin(PN_Exec);
+
+			static const FName PN_In = TEXT("in");
+			if (!ToExecutePin) ToExecutePin = FindPin(PN_In);
+
 			if ((FromThenPin != nullptr) && (FromThenPin->LinkedTo.Num() == 0) && (ToExecutePin != nullptr) && K2Schema->ArePinsCompatible(FromThenPin, ToExecutePin, NULL))
 			{
 				if (K2Schema->TryCreateConnection(FromThenPin, ToExecutePin))
diff --git a/Engine/Source/Editor/GraphEditor/Private/SGraphPanel.cpp b/Engine/Source/Editor/GraphEditor/Private/SGraphPanel.cpp
--- a/Engine/Source/Editor/GraphEditor/Private/SGraphPanel.cpp
+++ b/Engine/Source/Editor/GraphEditor/Private/SGraphPanel.cpp
@@ -816,6 +816,130 @@ FReply SGraphPanel::OnKeyDown( const FGeometry& MyGeometry, const FKeyEvent& InK
 			ChangeZoomLevel(+1, CachedAllottedGeometryScaledSize / 2.f, InKeyEvent.IsControlDown());
 			return FReply::Handled();
 		}
+
+		// If we're only dragging out from a single pin then we can process the node spawn keyboard shortcuts
+		// If you're dragging multiple connections then it's probably a Ctrl + Drag to move wires around, which
+		// wouldn't make as much sense to do node creation during
+		UEdGraphPin* PreviewConnectionPin = PreviewConnectorFromPins.Num() == 1 ? PreviewConnectorFromPins[0].GetPinObj(*this) : nullptr;
+		if (OnSpawnNodeByShortcut.IsBound() && PreviewConnectionPin)
+		{
+			// Note: We can't use SavedMousePosForOnPaintEventLocalSpace since it isn't updated while dragging,
+			// and so would just be the mouse position of the connection origin. So instead we'll just use the current cursor pos
+			FVector2D NewNodePosition = PanelCoordToGraphCoord(MyGeometry.AbsoluteToLocal(FSlateApplication::Get().GetCursorPos()));
+			FInputChord KeyChord = FInputChord(InKeyEvent.GetKey(), EModifierKey::FromBools(InKeyEvent.IsControlDown(), InKeyEvent.IsAltDown(), InKeyEvent.IsShiftDown(), InKeyEvent.IsCommandDown()));
+
+			int32 NodeCountBefore = GraphObj->Nodes.Num();
+			FReply SpawnNodeReply = OnSpawnNodeByShortcut.Execute(KeyChord, NewNodePosition);
+			int32 NodeCountAfter = GraphObj->Nodes.Num();
+
+			// If we spawned a node then we won't call down into super and instead do some extra handling
+			int32 NumSpawnedNodes = NodeCountAfter - NodeCountBefore;
+			if (NumSpawnedNodes > 0)
+			{
+				TArrayView<UEdGraphNode*> SpawnedNodes = MakeArrayView(&GraphObj->Nodes[NodeCountBefore], NumSpawnedNodes);
+
+				// Try to auto-wire the newly spawned node
+				// Note: Usually the auto-wiring is handled by a schema action or something like FBlueprintMenuActionItemImpl::AutowireSpawnedNodes,
+				// but since we're not going through the regular action menu codepath we'll try to just do it here ourselves
+				// with slightly fewer heuristics. Could be good to expose that more publicly/centrally though
+				if (NumSpawnedNodes == 1)
+				{
+					bool bWasAutoWired = false;
+					for (const UEdGraphPin* Pin : SpawnedNodes[0]->Pins)
+					{
+						if (Pin->LinkedTo.Num() > 0)
+						{
+							bWasAutoWired = true;
+							break;
+						}
+					}
+
+					if (!bWasAutoWired)
+					{
+						SpawnedNodes[0]->AutowireNewNode(PreviewConnectionPin);
+					}
+				}
+
+				// The parent SNodePanel won't get a chance to set its LastKeyChordDetected, so we'll
+				// clear it out here so this key press is sort of "consumed"
+				LastKeyChordDetected = FInputChord();
+
+				// We spawned a new node through hotkey instead of letting go of the mouse,
+				// so we should cancel the drag to avoid the mouse up summoning the add node context menu
+				FSlateApplication::Get().CancelDragDrop();
+
+				OnStopMakingConnection(/*bForceStop=*/ true);
+
+				// Try to make the newly spawned node's connected pin end up underneath the mouse
+				TArrayView<UEdGraphPin*> DraggedFromPins = MakeArrayView(&PreviewConnectionPin, 1);
+				AdjustNewlySpawnedNodePositions(SpawnedNodes, DraggedFromPins, NewNodePosition);
+
+				UEdGraphPin* ResumeDraggingFromPin = nullptr;
+
+				// For now we don't let the spawn node shortcut provide an explicit pin that should be
+				// used when resuming dragging, but instead we'll just see if we spawned a 'control point only' (reroute) node,
+				// and if so then automatically use its appropriate in/out pin to continue the drag connection from
+				if (SpawnedNodes.Num() == 1)
+				{
+					UEdGraphNode* SpawnedNode = SpawnedNodes[0];
+					int32 OutPinIndex, InPinIndex;
+					if (SpawnedNode && SpawnedNode->ShouldDrawNodeAsControlPointOnly(OutPinIndex, InPinIndex))
+					{
+						ResumeDraggingFromPin = PreviewConnectionPin->Direction == EGPD_Input ? SpawnedNode->Pins[OutPinIndex] : SpawnedNode->Pins[InPinIndex];
+					}
+				}
+
+				// If we found a pin then need to start a new drag operation from it
+				if (ResumeDraggingFromPin)
+				{
+					// We need to do this one frame later since node widgets aren't created synchronously.
+					// Luckily the 'create widget' timer is scheduled synchronously within the actual spawning above,
+					// so this should always run after the new widget exists, though before it's been painted
+					static auto ResumeDragDelegate = [](double, float, TSharedRef<SGraphPanel> Panel, FGraphPinHandle DragFromPinHandle) -> EActiveTimerReturnType
+					{
+						TSharedPtr<SGraphPin> DragFromPinWidget = DragFromPinHandle.FindInGraphPanel(*Panel);
+						if (DragFromPinWidget.IsValid())
+						{
+							FPointerEvent MouseEvent = FPointerEvent(
+								Panel->LastPointerEvent.GetUserIndex(),
+								Panel->LastPointerEvent.GetPointerIndex(),
+								FSlateApplication::Get().GetCursorPos(),
+								FSlateApplication::Get().GetLastCursorPos(),
+								FSlateApplication::Get().GetPressedMouseButtons(),
+								EKeys::LeftMouseButton, /* EffectingButton */
+								0.f,                    /* WheelDelta */
+								FModifierKeysState()    /* InModifierKeys */
+							);
+
+							// This will technically be a frame behind but it shouldn't matter too much for this case
+							FGeometry PinGeometry = DragFromPinWidget->GetTickSpaceGeometry();
+
+							// This is far from ideal, but SGraphPin doesn't expose its SpawnPinDragEvent method, and even if we made a public equivalent
+							// there's still some extra validation and bookkeeping that we'd probably want to ensure gets run,
+							// so emulating a mousedown to keep things to a single code-path might actually be an okay option for now
+							// If the pin is editable and able to be dragged from, then it should return a reply that wants to begin an FDragConnection
+							FReply ResumeDragReply = DragFromPinWidget->OnPinMouseDown(PinGeometry, MouseEvent);
+							if (ResumeDragReply.GetDragDropContent().IsValid())
+							{
+								// Then to start a drag event outside of a Slate event reply, we'll sneakily pretend an external drag started
+								TSharedPtr<SWindow> WidgetWindow = FSlateApplication::Get().FindWidgetWindow(DragFromPinWidget.ToSharedRef());
+								if (WidgetWindow)
+								{
+									FDragDropEvent DragDropEvent(MouseEvent, ResumeDragReply.GetDragDropContent());
+									FSlateApplication::Get().ProcessDragEnterEvent(WidgetWindow.ToSharedRef(), DragDropEvent);
+								}
+							}
+						}
+
+						return EActiveTimerReturnType::Stop;
+					};
+
+					RegisterActiveTimer(0.f, FWidgetActiveTimerDelegate::CreateStatic(ResumeDragDelegate, SharedThis(this), FGraphPinHandle(ResumeDraggingFromPin)));
+				}
+
+				return FReply::Handled();
+			}
+		}
 	}
 
 	return SNodePanel::OnKeyDown(MyGeometry, InKeyEvent);
@@ -1299,6 +1423,179 @@ UEdGraphPin* SGraphPanel::GetPinUnderMouse(const FGeometry& MyGeometry, const FP
 	return PinUnderCursor;
 }
 
+void SGraphPanel::AdjustNewlySpawnedNodePositions(TArrayView<UEdGraphNode*> SpawnedNodes, TArrayView<UEdGraphPin*> DraggedFromPins, FVector2D AnchorPosition)
+{
+	static auto FindFirstLinkedAutoWiredPin = [](TArrayView<UEdGraphNode*> SpawnedNodes, TArrayView<UEdGraphPin*> DraggedFromPins) -> UEdGraphPin*
+	{
+		for (UEdGraphPin* DraggedPin : DraggedFromPins)
+		{
+			for (UEdGraphPin* LinkedToPin : DraggedPin->LinkedTo)
+			{
+				if (SpawnedNodes.Contains(LinkedToPin->GetOwningNode()))
+				{
+					return LinkedToPin;
+				}
+			}
+		}
+
+		return nullptr;
+	};
+
+	if (UEdGraphPin* DraggedConnectionWasAutoWiredToNewPin = FindFirstLinkedAutoWiredPin(SpawnedNodes, DraggedFromPins))
+	{
+		MoveNodesToAnchorPinAtGraphPosition(SpawnedNodes, FGraphPinHandle(DraggedConnectionWasAutoWiredToNewPin), AnchorPosition);
+	}
+}
+
+void SGraphPanel::MoveNodesToAnchorPinAtGraphPosition(TArrayView<UEdGraphNode*> NodesToMove, FGraphPinHandle PinToAnchor, FVector2D DesiredPinGraphPosition)
+{
+	struct FAnchorUtils
+	{
+		static int32 RoundToGrid(int32 Value, int32 GridSize)
+		{
+			return FMath::RoundToInt(static_cast<float>(Value) / GridSize) * GridSize;
+		}
+
+		// The standard SnapToGrid() will floor values, but this will round them instead
+		static void SnapToGridRounded(UEdGraphNode* Node, uint32 GridSnapSize)
+		{
+			Node->NodePosX = RoundToGrid(Node->NodePosX, GridSnapSize);
+			Node->NodePosY = RoundToGrid(Node->NodePosY, GridSnapSize);
+		}
+
+		static EActiveTimerReturnType AlignPinToPositionDelayed(double, float, TSharedRef<SGraphPanel> Panel, FGraphPinHandle DragFromPinHandle, FVector2D DesiredPinImageCenterGraph, TArray<UEdGraphNode*> SpawnedNodes)
+		{
+			AlignPinToPosition(Panel, DragFromPinHandle, DesiredPinImageCenterGraph, SpawnedNodes);
+			return EActiveTimerReturnType::Stop;
+		}
+
+		static void AlignPinToPosition(TSharedRef<SGraphPanel> Panel, FGraphPinHandle DragFromPinHandle, FVector2D DesiredPinImageCenterGraph, TArrayView<UEdGraphNode*> SpawnedNodes)
+		{
+			TSharedPtr<SGraphPin> DragFromPinWidget = DragFromPinHandle.FindInGraphPanel(*Panel);
+			if (!DragFromPinWidget.IsValid())
+			{
+				return;
+			}
+
+			// Normally, the new node's widgets haven't been painted yet, so can't use GetTickSpaceGeometry(),
+			// but we want to avoid them painting in the wrong position for a frame anyway, so will just
+			// arrange all nodes synchronously as part of FindChildGeometry so we know where they'd be without actually
+			// painting to screen. We could theoretically set the node's visibility to hidden until we've done this,
+			// but then would also need the connections and overlays etc. to also be hidden
+			// TODO: Theoretically we only need to arrange the single node and its descendants, not the whole graph
+			// (like ArrangeChildrenForContextMenuSummon but even more reduced), but if it requires SecondPassLayout
+			// then would need to do the whole graph anyway, so might be better to just make both types have the same cost
+			// for now now to avoid the extra complexity, given certain schemas purge all nodes on any graph change anyway
+
+			// Note: If this schema return true for `ShouldAlwaysPurgeOnModification` then new nodes might have been created,
+			// but this timer will be executed before the SNodePanel's Paint, and thus Tick have been called, which is where
+			// it populates the VisibleChildren list. And since ArrangeChildren uses this list to know which nodes to arrange,
+			// if we don't force it to be populated here then we won't be able to arrange and find the new node's pin
+			if (Panel->VisibleChildren.Num() == 0)
+			{
+				Panel->PopulateVisibleChildren(Panel->GetTickSpaceGeometry());
+			}
+
+			UEdGraphPin* PinObj = DragFromPinWidget->GetPinObj();
+			UEdGraphNode* OwningNode = PinObj ? PinObj->GetOwningNode() : nullptr;
+			TSharedPtr<SGraphNode> OwningNodeWidget = OwningNode ? Panel->GetNodeWidgetFromGuid(OwningNode->NodeGuid) : nullptr;
+			if (OwningNodeWidget)
+			{
+				bool bNodeNeedsPrepass = OwningNodeWidget->NeedsPrepass();
+				bool bNoDesiredSize = OwningNodeWidget->GetDesiredSize().GetMax() <= 0.001f;
+				if (bNodeNeedsPrepass || bNoDesiredSize)
+				{
+					const int32 ChildIndex = Panel->Children.Find(OwningNodeWidget.ToSharedRef());
+					const float SelfLayoutScaleMultiplier = Panel->PrepassLayoutScaleMultiplier.Get(1.f);
+					const float ChildLayoutScaleMultiplier = Panel->bHasRelativeLayoutScale
+								? SelfLayoutScaleMultiplier * Panel->GetRelativeLayoutScale(ChildIndex, SelfLayoutScaleMultiplier)
+								: SelfLayoutScaleMultiplier;
+
+					OwningNodeWidget->MarkPrepassAsDirty();
+					OwningNodeWidget->SlatePrepass(ChildLayoutScaleMultiplier);
+				}
+			}
+
+			TSharedPtr<SWidget> PinImageWidget = DragFromPinWidget->GetPinImageWidget();
+			TSet<TSharedRef<SWidget>> WidgetsToFind = { DragFromPinWidget.ToSharedRef() };
+
+			// If this pin had an image, then may as well look for that too since it can be
+			// more accurate than the center of the overall pin if it has any text wrapping etc.
+			if (PinImageWidget.IsValid())
+			{
+				WidgetsToFind.Add(PinImageWidget.ToSharedRef());
+			}
+
+			// Purposefully not using FindChildGeometry() since that's actually checked,
+			// and we don't want to panic if the node doesn't happen to exist yet. This util will both
+			// force layout to be computed, and if we had an image widget, give us that widget's geometry too
+			TMap<TSharedRef<SWidget>, FArrangedWidget> Result;
+			Panel->FindChildGeometries(Panel->GetTickSpaceGeometry(), WidgetsToFind, Result);
+
+			// Check if we found anything
+			FArrangedWidget* ArrangedDragFromPinWidget = Result.Find(DragFromPinWidget.ToSharedRef());
+			FArrangedWidget* ArrangedPinImageWidget = PinImageWidget.IsValid() ? Result.Find(PinImageWidget.ToSharedRef()) : nullptr;
+
+			// If we couldn't even find the pin widget we dragged from then we can't do much else
+			if (!ArrangedDragFromPinWidget)
+			{
+				return;
+			}
+
+			// Default to a similar calculation to Paint(), though we don't have access to the drawing policy so can't add the arrow offset
+			FVector2D PinImageCenterAbsolute = DragFromPinWidget->GetDirection() == EGPD_Input ? FGeometryHelper::VerticalMiddleLeftOf(ArrangedDragFromPinWidget->Geometry) : FGeometryHelper::VerticalMiddleRightOf(ArrangedDragFromPinWidget->Geometry);
+
+			// Though if we found the actual pin image widget, then we can use its exact center instead, which will probably be more accurate
+			// Either should be close enough given we'll be snapping to the grid afterwards anyway
+			if (ArrangedPinImageWidget)
+			{
+				PinImageCenterAbsolute = ArrangedPinImageWidget->Geometry.GetAbsolutePositionAtCoordinates(FVector2D(0.5f, 0.5f));
+			}
+
+			FVector2D PinImageCenterPanel = Panel->GetTickSpaceGeometry().AbsoluteToLocal(PinImageCenterAbsolute);
+			FVector2D PinImageCenterGraph = Panel->PanelCoordToGraphCoord(PinImageCenterPanel);
+			FVector2D Delta = DesiredPinImageCenterGraph - PinImageCenterGraph;
+
+			// Offset all nodes that were spawned by this same delta so that their relative
+			// positioning is maintained (re-snapping them each individually though)
+			for (UEdGraphNode* SpawnedNode : SpawnedNodes)
+			{
+				// Extra safety in case this was called from within a timer
+				if (IsValid(SpawnedNode))
+				{
+					SpawnedNode->NodePosX += Delta.X;
+					SpawnedNode->NodePosY += Delta.Y;
+
+					// Note: Not using the standard SnapToGrid() on purpose since we actually want to be as close
+					// to the user's dragged location as possible, and flooring ends up with noticeably more error
+					// when you eg. drop just a few px above the previous node
+					SnapToGridRounded(SpawnedNode, Panel->GetSnapGridSize());
+				}
+			}
+		}
+	};
+
+	if (NodesToMove.Num() == 0)
+	{
+		return;
+	}
+
+	// If the widget already exists then we can just align it synchronously,
+	// but if it doesn't (as this is usually used for newly spawned nodes before
+	// their deferred construction has run) then we'll try again next frame
+	TSharedPtr<SGraphPin> DragFromPinWidget = PinToAnchor.FindInGraphPanel(*this);
+	if (DragFromPinWidget.IsValid())
+	{
+		FAnchorUtils::AlignPinToPosition(SharedThis(this), PinToAnchor, DesiredPinGraphPosition, NodesToMove);
+	}
+	else
+	{
+		TArray<UEdGraphNode*> NodesToMoveCopy;
+		NodesToMoveCopy.Append(NodesToMove);
+		RegisterActiveTimer(0.f, FWidgetActiveTimerDelegate::CreateStatic(FAnchorUtils::AlignPinToPositionDelayed, SharedThis(this), PinToAnchor, DesiredPinGraphPosition, NodesToMoveCopy));
+	}
+}
+
 void SGraphPanel::OnBeginMakingConnection(UEdGraphPin* InOriginatingPin)
 {
 	OnBeginMakingConnection(FGraphPinHandle(InOriginatingPin));
diff --git a/Engine/Source/Editor/GraphEditor/Public/SGraphPanel.h b/Engine/Source/Editor/GraphEditor/Public/SGraphPanel.h
--- a/Engine/Source/Editor/GraphEditor/Public/SGraphPanel.h
+++ b/Engine/Source/Editor/GraphEditor/Public/SGraphPanel.h
@@ -383,6 +383,12 @@ private:
 	/** Returns a pin that is under the mouse, given a specified node. Returns nullptr if no valid node is given or no pin found */
 	UEdGraphPin* GetPinUnderMouse(const FGeometry& MyGeometry, const FPointerEvent& MouseEvent, TSharedPtr<SGraphNode> GraphNode) const;
 
+	/** If the spawned nodes were auto-wired from any of the dragged pins, then this will try to make the newly connected pin end up at SpawnGraphPosition */
+	void AdjustNewlySpawnedNodePositions(TArrayView<UEdGraphNode*> SpawnedNodes, TArrayView<UEdGraphPin*> DraggedFromPins, FVector2D SpawnGraphPosition);
+
+	/** Will move a group of nodes by the amount needed for an anchor pin to be at a certain position */
+	void MoveNodesToAnchorPinAtGraphPosition(TArrayView<UEdGraphNode*> NodesToMove, FGraphPinHandle PinToAnchor, FVector2D DesiredPinGraphPosition);
+
 	/** Handle to timer callback that allows the UI to refresh it's arrangement each tick, allows animations to occur within the UI */
 	TWeakPtr<FActiveTimerHandle> ActiveTimerHandleInvalidatePerTick;
 
