Preemprively fixing compile errors ahead of API deprecation.

diff --git a/Engine/Source/Editor/GraphEditor/Private/SGraphPanel.cpp b/Engine/Source/Editor/GraphEditor/Private/SGraphPanel.cpp
--- a/Engine/Source/Editor/GraphEditor/Private/SGraphPanel.cpp
+++ b/Engine/Source/Editor/GraphEditor/Private/SGraphPanel.cpp
@@ -836,7 +836,7 @@ FReply SGraphPanel::OnKeyDown( const FGeometry& MyGeometry, const FKeyEvent& InK
 			int32 NumSpawnedNodes = NodeCountAfter - NodeCountBefore;
 			if (NumSpawnedNodes > 0)
 			{
-				TArrayView<UEdGraphNode*> SpawnedNodes = MakeArrayView(&GraphObj->Nodes[NodeCountBefore], NumSpawnedNodes);
+				TArrayView<UEdGraphNode* const> SpawnedNodes = MakeArrayView(&GraphObj->Nodes[NodeCountBefore], NumSpawnedNodes);
 
 				// Try to auto-wire the newly spawned node
 				// Note: Usually the auto-wiring is handled by a schema action or something like FBlueprintMenuActionItemImpl::AutowireSpawnedNodes,
@@ -1423,9 +1423,9 @@ UEdGraphPin* SGraphPanel::GetPinUnderMouse(const FGeometry& MyGeometry, const FP
 	return PinUnderCursor;
 }
 
-void SGraphPanel::AdjustNewlySpawnedNodePositions(TArrayView<UEdGraphNode*> SpawnedNodes, TArrayView<UEdGraphPin*> DraggedFromPins, FVector2D AnchorPosition)
+void SGraphPanel::AdjustNewlySpawnedNodePositions(TArrayView<UEdGraphNode* const> SpawnedNodes, TArrayView<UEdGraphPin*> DraggedFromPins, FVector2D AnchorPosition)
 {
-	static auto FindFirstLinkedAutoWiredPin = [](TArrayView<UEdGraphNode*> SpawnedNodes, TArrayView<UEdGraphPin*> DraggedFromPins) -> UEdGraphPin*
+	static auto FindFirstLinkedAutoWiredPin = [](TArrayView<UEdGraphNode* const> SpawnedNodes, TArrayView<UEdGraphPin*> DraggedFromPins) -> UEdGraphPin*
 	{
 		for (UEdGraphPin* DraggedPin : DraggedFromPins)
 		{
@@ -1447,7 +1447,7 @@ void SGraphPanel::AdjustNewlySpawnedNodePositions(TArrayView<UEdGraphNode*> Spaw
 	}
 }
 
-void SGraphPanel::MoveNodesToAnchorPinAtGraphPosition(TArrayView<UEdGraphNode*> NodesToMove, FGraphPinHandle PinToAnchor, FVector2D DesiredPinGraphPosition)
+void SGraphPanel::MoveNodesToAnchorPinAtGraphPosition(TArrayView<UEdGraphNode* const> NodesToMove, FGraphPinHandle PinToAnchor, FVector2D DesiredPinGraphPosition)
 {
 	struct FAnchorUtils
 	{
@@ -1469,7 +1469,7 @@ void SGraphPanel::MoveNodesToAnchorPinAtGraphPosition(TArrayView<UEdGraphNode*>
 			return EActiveTimerReturnType::Stop;
 		}
 
-		static void AlignPinToPosition(TSharedRef<SGraphPanel> Panel, FGraphPinHandle DragFromPinHandle, FVector2D DesiredPinImageCenterGraph, TArrayView<UEdGraphNode*> SpawnedNodes)
+		static void AlignPinToPosition(TSharedRef<SGraphPanel> Panel, FGraphPinHandle DragFromPinHandle, FVector2D DesiredPinImageCenterGraph, TArrayView<UEdGraphNode* const> SpawnedNodes)
 		{
 			TSharedPtr<SGraphPin> DragFromPinWidget = DragFromPinHandle.FindInGraphPanel(*Panel);
 			if (!DragFromPinWidget.IsValid())
diff --git a/Engine/Source/Editor/GraphEditor/Public/SGraphPanel.h b/Engine/Source/Editor/GraphEditor/Public/SGraphPanel.h
--- a/Engine/Source/Editor/GraphEditor/Public/SGraphPanel.h
+++ b/Engine/Source/Editor/GraphEditor/Public/SGraphPanel.h
@@ -384,10 +384,10 @@ private:
 	UEdGraphPin* GetPinUnderMouse(const FGeometry& MyGeometry, const FPointerEvent& MouseEvent, TSharedPtr<SGraphNode> GraphNode) const;
 
 	/** If the spawned nodes were auto-wired from any of the dragged pins, then this will try to make the newly connected pin end up at SpawnGraphPosition */
-	void AdjustNewlySpawnedNodePositions(TArrayView<UEdGraphNode*> SpawnedNodes, TArrayView<UEdGraphPin*> DraggedFromPins, FVector2D SpawnGraphPosition);
+	void AdjustNewlySpawnedNodePositions(TArrayView<UEdGraphNode* const> SpawnedNodes, TArrayView<UEdGraphPin*> DraggedFromPins, FVector2D SpawnGraphPosition);
 
 	/** Will move a group of nodes by the amount needed for an anchor pin to be at a certain position */
-	void MoveNodesToAnchorPinAtGraphPosition(TArrayView<UEdGraphNode*> NodesToMove, FGraphPinHandle PinToAnchor, FVector2D DesiredPinGraphPosition);
+	void MoveNodesToAnchorPinAtGraphPosition(TArrayView<UEdGraphNode* const> NodesToMove, FGraphPinHandle PinToAnchor, FVector2D DesiredPinGraphPosition);
 
 	/** Handle to timer callback that allows the UI to refresh it's arrangement each tick, allows animations to occur within the UI */
 	TWeakPtr<FActiveTimerHandle> ActiveTimerHandleInvalidatePerTick;
