RigVM: Fix crash when clicking on pin of default node in details panel
#rnx

diff --git a/Engine/Source/Editor/GraphEditor/Private/SGraphPin.cpp b/Engine/Source/Editor/GraphEditor/Private/SGraphPin.cpp
--- a/Engine/Source/Editor/GraphEditor/Private/SGraphPin.cpp
+++ b/Engine/Source/Editor/GraphEditor/Private/SGraphPin.cpp
@@ -516,7 +516,10 @@ FReply SGraphPin::OnPinMouseDown( const FGeometry& SenderGeometry, const FPointe
 					TArray<TSharedRef<SGraphPin>> PinArray;
 					PinArray.Add(SharedThis(this));
 
-					return FReply::Handled().BeginDragDrop(SpawnPinDragEvent(OwnerNodePinned->GetOwnerPanel().ToSharedRef(), PinArray));
+					if (TSharedPtr<SGraphPanel> OwnerGraphPanel = OwnerNodePinned->GetOwnerPanel())
+					{
+						return FReply::Handled().BeginDragDrop(SpawnPinDragEvent(OwnerGraphPanel.ToSharedRef(), PinArray));
+					}
 				}
 				else
 				{
