Graph Nodes can now specify what overflow policy to use for the node title

diff --git a/Engine/Source/Editor/GraphEditor/Private/SGraphNode.cpp b/Engine/Source/Editor/GraphEditor/Private/SGraphNode.cpp
--- a/Engine/Source/Editor/GraphEditor/Private/SGraphNode.cpp
+++ b/Engine/Source/Editor/GraphEditor/Private/SGraphNode.cpp
@@ -810,7 +810,8 @@ TSharedRef<SWidget> SGraphNode::CreateTitleWidget(TSharedPtr<SNodeTitle> NodeTit
 		.OnVerifyTextChanged(this, &SGraphNode::OnVerifyNameTextChanged)
 		.OnTextCommitted(this, &SGraphNode::OnNameTextCommited)
 		.IsReadOnly(this, &SGraphNode::IsNameReadOnly)
-		.IsSelected(this, &SGraphNode::IsSelectedExclusively);
+		.IsSelected(this, &SGraphNode::IsSelectedExclusively)
+		.OverflowPolicy(GetNameOverflowPolicy());
 	InlineEditableText->SetColorAndOpacity(TAttribute<FLinearColor>::Create(TAttribute<FLinearColor>::FGetter::CreateSP(this, &SGraphNode::GetNodeTitleTextColor)));
 
 	return InlineEditableText.ToSharedRef();
diff --git a/Engine/Source/Editor/GraphEditor/Public/SGraphNode.h b/Engine/Source/Editor/GraphEditor/Public/SGraphNode.h
--- a/Engine/Source/Editor/GraphEditor/Public/SGraphNode.h
+++ b/Engine/Source/Editor/GraphEditor/Public/SGraphNode.h
@@ -318,6 +318,9 @@ protected:
 	// Override this to add widgets below the pins but above advanced view arrow
 	virtual void CreateBelowPinControls(TSharedPtr<SVerticalBox> MainBox) {}
 
+	/** Determines how the node title overflow is handled */
+	virtual TOptional<ETextOverflowPolicy> GetNameOverflowPolicy() const { return {}; }
+
 	/* Helper function to check if node can be renamed */
 	virtual bool IsNameReadOnly () const;
 
