Made the title border margin configurable via protected member

diff --git a/Engine/Source/Editor/GraphEditor/Private/SGraphNode.cpp b/Engine/Source/Editor/GraphEditor/Private/SGraphNode.cpp
--- a/Engine/Source/Editor/GraphEditor/Private/SGraphNode.cpp
+++ b/Engine/Source/Editor/GraphEditor/Private/SGraphNode.cpp
@@ -877,9 +877,7 @@ void SGraphNode::UpdateGraphNode()
 			[
 				SNew(SBorder)
 				.BorderImage( FAppStyle::GetBrush("Graph.Node.ColorSpill") )
-				// The extra margin on the right
-				// is for making the color spill stretch well past the node title
-				.Padding( FMargin(10,5,30,3) )
+				.Padding(TitleBorderMargin)
 				.BorderBackgroundColor( this, &SGraphNode::GetNodeTitleColor )
 				[
 					SNew(SHorizontalBox)
diff --git a/Engine/Source/Editor/GraphEditor/Public/SGraphNode.h b/Engine/Source/Editor/GraphEditor/Public/SGraphNode.h
--- a/Engine/Source/Editor/GraphEditor/Public/SGraphNode.h
+++ b/Engine/Source/Editor/GraphEditor/Public/SGraphNode.h
@@ -443,6 +443,10 @@ protected:
 	FCurveHandle ZoomCurve;
 	FCurveHandle FadeCurve;
 
+	/* The margin used by the border containing the title. Can be changed by child classes if necessary.
+	*  The extra padding on the right is for making the color spill stretch well past the node title */
+	FMargin TitleBorderMargin = FMargin(10.f, 5.f, 30.f, 3.f);
+
 	/** Is this node editable */
 	TAttribute<bool> IsEditable;
 
