Add color Substrate wire.

diff --git a/Engine/Source/Editor/GraphEditor/Private/MaterialGraphConnectionDrawingPolicy.cpp b/Engine/Source/Editor/GraphEditor/Private/MaterialGraphConnectionDrawingPolicy.cpp
--- a/Engine/Source/Editor/GraphEditor/Private/MaterialGraphConnectionDrawingPolicy.cpp
+++ b/Engine/Source/Editor/GraphEditor/Private/MaterialGraphConnectionDrawingPolicy.cpp
@@ -17,6 +17,7 @@
 #include "Templates/Casts.h"
 #include "Templates/SharedPointer.h"
 #include "UObject/NameTypes.h"
+#include "SGraphSubstrateMaterial.h"
 
 class FSlateRect;
 
@@ -124,6 +125,11 @@ void FMaterialGraphConnectionDrawingPolicy::DetermineWiringStyle(UEdGraphPin* Ou
 	Params.AssociatedPin2 = InputPin;
 	Params.WireColor = MaterialGraphSchema->ActivePinColor;
 
+	if (Strata::IsStrataEnabled() && (FSubstrateWidget::HasSubstrateType(OutputPin) || FSubstrateWidget::HasSubstrateType(InputPin)))
+	{
+		Params.WireColor = FSubstrateWidget::GetConnectionColor();
+	}
+
 	UE::Shader::EValueType InputType = UE::Shader::EValueType::Void;
 	UE::Shader::EValueType OutputType = UE::Shader::EValueType::Void;
 	bool bInactivePin = false;
diff --git a/Engine/Source/Editor/GraphEditor/Private/MaterialNodes/SGraphSubstrateMaterial.cpp b/Engine/Source/Editor/GraphEditor/Private/MaterialNodes/SGraphSubstrateMaterial.cpp
--- a/Engine/Source/Editor/GraphEditor/Private/MaterialNodes/SGraphSubstrateMaterial.cpp
+++ b/Engine/Source/Editor/GraphEditor/Private/MaterialNodes/SGraphSubstrateMaterial.cpp
@@ -5,6 +5,7 @@
 #include "Rendering/RenderingCommon.h"
 #include "Rendering/StrataMaterialShared.h"
 #include "MaterialGraph/MaterialGraphNode.h"
+#include "MaterialGraph/MaterialGraphNode_Root.h"
 #include "Materials/MaterialExpression.h"
 #include "Materials/MaterialExpressionStrata.h"
 #include "SGraphPin.h"
@@ -21,6 +22,45 @@ enum class ESubstrateWidgetOutputType : uint8
 
 static EStyleColor GetSubstrateWidgetColor0() { return EStyleColor::AccentBlue;  }
 static EStyleColor GetSubstrateWidgetColor1() { return EStyleColor::AccentGreen;  }
+FLinearColor FSubstrateWidget::GetConnectionColor() { return FLinearColor(0.16, 0.015, 0.24) * 4.f; }
+
+bool FSubstrateWidget::HasSubstrateType(const UEdGraphPin* InPin)
+{
+	if (InPin == nullptr) return false;
+
+	if (UMaterialGraphNode_Root* RootPinNode = Cast<UMaterialGraphNode_Root>(InPin->GetOwningNode()))
+	{
+		EMaterialProperty propertyId = (EMaterialProperty)FCString::Atoi(*InPin->PinType.PinSubCategory.ToString());
+		switch (propertyId)
+		{
+			case MP_FrontMaterial:
+				return true;
+		}
+	}
+	if (UMaterialGraphNode* PinNode = Cast<UMaterialGraphNode>(InPin->GetOwningNode()))
+	{
+		TArrayView<FExpressionInput*> ExpressionInputs = PinNode->MaterialExpression->GetInputsView();
+		FName TargetPinName = PinNode->GetShortenPinName(InPin->PinName);
+
+		for (int32 Index = 0; Index < ExpressionInputs.Num(); ++Index)
+		{
+			FExpressionInput* Input = ExpressionInputs[Index];
+			FName InputName = PinNode->MaterialExpression->GetInputName(Index);
+			InputName = PinNode->GetShortenPinName(InputName);
+
+			if (InputName == TargetPinName)
+			{
+				switch (PinNode->MaterialExpression->GetInputType(Index))
+				{
+					case MCT_Strata:
+						return true;
+				}
+				break;
+			}
+		}
+	}
+	return false;
+}
 
 static const TSharedRef<SWidget> InternalProcessOperator(
 	const FStrataMaterialCompilationOutput& CompilationOutput, 
diff --git a/Engine/Source/Editor/GraphEditor/Public/SGraphSubstrateMaterial.h b/Engine/Source/Editor/GraphEditor/Public/SGraphSubstrateMaterial.h
--- a/Engine/Source/Editor/GraphEditor/Public/SGraphSubstrateMaterial.h
+++ b/Engine/Source/Editor/GraphEditor/Public/SGraphSubstrateMaterial.h
@@ -14,4 +14,6 @@ struct GRAPHEDITOR_API FSubstrateWidget
 	static const TSharedRef<SWidget> ProcessOperator(const FStrataMaterialCompilationOutput& CompilationOutput);
 	static const TSharedRef<SWidget> ProcessOperator(const FStrataMaterialCompilationOutput& CompilationOutput, const FGuid& InGuid);
 	static void GetPinColor(TSharedPtr<SGraphPin>& Out, const UMaterialGraphNode* InNode);
+	static FLinearColor GetConnectionColor();
+	static bool HasSubstrateType(const UEdGraphPin* InPin);
 };
\ No newline at end of file
diff --git a/Engine/Source/Editor/MaterialEditor/Private/MaterialEditor.cpp b/Engine/Source/Editor/MaterialEditor/Private/MaterialEditor.cpp
--- a/Engine/Source/Editor/MaterialEditor/Private/MaterialEditor.cpp
+++ b/Engine/Source/Editor/MaterialEditor/Private/MaterialEditor.cpp
@@ -156,6 +156,7 @@
 #include "Materials/MaterialExpression.h"
 #include "MaterialCachedHLSLTree.h"
 #include "SMaterialEditorStrataWidget.h"
+#include "SGraphSubstrateMaterial.h"
 
 #include "SMaterialParametersOverviewWidget.h"
 #include "SMaterialEditorCustomPrimitiveDataWidget.h"
@@ -169,6 +170,8 @@
 #include "UObject/MetaData.h"
 #include "ToolMenus.h"
 
+
+
 #define LOCTEXT_NAMESPACE "MaterialEditor"
 
 DEFINE_LOG_CATEGORY_STATIC(LogMaterialEditor, Log, All);
@@ -4895,37 +4898,7 @@ bool FMaterialEditor::OnCanCreateStrataNodeForPin(const FToolMenuContext& InMenu
 
 	if ((TargetPin->Direction == EEdGraphPinDirection::EGPD_Input) && (TargetPin->LinkedTo.Num() == 0))
 	{
-		if (RootPinNode != nullptr)
-		{
-			EMaterialProperty propertyId = (EMaterialProperty)FCString::Atoi(*TargetPin->PinType.PinSubCategory.ToString());
-			switch (propertyId)
-			{
-			case MP_FrontMaterial:
-				return true;
-			}
-		}
-		else if (OtherPinNode)
-		{
-			TArrayView<FExpressionInput*> ExpressionInputs = OtherPinNode->MaterialExpression->GetInputsView();
-			FName TargetPinName = OtherPinNode->GetShortenPinName(TargetPin->PinName);
-
-			for (int32 Index = 0; Index < ExpressionInputs.Num(); ++Index)
-			{
-				FExpressionInput* Input = ExpressionInputs[Index];
-				FName InputName = OtherPinNode->MaterialExpression->GetInputName(Index);
-				InputName = OtherPinNode->GetShortenPinName(InputName);
-
-				if (InputName == TargetPinName)
-				{
-					switch (OtherPinNode->MaterialExpression->GetInputType(Index))
-					{
-					case MCT_Strata:
-						return true;
-					}
-					break;
-				}
-			}
-		}
+		return FSubstrateWidget::HasSubstrateType(TargetPin);
 	}
 	else if (TargetPin && (TargetPin->Direction == EEdGraphPinDirection::EGPD_Output) && (TargetPin->LinkedTo.Num() == 0) && NodeForPin != EStrataNodeForPin::Slab)
 	{
