* Fix Substrate colored connection on input
* Add colored pin for Substrate data flow

diff --git a/Engine/Source/Editor/GraphEditor/Private/MaterialGraphConnectionDrawingPolicy.cpp b/Engine/Source/Editor/GraphEditor/Private/MaterialGraphConnectionDrawingPolicy.cpp
--- a/Engine/Source/Editor/GraphEditor/Private/MaterialGraphConnectionDrawingPolicy.cpp
+++ b/Engine/Source/Editor/GraphEditor/Private/MaterialGraphConnectionDrawingPolicy.cpp
@@ -125,7 +125,7 @@ void FMaterialGraphConnectionDrawingPolicy::DetermineWiringStyle(UEdGraphPin* Ou
 	Params.AssociatedPin2 = InputPin;
 	Params.WireColor = MaterialGraphSchema->ActivePinColor;
 
-	if (Strata::IsStrataEnabled() && (FSubstrateWidget::HasSubstrateType(OutputPin) || FSubstrateWidget::HasSubstrateType(InputPin)))
+	if (Strata::IsStrataEnabled() && (FSubstrateWidget::HasOutputSubstrateType(OutputPin) || FSubstrateWidget::HasInputSubstrateType(InputPin) || FSubstrateWidget::HasInputSubstrateType(OutputPin)))
 	{
 		Params.WireColor = FSubstrateWidget::GetConnectionColor();
 	}
diff --git a/Engine/Source/Editor/GraphEditor/Private/MaterialNodes/SGraphSubstrateMaterial.cpp b/Engine/Source/Editor/GraphEditor/Private/MaterialNodes/SGraphSubstrateMaterial.cpp
--- a/Engine/Source/Editor/GraphEditor/Private/MaterialNodes/SGraphSubstrateMaterial.cpp
+++ b/Engine/Source/Editor/GraphEditor/Private/MaterialNodes/SGraphSubstrateMaterial.cpp
@@ -6,6 +6,7 @@
 #include "Rendering/SubstrateMaterialShared.h"
 #include "MaterialGraph/MaterialGraphNode.h"
 #include "MaterialGraph/MaterialGraphNode_Root.h"
+#include "MaterialGraph/MaterialGraphSchema.h"
 #include "Materials/MaterialExpression.h"
 #include "Materials/MaterialExpressionSubstrate.h"
 #include "SGraphPin.h"
@@ -24,7 +25,7 @@ static EStyleColor GetSubstrateWidgetColor0() { return EStyleColor::AccentBlue;
 static EStyleColor GetSubstrateWidgetColor1() { return EStyleColor::AccentGreen;  }
 FLinearColor FSubstrateWidget::GetConnectionColor() { return FLinearColor(0.16, 0.015, 0.24) * 4.f; }
 
-bool FSubstrateWidget::HasSubstrateType(const UEdGraphPin* InPin)
+bool FSubstrateWidget::HasInputSubstrateType(const UEdGraphPin* InPin)
 {
 	if (InPin == nullptr) return false;
 
@@ -62,6 +63,21 @@ bool FSubstrateWidget::HasSubstrateType(const UEdGraphPin* InPin)
 	return false;
 }
 
+bool FSubstrateWidget::HasOutputSubstrateType(const UEdGraphPin* InPin)
+{
+	if (InPin == nullptr || InPin->Direction != EGPD_Output) return false;
+
+	if (UMaterialGraphNode* PinNode = Cast<UMaterialGraphNode>(InPin->GetOwningNode()))
+	{
+		const TArray<FExpressionOutput>& ExpressionOutputs = PinNode->MaterialExpression->GetOutputs();
+		if (InPin->SourceIndex < ExpressionOutputs.Num() && PinNode->MaterialExpression->GetOutputType(InPin->SourceIndex) == MCT_Strata)
+		{
+			return true;
+		}
+	}
+	return false;
+}
+
 static const TSharedRef<SWidget> InternalProcessOperator(
 	const FStrataMaterialCompilationOutput& CompilationOutput, 
 	const FStrataOperator& Op, 
@@ -208,7 +224,8 @@ void FSubstrateWidget::GetPinColor(TSharedPtr<SGraphPin>& Out, const UMaterialGr
 	FLinearColor ColorModifier;
 	bool bHasColorModifier = false;
 	// Strata operator override pin color to ease material topology visualization
-	const FName PinName = Out->GetPinObj()->PinName;
+	const UEdGraphPin* Pin = Out->SGraphPin::GetPinObj();
+	const FName PinName = Pin->PinName;
 	if (InNode->MaterialExpression->IsA(UMaterialExpressionStrataVerticalLayering::StaticClass()))
 	{			
 		if (PinName == InNode->MaterialExpression->GetInputName(0)) // Top
@@ -249,6 +266,20 @@ void FSubstrateWidget::GetPinColor(TSharedPtr<SGraphPin>& Out, const UMaterialGr
 		}
 	}
 
+	if (InNode->MaterialExpression->IsA(UMaterialExpressionStrataBSDF::StaticClass()) && Out->GetDirection() == EGPD_Output)
+	{
+		bHasColorModifier = true;
+		ColorModifier = FSubstrateWidget::GetConnectionColor();
+	}
+	else if (Pin && UMaterialGraphSchema::GetMaterialValueType(Pin) == MCT_Strata)
+	{
+		if (!Out->IsConnected())
+		{
+			bHasColorModifier = true;
+			ColorModifier = FSubstrateWidget::GetConnectionColor();
+		}
+	}	
+	
 	if (bHasColorModifier)
 	{
 		Out->SetPinColorModifier(ColorModifier);
diff --git a/Engine/Source/Editor/GraphEditor/Private/MaterialPins/SGraphPinMaterialInput.cpp b/Engine/Source/Editor/GraphEditor/Private/MaterialPins/SGraphPinMaterialInput.cpp
--- a/Engine/Source/Editor/GraphEditor/Private/MaterialPins/SGraphPinMaterialInput.cpp
+++ b/Engine/Source/Editor/GraphEditor/Private/MaterialPins/SGraphPinMaterialInput.cpp
@@ -9,10 +9,19 @@
 #include "MaterialGraph/MaterialGraphSchema.h"
 #include "Misc/AssertionMacros.h"
 #include "Templates/Casts.h"
+#include "SGraphSubstrateMaterial.h"
 
 void SGraphPinMaterialInput::Construct(const FArguments& InArgs, UEdGraphPin* InGraphPinObj)
 {
-	SGraphPin::Construct(SGraphPin::FArguments().UsePinColorForText(true), InGraphPinObj);
+	bool bUsePinColor = true;
+	if (Strata::IsStrataEnabled())
+	{
+		if (UMaterialGraphSchema::GetMaterialValueType(InGraphPinObj) == MCT_Strata)
+		{
+			bUsePinColor = false;
+		}			
+	}
+	SGraphPin::Construct(SGraphPin::FArguments().UsePinColorForText(bUsePinColor), InGraphPinObj);
 }
 
 FSlateColor SGraphPinMaterialInput::GetPinColor() const
@@ -23,6 +32,13 @@ FSlateColor SGraphPinMaterialInput::GetPinColor() const
 
 	if (MaterialGraph->IsInputActive(GraphPinObj))
 	{
+		if (Strata::IsStrataEnabled())
+		{
+			if (UMaterialGraphSchema::GetMaterialValueType(GraphPinObj) == MCT_Strata)
+			{
+				return FSlateColor(FSubstrateWidget::GetConnectionColor());
+			}			
+		}
 		return Schema->ActivePinColor;
 	}
 	else
diff --git a/Engine/Source/Editor/GraphEditor/Public/SGraphSubstrateMaterial.h b/Engine/Source/Editor/GraphEditor/Public/SGraphSubstrateMaterial.h
--- a/Engine/Source/Editor/GraphEditor/Public/SGraphSubstrateMaterial.h
+++ b/Engine/Source/Editor/GraphEditor/Public/SGraphSubstrateMaterial.h
@@ -15,5 +15,6 @@ struct GRAPHEDITOR_API FSubstrateWidget
 	static const TSharedRef<SWidget> ProcessOperator(const FStrataMaterialCompilationOutput& CompilationOutput, const FGuid& InGuid);
 	static void GetPinColor(TSharedPtr<SGraphPin>& Out, const UMaterialGraphNode* InNode);
 	static FLinearColor GetConnectionColor();
-	static bool HasSubstrateType(const UEdGraphPin* InPin);
+	static bool HasInputSubstrateType(const UEdGraphPin* InPin);
+	static bool HasOutputSubstrateType(const UEdGraphPin* InPin);
 };
\ No newline at end of file
diff --git a/Engine/Source/Editor/MaterialEditor/Private/MaterialEditor.cpp b/Engine/Source/Editor/MaterialEditor/Private/MaterialEditor.cpp
--- a/Engine/Source/Editor/MaterialEditor/Private/MaterialEditor.cpp
+++ b/Engine/Source/Editor/MaterialEditor/Private/MaterialEditor.cpp
@@ -4898,19 +4898,11 @@ bool FMaterialEditor::OnCanCreateStrataNodeForPin(const FToolMenuContext& InMenu
 
 	if ((TargetPin->Direction == EEdGraphPinDirection::EGPD_Input) && (TargetPin->LinkedTo.Num() == 0))
 	{
-		return FSubstrateWidget::HasSubstrateType(TargetPin);
+		return FSubstrateWidget::HasInputSubstrateType(TargetPin);
 	}
 	else if (TargetPin && (TargetPin->Direction == EEdGraphPinDirection::EGPD_Output) && (TargetPin->LinkedTo.Num() == 0) && NodeForPin != EStrataNodeForPin::Slab)
 	{
-		if (OtherPinNode)
-		{
-			const TArray<FExpressionOutput>& ExpressionOutputs = OtherPinNode->MaterialExpression->GetOutputs();
-			check(TargetPin->SourceIndex < ExpressionOutputs.Num());
-			if (OtherPinNode->MaterialExpression->GetOutputType(TargetPin->SourceIndex) == MCT_Strata)
-			{
-				return true;
-			}
-		}
+		return FSubstrateWidget::HasOutputSubstrateType(TargetPin);
 	}
 
 	return false;
