Fix CIS unity build issue.

diff --git a/Engine/Source/Editor/GraphEditor/Private/MaterialNodes/SGraphSubstrateMaterial.cpp b/Engine/Source/Editor/GraphEditor/Private/MaterialNodes/SGraphSubstrateMaterial.cpp
--- a/Engine/Source/Editor/GraphEditor/Private/MaterialNodes/SGraphSubstrateMaterial.cpp
+++ b/Engine/Source/Editor/GraphEditor/Private/MaterialNodes/SGraphSubstrateMaterial.cpp
@@ -14,6 +14,7 @@
 #include "Styling/StyleColors.h"
 #include "SubstrateDefinitions.h"
 #include "Widgets/SBoxPanel.h"
+#include "MaterialValueType.h"
 
 enum class ESubstrateWidgetOutputType : uint8
 {
