Fix CIS non-unity build.

diff --git a/Engine/Source/Editor/GraphEditor/Private/MaterialPins/SGraphPinMaterialInput.cpp b/Engine/Source/Editor/GraphEditor/Private/MaterialPins/SGraphPinMaterialInput.cpp
--- a/Engine/Source/Editor/GraphEditor/Private/MaterialPins/SGraphPinMaterialInput.cpp
+++ b/Engine/Source/Editor/GraphEditor/Private/MaterialPins/SGraphPinMaterialInput.cpp
@@ -10,6 +10,7 @@
 #include "Misc/AssertionMacros.h"
 #include "Templates/Casts.h"
 #include "SGraphSubstrateMaterial.h"
+#include "MaterialValueType.h"
 
 void SGraphPinMaterialInput::Construct(const FArguments& InArgs, UEdGraphPin* InGraphPinObj)
 {
