Fix Substrate data pin color on node with several inputs.

diff --git a/Engine/Source/Editor/GraphEditor/Private/MaterialNodes/SGraphSubstrateMaterial.cpp b/Engine/Source/Editor/GraphEditor/Private/MaterialNodes/SGraphSubstrateMaterial.cpp
--- a/Engine/Source/Editor/GraphEditor/Private/MaterialNodes/SGraphSubstrateMaterial.cpp
+++ b/Engine/Source/Editor/GraphEditor/Private/MaterialNodes/SGraphSubstrateMaterial.cpp
@@ -267,7 +267,7 @@ void FSubstrateWidget::GetPinColor(TSharedPtr<SGraphPin>& Out, const UMaterialGr
 		}
 	}
 
-	if (InNode->MaterialExpression->IsA(UMaterialExpressionSubstrateBSDF::StaticClass()) && Out->GetDirection() == EGPD_Output)
+	if (InNode->MaterialExpression->IsA(UMaterialExpressionSubstrateBSDF::StaticClass()) && Out->GetDirection() == EGPD_Output && UMaterialGraphSchema::GetMaterialValueType(Pin) == MCT_Substrate)
 	{
 		bHasColorModifier = true;
 		ColorModifier = FSubstrateWidget::GetConnectionColor();
