Fix Substrate topology visualization on Coverage operator.

diff --git a/Engine/Source/Editor/GraphEditor/Private/MaterialNodes/SGraphSubstrateMaterial.cpp b/Engine/Source/Editor/GraphEditor/Private/MaterialNodes/SGraphSubstrateMaterial.cpp
--- a/Engine/Source/Editor/GraphEditor/Private/MaterialNodes/SGraphSubstrateMaterial.cpp
+++ b/Engine/Source/Editor/GraphEditor/Private/MaterialNodes/SGraphSubstrateMaterial.cpp
@@ -92,8 +92,11 @@ static const TSharedRef<SWidget> InternalProcessOperator(
 	switch (Op.OperatorType)
 	{
 		case SUBSTRATE_OPERATOR_WEIGHT:
-			return InternalProcessOperator(CompilationOutput, CompilationOutput.Operators[Op.LeftIndex], OutputType, InGuid, OverrideColor);
-			break;
+		{
+			const EStyleColor Color = bIsCurrent ? EStyleColor::AccentGreen : OverrideColor;
+			return InternalProcessOperator(CompilationOutput, CompilationOutput.Operators[Op.LeftIndex], OutputType, InGuid, Color);
+		}
+		break;
 		case SUBSTRATE_OPERATOR_VERTICAL:
 		{
 			auto VerticalOperator = SNew(SVerticalBox)
