Fix crash with Substrate enabled with certain material.

diff --git a/Engine/Source/Editor/GraphEditor/Private/MaterialNodes/SGraphSubstrateMaterial.cpp b/Engine/Source/Editor/GraphEditor/Private/MaterialNodes/SGraphSubstrateMaterial.cpp
--- a/Engine/Source/Editor/GraphEditor/Private/MaterialNodes/SGraphSubstrateMaterial.cpp
+++ b/Engine/Source/Editor/GraphEditor/Private/MaterialNodes/SGraphSubstrateMaterial.cpp
@@ -222,6 +222,11 @@ const TSharedRef<SWidget> FSubstrateWidget::ProcessOperator(const FSubstrateMate
 
 void FSubstrateWidget::GetPinColor(TSharedPtr<SGraphPin>& Out, const UMaterialGraphNode* InNode)
 {	
+	if (!InNode || !InNode->MaterialExpression)
+	{
+		return;
+	}
+
 	const FLinearColor Color0 = USlateThemeManager::Get().GetColor(GetSubstrateWidgetColor0());
 	const FLinearColor Color1 = USlateThemeManager::Get().GetColor(GetSubstrateWidgetColor1());
 
