Material Root Node Inline Editting
Support for sliders

diff --git a/Engine/Source/Editor/GraphEditor/Private/KismetPins/SGraphPinIntegerSlider.cpp b/Engine/Source/Editor/GraphEditor/Private/KismetPins/SGraphPinIntegerSlider.cpp
new file mode 100644
--- /dev/null
+++ b/Engine/Source/Editor/GraphEditor/Private/KismetPins/SGraphPinIntegerSlider.cpp
@@ -0,0 +1,189 @@
+// Copyright Epic Games, Inc. All Rights Reserved.
+
+#include "KismetPins/SGraphPinIntegerSlider.h"
+
+#include "Containers/Array.h"
+#include "Containers/UnrealString.h"
+#include "EdGraph/EdGraphPin.h"
+#include "EdGraph/EdGraphSchema.h"
+#include "EdGraphSchema_K2.h"
+#include "Fonts/SlateFontInfo.h"
+#include "Framework/Commands/UIAction.h"
+#include "Framework/Commands/UICommandInfo.h"
+#include "Framework/MultiBox/MultiBoxBuilder.h"
+#include "Internationalization/Internationalization.h"
+#include "Internationalization/Text.h"
+#include "Layout/Margin.h"
+#include "Math/NumericLimits.h"
+#include "Math/UnrealMathSSE.h"
+#include "Misc/AssertionMacros.h"
+#include "Misc/Attribute.h"
+#include "Misc/CString.h"
+#include "SGraphPin.h"
+#include "ScopedTransaction.h"
+#include "Styling/AppStyle.h"
+#include "Templates/Casts.h"
+#include "Textures/SlateIcon.h"
+#include "Types/SlateEnums.h"
+#include "Types/SlateStructs.h"
+#include "UObject/Class.h"
+#include "UObject/NameTypes.h"
+#include "UObject/Object.h"
+#include "UObject/UnrealNames.h"
+#include "UObject/WeakObjectPtrTemplates.h"
+#include "Widgets/Input/SComboButton.h"
+#include "Widgets/Layout/SBox.h"
+#include "Widgets/Text/STextBlock.h"
+
+class SWidget;
+
+void SGraphPinIntegerSlider::Construct(const FArguments& InArgs, UEdGraphPin* InGraphPinObj, FProperty* InProperty)
+{
+	SGraphPinNumSlider<int32>::Construct(SGraphPinNumSlider<int32>::FArguments(), InGraphPinObj, InProperty);
+}
+
+TSharedRef<SWidget>	SGraphPinIntegerSlider::GetDefaultValueWidget()
+{
+	check(GraphPinObj);
+
+	const FName PinSubCategory = GraphPinObj->PinType.PinSubCategory;
+	const UEnum* BitmaskEnum = Cast<UEnum>(GraphPinObj->PinType.PinSubCategoryObject.Get());
+
+	if (PinSubCategory == UEdGraphSchema_K2::PSC_Bitmask)
+	{
+		struct BitmaskFlagInfo
+		{
+			int32 Value;
+			FText DisplayName;
+			FText ToolTipText;
+		};
+
+		const int32 BitmaskBitCount = sizeof(int32) << 3;
+
+		TArray<BitmaskFlagInfo> BitmaskFlags;
+		BitmaskFlags.Reserve(BitmaskBitCount);
+
+		if(BitmaskEnum)
+		{
+			const bool bUseEnumValuesAsMaskValues = BitmaskEnum->GetBoolMetaData(FBlueprintMetadata::MD_UseEnumValuesAsMaskValuesInEditor);
+			auto AddNewBitmaskFlagLambda = [BitmaskEnum, &BitmaskFlags](int32 InEnumIndex, int32 InFlagValue)
+			{
+				BitmaskFlagInfo* BitmaskFlag = new(BitmaskFlags) BitmaskFlagInfo();
+
+				BitmaskFlag->Value = InFlagValue;
+				BitmaskFlag->DisplayName = BitmaskEnum->GetDisplayNameTextByIndex(InEnumIndex);
+				BitmaskFlag->ToolTipText = BitmaskEnum->GetToolTipTextByIndex(InEnumIndex);
+				if (BitmaskFlag->ToolTipText.IsEmpty())
+				{
+					BitmaskFlag->ToolTipText = FText::Format(NSLOCTEXT("GraphEditor", "BitmaskDefaultFlagToolTipText", "Toggle {0} on/off"), BitmaskFlag->DisplayName);
+				}
+			};
+
+			// Note: This loop doesn't include (BitflagsEnum->NumEnums() - 1) in order to skip the implicit "MAX" value that gets added to the enum type at compile time.
+			for(int32 BitmaskEnumIndex = 0; BitmaskEnumIndex < BitmaskEnum->NumEnums() - 1; ++BitmaskEnumIndex)
+			{
+				const int64 EnumValue = BitmaskEnum->GetValueByIndex(BitmaskEnumIndex);
+				const bool bIsHidden = BitmaskEnum->HasMetaData(TEXT("Hidden"), BitmaskEnumIndex);
+				if (EnumValue >= 0 && !bIsHidden)
+				{
+					if (bUseEnumValuesAsMaskValues)
+					{
+						if (EnumValue < MAX_int32 && FMath::IsPowerOfTwo(EnumValue))
+						{
+							AddNewBitmaskFlagLambda(BitmaskEnumIndex, static_cast<int32>(EnumValue));
+						}
+					}
+					else if (EnumValue < BitmaskBitCount)
+					{
+						AddNewBitmaskFlagLambda(BitmaskEnumIndex, 1 << static_cast<int32>(EnumValue));
+					}
+				}
+			}
+		}
+		else
+		{
+			for(int BitmaskFlagIndex = 0; BitmaskFlagIndex < BitmaskBitCount; ++BitmaskFlagIndex)
+			{
+				BitmaskFlagInfo* BitmaskFlag = new(BitmaskFlags) BitmaskFlagInfo();
+
+				BitmaskFlag->Value = static_cast<int32>(1 << BitmaskFlagIndex);
+				BitmaskFlag->DisplayName = FText::Format(NSLOCTEXT("GraphEditor", "BitmaskDefaultFlagDisplayName", "Flag {0}"), FText::AsNumber(BitmaskFlagIndex + 1));
+				BitmaskFlag->ToolTipText = FText::Format(NSLOCTEXT("GraphEditor", "BitmaskDefaultFlagToolTipText", "Toggle {0} on/off"), BitmaskFlag->DisplayName);
+			}
+		}
+
+		const auto& GetComboButtonText = [this, BitmaskFlags]() -> FText
+		{
+			int32 BitmaskValue = FCString::Atoi(*GraphPinObj->GetDefaultAsString());
+			if(BitmaskValue != 0)
+			{
+				if(BitmaskValue & (BitmaskValue - 1))
+				{
+					return NSLOCTEXT("GraphEditor", "BitmaskButtonContentMultipleBitsSet", "(Multiple)");
+				}
+				else
+				{
+					for(int i = 0; i < BitmaskFlags.Num(); ++i)
+					{
+						if(BitmaskValue & BitmaskFlags[i].Value)
+						{
+							return BitmaskFlags[i].DisplayName;
+						}
+					}
+				}
+			}
+
+			return NSLOCTEXT("GraphEditor", "BitmaskButtonContentNoFlagsSet", "(No Flags)");
+		};
+
+		return SNew(SComboButton)
+			.ContentPadding(3)
+			.MenuPlacement(MenuPlacement_BelowAnchor)
+			.Visibility(this, &SGraphPin::GetDefaultValueVisibility)
+			.ButtonContent()
+			[
+				// Wrap in configurable box to restrain height/width of menu
+				SNew(SBox)
+				.MinDesiredWidth(84.0f)
+				[
+					SNew(STextBlock)
+					.Text_Lambda(GetComboButtonText)
+					.Font(FAppStyle::GetFontStyle(TEXT("PropertyWindow.NormalFont")))
+				]
+			]
+			.OnGetMenuContent_Lambda([this, BitmaskFlags]()
+			{
+				FMenuBuilder MenuBuilder(false, nullptr);
+
+				for(int i = 0; i < BitmaskFlags.Num(); ++i)
+				{
+					MenuBuilder.AddMenuEntry(
+						BitmaskFlags[i].DisplayName,
+						BitmaskFlags[i].ToolTipText,
+						FSlateIcon(),
+						FUIAction
+						(
+							FExecuteAction::CreateLambda([this, BitmaskFlags, i]()
+							{
+								const FScopedTransaction Transaction(NSLOCTEXT("GraphEditor", "ChangePinValue", "Change Pin Value"));
+								const int32 CurValue = FCString::Atoi(*GraphPinObj->GetDefaultAsString());
+								GraphPinObj->Modify();
+								GraphPinObj->GetSchema()->TrySetDefaultValue(*GraphPinObj, FString::FromInt(CurValue ^ BitmaskFlags[i].Value));
+							}),
+							FCanExecuteAction(),
+							FIsActionChecked::CreateLambda([this, BitmaskFlags, i]() -> bool
+							{
+								const int32 CurValue = FCString::Atoi(*GraphPinObj->GetDefaultAsString());
+								return (CurValue & BitmaskFlags[i].Value) != 0;
+							})
+						),
+						NAME_None,
+						EUserInterfaceActionType::Check);
+				}
+
+				return MenuBuilder.MakeWidget();
+			});
+	}
+	
+	return SGraphPinNumSlider<int32>::GetDefaultValueWidget();
+}
diff --git a/Engine/Source/Editor/GraphEditor/Public/KismetPins/SGraphPinIntegerSlider.h b/Engine/Source/Editor/GraphEditor/Public/KismetPins/SGraphPinIntegerSlider.h
new file mode 100644
--- /dev/null
+++ b/Engine/Source/Editor/GraphEditor/Public/KismetPins/SGraphPinIntegerSlider.h
@@ -0,0 +1,31 @@
+// Copyright Epic Games, Inc. All Rights Reserved.
+
+#pragma once
+
+#include "Containers/UnrealString.h"
+#include "CoreMinimal.h"
+#include "Delegates/Delegate.h"
+#include "HAL/Platform.h"
+#include "KismetPins/SGraphPinNumSlider.h"
+#include "Misc/Optional.h"
+#include "Templates/SharedPointer.h"
+#include "Templates/UnrealTemplate.h"
+#include "Widgets/DeclarativeSyntaxSupport.h"
+#include "Widgets/SWidget.h"
+
+class SWidget;
+class UEdGraphPin;
+
+class GRAPHEDITOR_API SGraphPinIntegerSlider : public SGraphPinNumSlider<int32>
+{
+public:
+	SLATE_BEGIN_ARGS(SGraphPinIntegerSlider) {}
+	SLATE_END_ARGS()
+
+	void Construct(const FArguments& InArgs, UEdGraphPin* InGraphPinObj, FProperty* InProperty);
+
+protected:
+	//~ Begin SGraphPinString Interface
+	virtual TSharedRef<SWidget>	GetDefaultValueWidget() override;
+	//~ End SGraphPinString Interface
+};
diff --git a/Engine/Source/Editor/GraphEditor/Public/KismetPins/SGraphPinNumSlider.h b/Engine/Source/Editor/GraphEditor/Public/KismetPins/SGraphPinNumSlider.h
new file mode 100644
--- /dev/null
+++ b/Engine/Source/Editor/GraphEditor/Public/KismetPins/SGraphPinNumSlider.h
@@ -0,0 +1,135 @@
+// Copyright Epic Games, Inc. All Rights Reserved.
+
+#pragma once
+
+#include "CoreMinimal.h"
+#include "Widgets/DeclarativeSyntaxSupport.h"
+#include "Widgets/SWidget.h"
+#include "Widgets/Layout/SBox.h"
+#include "Widgets/Input/SNumericEntryBox.h"
+#include "Misc/DefaultValueHelper.h"
+#include "ScopedTransaction.h"
+#include "SGraphPin.h"
+#include "Editor.h"
+#include "NumericPropertyParams.h"
+
+template <typename NumericType>
+class SGraphPinNumSlider : public SGraphPin
+{
+public:
+	SLATE_BEGIN_ARGS(SGraphPinNumSlider) 
+	{}
+	SLATE_END_ARGS()
+
+	void Construct(const FArguments& InArgs, UEdGraphPin* InGraphPinObj, FProperty* InProperty)
+	{
+		PinProperty = InProperty;
+		SGraphPin::Construct(SGraphPin::FArguments(), InGraphPinObj);
+
+	}
+
+protected:
+	// SGraphPin interface
+	virtual TSharedRef<SWidget>	GetDefaultValueWidget() override
+	{
+		const typename TNumericPropertyParams<NumericType>::FMetaDataGetter MetaDataGetter = TNumericPropertyParams<NumericType>::FMetaDataGetter::CreateLambda([&](const FName& Key)
+			{
+				return (PinProperty) ? PinProperty->GetMetaData(Key) : FString();
+			});
+
+		TNumericPropertyParams<NumericType> NumericPropertyParams(PinProperty, PinProperty ? MetaDataGetter : nullptr);
+		
+		const bool bAllowSpin = !(PinProperty && PinProperty->GetBoolMetaData("NoSpinbox"));
+
+		return SNew(SBox)
+			.MinDesiredWidth(18)
+			.MaxDesiredWidth(400)
+			[
+				SNew(SNumericEntryBox<NumericType>)
+				.EditableTextBoxStyle(FAppStyle::Get(), "Graph.EditableTextBox")
+				.BorderForegroundColor(FSlateColor::UseForeground())
+				.Visibility(this, &SGraphPinNumSlider::GetDefaultValueVisibility)
+				.IsEnabled(this, &SGraphPinNumSlider::GetDefaultValueIsEditable)
+				.Value(this, &SGraphPinNumSlider::GetNumericValue)
+				.MinValue(NumericPropertyParams.MinValue)
+				.MaxValue(NumericPropertyParams.MaxValue)
+				.MinSliderValue(NumericPropertyParams.MinSliderValue)
+				.MaxSliderValue(NumericPropertyParams.MaxSliderValue)
+				.SliderExponent(NumericPropertyParams.SliderExponent)
+				.Delta(NumericPropertyParams.Delta)
+				.LinearDeltaSensitivity(NumericPropertyParams.GetLinearDeltaSensitivityAttribute())
+				.AllowWheel(bAllowSpin)
+				.WheelStep(NumericPropertyParams.WheelStep)
+				.AllowSpin(bAllowSpin)
+				.OnValueCommitted(this, &SGraphPinNumSlider::OnValueCommitted)
+				.OnValueChanged(this, &SGraphPinNumSlider::OnValueChanged)
+				.OnBeginSliderMovement(this, &SGraphPinNumSlider::OnBeginSliderMovement)
+				.OnEndSliderMovement(this, &SGraphPinNumSlider::OnEndSliderMovement)
+			];
+	}
+
+	void OnValueChanged( NumericType NewValue )
+	{
+		SliderValue = NewValue;
+	}
+
+	void OnValueCommitted( NumericType NewValue, ETextCommit::Type CommitInfo )
+	{
+		if (GraphPinObj->IsPendingKill())
+		{
+			return;
+		}
+
+		if (LastSliderCommittedValue != NewValue)
+		{
+			const FScopedTransaction Transaction(NSLOCTEXT("GraphEditor", "ChangeNumberPinValue", "Change Number Pin Value"));
+			GraphPinObj->Modify();
+			GraphPinObj->GetSchema()->TrySetDefaultValue(*GraphPinObj, *LexToString(NewValue));
+			LastSliderCommittedValue = NewValue;
+		}
+	}
+
+	/**
+	* Called when the slider begins to move.
+	*/
+	void OnBeginSliderMovement()
+	{
+		SliderValue = LastSliderCommittedValue = GetNumericValue().GetValue();
+		bIsUsingSlider = true;
+	}
+
+	/**
+	* Called when the slider stops moving.
+	*/
+	void OnEndSliderMovement( NumericType NewValue )
+	{
+		bIsUsingSlider = false;
+	}
+	
+	TOptional<NumericType> GetNumericValue() const
+	{
+		NumericType Num = NumericType();
+		LexFromString(Num, *GraphPinObj->GetDefaultAsString());
+		return bIsUsingSlider ? SliderValue : Num;
+	}
+
+public:
+	// virtual EVisibility GetDefaultValueVisibility() const override
+	// {
+	// 	UEdGraphPin* GraphPin = GetPinObj();
+	// 	if (GraphPin->Direction == EGPD_Output)
+	// 	{
+	// 		return EVisibility::Collapsed;
+	// 	}
+	// 	else
+	// 	{
+	// 		return EVisibility::Visible;
+	// 	}
+	// }
+
+private:
+	FProperty* PinProperty;
+	NumericType LastSliderCommittedValue;
+	NumericType SliderValue;
+	bool bIsUsingSlider = false;
+};
diff --git a/Engine/Source/Editor/GraphEditor/Public/KismetPins/SGraphPinVector2DSlider.h b/Engine/Source/Editor/GraphEditor/Public/KismetPins/SGraphPinVector2DSlider.h
new file mode 100644
--- /dev/null
+++ b/Engine/Source/Editor/GraphEditor/Public/KismetPins/SGraphPinVector2DSlider.h
@@ -0,0 +1,151 @@
+// Copyright Epic Games, Inc. All Rights Reserved.
+
+#pragma once
+
+#include "CoreMinimal.h"
+#include "Widgets/DeclarativeSyntaxSupport.h"
+#include "Widgets/SWidget.h"
+#include "Widgets/Layout/SBox.h"
+#include "Widgets/Input/SNumericEntryBox.h"
+#include "Misc/DefaultValueHelper.h"
+#include "ScopedTransaction.h"
+#include "SGraphPin.h"
+//#include "TG_Pin.h"
+//#include "EdGraph/TG_EdGraphNode.h"
+#include "Editor.h"
+//#include "TG_Expression.h"
+//#include "TG_Node.h"
+#include "NumericPropertyParams.h"
+#include "SVector2DSlider.h"
+
+template <typename NumericType>
+class SGraphPinVector2DSlider : public SGraphPin
+{
+	FProperty* PinProperty;
+public:
+	SLATE_BEGIN_ARGS(SGraphPinVector2DSlider)
+	{}
+	SLATE_END_ARGS()
+
+	void Construct(const FArguments& InArgs, UEdGraphPin* InGraphPinObj, FProperty* InProperty)
+	{
+		PinProperty = InProperty;
+		SGraphPin::Construct(SGraphPin::FArguments(), InGraphPinObj);
+
+	}
+
+protected:
+	/**
+	 *	Function to create class specific widget.
+	 *
+	 *	@return Reference to the newly created widget object
+	 */
+	virtual TSharedRef<SWidget>	GetDefaultValueWidget() override
+	{
+		// Create widget
+		return SNew(SVector2DSlider<NumericType>, PinProperty)
+			.VisibleText_X(this, &SGraphPinVector2DSlider::GetCurrentValue_X)
+			.VisibleText_Y(this, &SGraphPinVector2DSlider::GetCurrentValue_Y)
+			.Visibility(this, &SGraphPin::GetDefaultValueVisibility)
+			.IsEnabled(this, &SGraphPin::GetDefaultValueIsEditable)
+			.OnNumericCommitted_Box_X(this, &SGraphPinVector2DSlider::OnChangedValueTextBox_X)
+			.OnNumericCommitted_Box_Y(this, &SGraphPinVector2DSlider::OnChangedValueTextBox_Y);
+	}
+
+private:
+
+	// Enum values to identify text boxes.
+	enum ETextBoxIndex
+	{
+		TextBox_X,
+		TextBox_Y
+	};
+
+	/*
+	 *	Function to get current value in text box 0
+	 *
+	 *	@return current string value
+	 */
+	FString GetCurrentValue_X() const
+	{
+		return FString::Printf(TEXT("%f"), GetValue().X);
+	}
+
+	/*
+	 *	Function to get current value in text box 1
+	 *
+	 *	@return current string value
+	 */
+	FString GetCurrentValue_Y() const
+	{
+		return FString::Printf(TEXT("%f"), GetValue().Y);
+	}
+
+	/*
+	 *	Function to getch current value based on text box index value
+	 *
+	 *	@param: Text box index
+	 *
+	 *	@return current value
+	 */
+	FVector2D GetValue() const
+	{
+		const FString& DefaultString = GraphPinObj->GetDefaultAsString();
+
+		FVector2D Value;
+		Value.InitFromString(DefaultString);
+
+		return Value;
+	}
+
+	static FString MakeVector2DString(NumericType X, NumericType Y)
+	{
+		return FString::Format(TEXT("(X={0},Y={1})"), { X, Y });
+	}
+
+	/*
+	 *	Function to store value when text box 0 value in modified
+	 *
+	 *	@param 0: Updated numeric value
+	 */
+	void OnChangedValueTextBox_X(NumericType NewValue, ETextCommit::Type CommitInfo)
+	{
+		if (GraphPinObj->IsPendingKill())
+		{
+			return;
+		}
+
+		const FVector2D OldValue = GetValue();
+		if (NewValue != OldValue.X)
+		{
+			const FScopedTransaction Transaction(NSLOCTEXT("GraphEditor", "ChangeVectorPinValue", "Change Vector Pin Value"));
+			GraphPinObj->Modify();
+
+			const FString Vector2DString = MakeVector2DString(NewValue, OldValue.Y);
+			GraphPinObj->GetSchema()->TrySetDefaultValue(*GraphPinObj, Vector2DString);
+		}
+	}
+
+	/*
+	 *	Function to store value when text box 1 value in modified
+	 *
+	 *	@param 0: Updated numeric value
+	 */
+	void OnChangedValueTextBox_Y(NumericType NewValue, ETextCommit::Type CommitInfo)
+	{
+		if (GraphPinObj->IsPendingKill())
+		{
+			return;
+		}
+
+		const FVector2D OldValue = GetValue();
+		if (NewValue != OldValue.Y)
+		{
+			const FScopedTransaction Transaction(NSLOCTEXT("GraphEditor", "ChangeVectorPinValue", "Change Vector Pin Value"));
+			GraphPinObj->Modify();
+
+			const FString Vector2DString = MakeVector2DString(OldValue.X, NewValue);
+			GraphPinObj->GetSchema()->TrySetDefaultValue(*GraphPinObj, Vector2DString);
+		}
+	}
+};
diff --git a/Engine/Source/Editor/GraphEditor/Public/KismetPins/SGraphPinVector4Slider.h b/Engine/Source/Editor/GraphEditor/Public/KismetPins/SGraphPinVector4Slider.h
new file mode 100644
--- /dev/null
+++ b/Engine/Source/Editor/GraphEditor/Public/KismetPins/SGraphPinVector4Slider.h
@@ -0,0 +1,200 @@
+// Copyright Epic Games, Inc. All Rights Reserved.
+
+#pragma once
+
+#include "Containers/UnrealString.h"
+#include "CoreMinimal.h"
+#include "Math/Vector4.h"
+#include "SGraphPin.h"
+#include "Templates/SharedPointer.h"
+#include "Types/SlateEnums.h"
+#include "Widgets/DeclarativeSyntaxSupport.h"
+#include "Widgets/SWidget.h"
+#include "SVector4Slider.h"
+#include "ScopedTransaction.h"
+
+class SWidget;
+class UEdGraphPin;
+
+template <typename NumericType>
+class SGraphPinVector4Slider : public SGraphPin
+{
+public:
+	SLATE_BEGIN_ARGS(SGraphPinVector4Slider) {}
+	SLATE_END_ARGS()
+
+	void Construct(const FArguments& InArgs, UEdGraphPin* InGraphPinObj, FProperty* InProperty)
+	{
+		PinProperty = InProperty;
+		SGraphPin::Construct(SGraphPin::FArguments(), InGraphPinObj);
+	}
+
+protected:
+
+	/**
+	 *	Function to create class specific widget.
+	 *
+	 *	@return Reference to the newly created widget object
+	 */
+	virtual TSharedRef<SWidget>	GetDefaultValueWidget() override
+	{
+		UScriptStruct* RotatorStruct = TBaseStructure<FRotator>::Get();
+
+		return SNew(SVector4Slider<NumericType>, PinProperty)
+			.VisibleText_0(this, &SGraphPinVector4Slider::GetCurrentValue_0)
+			.VisibleText_1(this, &SGraphPinVector4Slider::GetCurrentValue_1)
+			.VisibleText_2(this, &SGraphPinVector4Slider::GetCurrentValue_2)
+			.VisibleText_3(this, &SGraphPinVector4Slider::GetCurrentValue_3)
+			.Visibility(this, &SGraphPin::GetDefaultValueVisibility)
+			.IsEnabled(this, &SGraphPin::GetDefaultValueIsEditable)
+			.OnNumericCommitted_Box_0(this, &SGraphPinVector4Slider::OnChangedValueTextBox_0)
+			.OnNumericCommitted_Box_1(this, &SGraphPinVector4Slider::OnChangedValueTextBox_1)
+			.OnNumericCommitted_Box_2(this, &SGraphPinVector4Slider::OnChangedValueTextBox_2)
+			.OnNumericCommitted_Box_3(this, &SGraphPinVector4Slider::OnChangedValueTextBox_3);
+	}
+
+private:
+
+	// Enum values to identify text boxes.
+	enum ETextBoxIndex
+	{
+		TextBox_0,
+		TextBox_1,
+		TextBox_2,
+		TextBox_3,
+	};
+
+	// Rotator is represented as X->Roll, Y->Pitch, Z->Yaw
+
+	FString GetCurrentValue_0() const
+	{
+		// Text box 0: Rotator->Roll, Vector->X
+		return GetValue(TextBox_0);
+	}
+
+	FString GetCurrentValue_1() const
+	{
+		// Text box 1: Rotator->Pitch, Vector->Y
+		return GetValue(TextBox_1);
+	}
+
+	FString GetCurrentValue_2() const
+	{
+		// Text box 2: Rotator->Yaw, Vector->Z
+		return GetValue(TextBox_2);
+	}
+
+	FString GetCurrentValue_3() const
+	{
+		// Text box 3: Vector->W
+		return GetValue(TextBox_3);
+	}
+
+	/*
+	*	Function to getch current value based on text box index value
+	*
+	*	@param: Text box index
+	*
+	*	@return current string value
+	*/
+	FString GetValue(ETextBoxIndex Index) const
+	{
+		FString DefaultString = GraphPinObj->GetDefaultAsString();
+		TArray<FString> ResultString;
+
+		// Parse string to split its contents separated by ','
+		DefaultString.TrimStartAndEndInline();
+		DefaultString.ParseIntoArray(ResultString, TEXT(","), true);
+
+		if (Index < ResultString.Num())
+		{
+			return ResultString[Index];
+		}
+		else
+		{
+			return FString(TEXT("0"));
+		}
+	}
+
+	/*
+	*	Function to store value when text box value in modified
+	*
+	*	@param 0: Updated numeric value
+	*/
+	void OnChangedValueTextBox_0(NumericType NewValue, ETextCommit::Type CommitInfo)
+	{
+		const FString ValueStr = FString::Printf(TEXT("%f"), NewValue);
+
+		FString DefaultValue;
+		// Update X value
+		DefaultValue = ValueStr + FString(TEXT(",")) + GetValue(TextBox_1) + FString(TEXT(",")) + GetValue(TextBox_2) + FString(TEXT(",")) + GetValue(TextBox_3);
+
+		if (GraphPinObj->GetDefaultAsString() != DefaultValue)
+		{
+			const FScopedTransaction Transaction(NSLOCTEXT("GraphEditor", "ChangeVector4PinValue", "Change Vector4 Pin Value"));
+			GraphPinObj->Modify();
+
+			// Set new default value
+			GraphPinObj->GetSchema()->TrySetDefaultValue(*GraphPinObj, DefaultValue);
+		}
+	}
+
+	void OnChangedValueTextBox_1(NumericType NewValue, ETextCommit::Type CommitInfo)
+	{
+		const FString ValueStr = FString::Printf(TEXT("%f"), NewValue);
+
+		FString DefaultValue;
+		// Update Y value
+		DefaultValue = GetValue(TextBox_0) + FString(TEXT(",")) + ValueStr + FString(TEXT(",")) + GetValue(TextBox_2) + FString(TEXT(",")) + GetValue(TextBox_3);
+
+		if (GraphPinObj->GetDefaultAsString() != DefaultValue)
+		{
+			const FScopedTransaction Transaction(NSLOCTEXT("GraphEditor", "ChangeVector4PinValue", "Change Vector4 Pin Value"));
+			GraphPinObj->Modify();
+
+			// Set new default value
+			GraphPinObj->GetSchema()->TrySetDefaultValue(*GraphPinObj, DefaultValue);
+		}
+	}
+
+	void OnChangedValueTextBox_2(NumericType NewValue, ETextCommit::Type CommitInfo)
+	{
+		const FString ValueStr = FString::Printf(TEXT("%f"), NewValue);
+
+		FString DefaultValue;
+		// Update Z value
+		DefaultValue = GetValue(TextBox_0) + FString(TEXT(",")) + GetValue(TextBox_1) + FString(TEXT(",")) + ValueStr + FString(TEXT(",")) + GetValue(TextBox_3);
+
+		if (GraphPinObj->GetDefaultAsString() != DefaultValue)
+		{
+			const FScopedTransaction Transaction(NSLOCTEXT("GraphEditor", "ChangeVector4PinValue", "Change Vector4 Pin Value"));
+			GraphPinObj->Modify();
+
+			// Set new default value
+			GraphPinObj->GetSchema()->TrySetDefaultValue(*GraphPinObj, DefaultValue);
+		}
+	}
+
+	void OnChangedValueTextBox_3(NumericType NewValue, ETextCommit::Type CommitInfo)
+	{
+		const FString ValueStr = FString::Printf(TEXT("%f"), NewValue);
+
+		FString DefaultValue;
+		// Update W value
+		DefaultValue = GetValue(TextBox_0) + FString(TEXT(",")) + GetValue(TextBox_1) + FString(TEXT(",")) + GetValue(TextBox_2) + FString(TEXT(",")) + ValueStr;
+
+		if (GraphPinObj->GetDefaultAsString() != DefaultValue)
+		{
+			const FScopedTransaction Transaction(NSLOCTEXT("GraphEditor", "ChangeVector4PinValue", "Change Vector4 Pin Value"));
+			GraphPinObj->Modify();
+
+			// Set new default value
+			GraphPinObj->GetSchema()->TrySetDefaultValue(*GraphPinObj, DefaultValue);
+		}
+	}
+
+private:
+
+	// Property to get metadata values out of for the slider widget
+	FProperty* PinProperty;
+};
diff --git a/Engine/Source/Editor/GraphEditor/Public/KismetPins/SGraphPinVectorSlider.h b/Engine/Source/Editor/GraphEditor/Public/KismetPins/SGraphPinVectorSlider.h
new file mode 100644
--- /dev/null
+++ b/Engine/Source/Editor/GraphEditor/Public/KismetPins/SGraphPinVectorSlider.h
@@ -0,0 +1,236 @@
+// Copyright Epic Games, Inc. All Rights Reserved.
+
+#pragma once
+
+#include "CoreMinimal.h"
+#include "Widgets/DeclarativeSyntaxSupport.h"
+#include "Widgets/SWidget.h"
+#include "Widgets/Layout/SBox.h"
+#include "Widgets/Input/SNumericEntryBox.h"
+#include "Misc/DefaultValueHelper.h"
+#include "ScopedTransaction.h"
+#include "SGraphPin.h"
+//#include "TG_Pin.h"
+//#include "EdGraph/TG_EdGraphNode.h"
+#include "Editor.h"
+//#include "TG_Expression.h"
+//#include "TG_Node.h"
+#include "SVectorSlider.h"
+
+template <typename NumericType>
+class SGraphPinVectorSlider : public SGraphPin
+{
+public:
+	SLATE_BEGIN_ARGS(SGraphPinVectorSlider)
+	{}
+	SLATE_END_ARGS()
+
+	void Construct(const FArguments& InArgs, UEdGraphPin* InGraphPinObj, FProperty* InProperty)
+	{
+		PinProperty = InProperty;
+		SGraphPin::Construct(SGraphPin::FArguments(), InGraphPinObj);
+
+	}
+
+protected:
+	/**
+	 *	Function to create class specific widget.
+	 *
+	 *	@return Reference to the newly created widget object
+	 */
+	virtual TSharedRef<SWidget>	GetDefaultValueWidget() override
+	{
+		UScriptStruct* RotatorStruct = TBaseStructure<FRotator>::Get();
+		bIsRotator = (GraphPinObj->PinType.PinSubCategoryObject == RotatorStruct);
+
+		// Create widget
+		return SNew(SVectorSlider<NumericType>, bIsRotator, PinProperty)
+			.VisibleText_0(this, &SGraphPinVectorSlider::GetCurrentValue_0)
+			.VisibleText_1(this, &SGraphPinVectorSlider::GetCurrentValue_1)
+			.VisibleText_2(this, &SGraphPinVectorSlider::GetCurrentValue_2)
+			.Visibility(this, &SGraphPin::GetDefaultValueVisibility)
+			.IsEnabled(this, &SGraphPin::GetDefaultValueIsEditable)
+			.OnNumericCommitted_Box_0(this, &SGraphPinVectorSlider::OnChangedValueTextBox_0)
+			.OnNumericCommitted_Box_1(this, &SGraphPinVectorSlider::OnChangedValueTextBox_1)
+			.OnNumericCommitted_Box_2(this, &SGraphPinVectorSlider::OnChangedValueTextBox_2);
+	}
+
+private:
+
+	// Enum values to identify text boxes.
+	enum ETextBoxIndex
+	{
+		TextBox_0,
+		TextBox_1,
+		TextBox_2,
+	};
+
+	using FVectorType = UE::Math::TVector<NumericType>;
+
+	// Rotator is represented as X->Roll, Y->Pitch, Z->Yaw
+
+	/*
+	 *	Function to get current value in text box 0
+	 *
+	 *	@return current string value
+	 */
+	FString GetCurrentValue_0() const
+	{
+		// Text box 0: Rotator->Roll, Vector->X
+		return GetValue(bIsRotator ? TextBox_2 : TextBox_0);
+	}
+
+	/*
+	 *	Function to get current value in text box 1
+	 *
+	 *	@return current string value
+	 */
+	FString GetCurrentValue_1() const
+	{
+		// Text box 1: Rotator->Pitch, Vector->Y
+		return GetValue(bIsRotator ? TextBox_0 : TextBox_1);
+	}
+
+	/*
+	 *	Function to get current value in text box 2
+	 *
+	 *	@return current string value
+	 */
+	FString GetCurrentValue_2() const
+	{
+		// Text box 2: Rotator->Yaw, Vector->Z
+		return GetValue(bIsRotator ? TextBox_1 : TextBox_2);
+	}
+
+	/*
+	 *	Function to get a string array of the components composing the default string.
+	 *
+	 *	@return The array filled with each component as string
+	 *
+	 */
+	TArray<FString> GetComponentArray() const
+	{
+		TArray<FString> VecComponentStrings;
+
+		FString DefaultString = GraphPinObj->GetDefaultAsString();
+		// Parse string to split its contents separated by ','
+		DefaultString.TrimStartInline();
+		DefaultString.TrimEndInline();
+		DefaultString.ParseIntoArray(VecComponentStrings, TEXT(","), true);
+
+		return VecComponentStrings;
+	}
+
+	/*
+	 *	Function to get current value based on text box index value
+	 *
+	 *	@param: Text box index
+	 *
+	 *	@return current string value
+	 */
+	FString GetValue(ETextBoxIndex Index) const
+	{
+		const TArray<FString> VecComponentStrings = GetComponentArray();
+
+		if (Index < VecComponentStrings.Num())
+		{
+			return VecComponentStrings[Index];
+		}
+		else
+		{
+			return FString(TEXT("0"));
+		}
+	}
+
+	/*
+	 *	Function to store value when text box 0 value in modified
+	 *
+	 *	@param 0: Updated numeric Value
+	 */
+	void OnChangedValueTextBox_0(NumericType NewValue, ETextCommit::Type CommitInfo)
+	{
+		const EAxis::Type Axis = bIsRotator ? /* update roll */ EAxis::Z : EAxis::X;
+		SetNewValueHelper(Axis, NewValue);
+	}
+
+	/*
+	 *	Function to store value when text box 1 value in modified
+	 *
+	 *	@param 0: Updated numeric Value
+	 */
+	void OnChangedValueTextBox_1(NumericType NewValue, ETextCommit::Type CommitInfo)
+	{
+		const EAxis::Type Axis = bIsRotator ? /* update pitch */ EAxis::X : EAxis::Y;
+		SetNewValueHelper(Axis, NewValue);
+	}
+
+	/*
+	 *	Function to store value when text box 2 value in modified
+	 *
+	 *	@param 0: Updated numeric Value
+	 */
+	void OnChangedValueTextBox_2(NumericType NewValue, ETextCommit::Type CommitInfo)
+	{
+		const EAxis::Type Axis = bIsRotator ? /* update yaw */ EAxis::Y : EAxis::Z;
+		SetNewValueHelper(Axis, NewValue);
+	}
+
+	void SetNewValueHelper(EAxis::Type Axis, NumericType NewValue)
+	{
+		if (GraphPinObj->IsPendingKill())
+		{
+			return;
+		}
+
+		FVectorType NewVector = ConvertDefaultValueStringToVector();
+		const NumericType OldValue = NewVector.GetComponentForAxis(Axis);
+
+		if (OldValue == NewValue)
+		{
+			return;
+		}
+
+		NewVector.SetComponentForAxis(Axis, NewValue);
+
+		const FScopedTransaction Transaction(NSLOCTEXT("GraphEditor", "ChangeVectorPinValue", "Change Vector Pin Value"));
+		GraphPinObj->Modify();
+
+		// Create the new value string
+		FString DefaultValue = FString::Format(TEXT("{0},{1},{2}"), { NewVector.X, NewVector.Y, NewVector.Z });
+
+		//Set new default value
+		GraphPinObj->GetSchema()->TrySetDefaultValue(*GraphPinObj, DefaultValue);
+	}
+
+	/*
+	 * @brief Converts the default string value to a value of VectorType.
+	 *
+	 * Example: it converts the default string value "2.00,3.00,4.00" to the corresponding 3D vector.
+	 */
+	FVectorType ConvertDefaultValueStringToVector() const
+	{
+		const TArray<FString> VecComponentStrings = GetComponentArray();
+
+		// Construct the vector from the string parts.
+		FVectorType Vec = FVectorType::ZeroVector;
+		TDefaultNumericTypeInterface<NumericType> NumericTypeInterface{};
+
+		// If default string value contained a fully specified 3D vector, set the vector components, otherwise leave it zero'ed.
+		if (VecComponentStrings.Num() == 3)
+		{
+			Vec.X = NumericTypeInterface.FromString(VecComponentStrings[0], 0).Get(0);
+			Vec.Y = NumericTypeInterface.FromString(VecComponentStrings[1], 0).Get(0);
+			Vec.Z = NumericTypeInterface.FromString(VecComponentStrings[2], 0).Get(0);
+		}
+
+		return Vec;
+	}
+
+private:
+	// Flag is true if the widget is used to represent a rotator; false otherwise
+	bool bIsRotator;
+
+	// Property to get metadata values out of for the slider widget
+	FProperty* PinProperty;
+		
+};
diff --git a/Engine/Source/Editor/GraphEditor/Public/KismetPins/SVector2DSlider.h b/Engine/Source/Editor/GraphEditor/Public/KismetPins/SVector2DSlider.h
new file mode 100644
--- /dev/null
+++ b/Engine/Source/Editor/GraphEditor/Public/KismetPins/SVector2DSlider.h
@@ -0,0 +1,236 @@
+// Copyright Epic Games, Inc. All Rights Reserved.
+
+#pragma once
+
+#include "CoreMinimal.h"
+#include "Widgets/Input/SNumericEntryBox.h"
+#include "Widgets/SBoxPanel.h"
+#include "Widgets/SCompoundWidget.h"
+#include "NumericPropertyParams.h"
+#include "Editor.h"
+
+#define LOCTEXT_NAMESPACE "Vector2DTextBox"
+
+// Class implementation to create 2 editable text boxes to represent vector2D graph pin
+template <typename NumericType>
+class SVector2DSlider : public SCompoundWidget
+{
+public:
+
+	// Notification for numeric value committed
+	DECLARE_DELEGATE_TwoParams(FOnNumericValueCommitted, NumericType, ETextCommit::Type);
+
+	SLATE_BEGIN_ARGS( SVector2DSlider ){}	
+		SLATE_ATTRIBUTE( FString, VisibleText_X )
+		SLATE_ATTRIBUTE( FString, VisibleText_Y )
+		SLATE_EVENT( FOnNumericValueCommitted, OnNumericCommitted_Box_X )
+		SLATE_EVENT( FOnNumericValueCommitted, OnNumericCommitted_Box_Y )
+	SLATE_END_ARGS()
+
+	// Construct editable text boxes with the appropriate getter & setter functions along with tool tip text
+	void Construct( const FArguments& InArgs, FProperty* InProperty)
+	{
+		PinProperty = InProperty;
+		VisibleText_X = InArgs._VisibleText_X;
+		VisibleText_Y = InArgs._VisibleText_Y;
+		OnNumericCommitted_X = InArgs._OnNumericCommitted_Box_X;
+		OnNumericCommitted_Y = InArgs._OnNumericCommitted_Box_Y;
+
+		const FLinearColor LabelClr = FLinearColor( 1.f, 1.f, 1.f, 0.4f );
+
+		const typename TNumericPropertyParams<NumericType>::FMetaDataGetter MetaDataGetter = TNumericPropertyParams<NumericType>::FMetaDataGetter::CreateLambda([&](const FName& Key)
+			{
+				return PinProperty->GetMetaData(Key);
+			});
+
+		TNumericPropertyParams<NumericType> NumericPropertyParams(PinProperty, PinProperty ? MetaDataGetter : nullptr);
+
+		const bool bAllowSpin = !(PinProperty && PinProperty->GetBoolMetaData("NoSpinbox"));
+
+		this->ChildSlot
+		[
+			SNew(SVerticalBox)
+			+SVerticalBox::Slot()
+			.AutoHeight()
+			.Padding(0)
+			[
+				SNew(SHorizontalBox)
+				+SHorizontalBox::Slot()
+				.AutoWidth().Padding(2) .HAlign(HAlign_Fill)
+				[
+					// Create Text box 0 
+					SNew( SNumericEntryBox<NumericType> )
+					.LabelVAlign(VAlign_Center)
+					.Label()
+					[
+						SNew( STextBlock )
+						.Font( FAppStyle::GetFontStyle( "Graph.VectorEditableTextBox" ) )
+						.Text( LOCTEXT("VectorNodeXAxisValueLabel", "X") )
+						.ColorAndOpacity( LabelClr )
+					]
+					.MinValue(NumericPropertyParams.MinValue)
+					.MaxValue(NumericPropertyParams.MaxValue)
+					.MinSliderValue(NumericPropertyParams.MinSliderValue)
+					.MaxSliderValue(NumericPropertyParams.MaxSliderValue)
+					.SliderExponent(NumericPropertyParams.SliderExponent)
+					.Delta(NumericPropertyParams.Delta)
+					.Value( this, &SVector2DSlider::GetTypeInValue_X )
+					.LinearDeltaSensitivity(NumericPropertyParams.GetLinearDeltaSensitivityAttribute())
+					.AllowWheel(bAllowSpin)
+					.WheelStep(NumericPropertyParams.WheelStep)
+					.AllowSpin(bAllowSpin)
+					.OnValueCommitted( InArgs._OnNumericCommitted_Box_X )
+					.OnValueChanged(this, &SVector2DSlider::OnValueChanged_X)
+					.OnBeginSliderMovement(this, &SVector2DSlider::OnBeginSliderMovement_X)
+					.OnEndSliderMovement(this, &SVector2DSlider::OnEndSliderMovement_X)
+					.Font( FAppStyle::GetFontStyle( "Graph.VectorEditableTextBox" ) )
+					.UndeterminedString( LOCTEXT("MultipleValues", "Multiple Values") )
+					.ToolTipText(LOCTEXT("VectorNodeXAxisValueLabel_ToolTip", "X value") )
+					.EditableTextBoxStyle( &FAppStyle::GetWidgetStyle<FEditableTextBoxStyle>( "Graph.VectorEditableTextBox" ))
+					.BorderForegroundColor( FLinearColor::White )
+					.BorderBackgroundColor( FLinearColor::White )
+				]
+				+SHorizontalBox::Slot()
+				.AutoWidth().Padding(2) .HAlign(HAlign_Fill)
+				[
+					// Create Text box 1
+					SNew( SNumericEntryBox<NumericType> )
+					.LabelVAlign(VAlign_Center)
+					.Label()
+					[
+						SNew( STextBlock )
+						.Font( FAppStyle::GetFontStyle( "Graph.VectorEditableTextBox" ) )
+						.Text( LOCTEXT("VectorNodeYAxisValueLabel", "Y") )
+						.ColorAndOpacity( LabelClr )
+					]
+					.Value( this, &SVector2DSlider::GetTypeInValue_Y )
+					.LinearDeltaSensitivity(NumericPropertyParams.GetLinearDeltaSensitivityAttribute())
+					.AllowWheel(bAllowSpin)
+					.WheelStep(NumericPropertyParams.WheelStep)
+					.AllowSpin(bAllowSpin)
+					.OnValueCommitted( InArgs._OnNumericCommitted_Box_Y )
+					.OnValueChanged(this, &SVector2DSlider::OnValueChanged_Y)
+					.OnBeginSliderMovement(this, &SVector2DSlider::OnBeginSliderMovement_Y)
+					.OnEndSliderMovement(this, &SVector2DSlider::OnEndSliderMovement_Y)
+					.Font( FAppStyle::GetFontStyle( "Graph.VectorEditableTextBox" ) )
+					.UndeterminedString( LOCTEXT("MultipleValues", "Multiple Values") )
+					.ToolTipText(LOCTEXT("VectorNodeYAxisValueLabel_ToolTip", "Y value"))
+					.EditableTextBoxStyle( &FAppStyle::GetWidgetStyle<FEditableTextBoxStyle>( "Graph.VectorEditableTextBox" ))
+					.BorderForegroundColor( FLinearColor::White )
+					.BorderBackgroundColor( FLinearColor::White )
+				]
+			]
+		];
+	}
+
+private:
+
+	/**
+	* Called when the slider value changes
+	*/
+	void OnValueChanged_X(NumericType NewValue)
+	{
+		SliderValue_X = NewValue;
+	}
+
+	/**
+	* Called when the slider value changes
+	*/
+	void OnValueChanged_Y(NumericType NewValue)
+	{
+		SliderValue_Y = NewValue;
+	}
+
+	/**
+	* Called when the slider begins to move.
+	*/
+	void OnBeginSliderMovement_X()
+	{
+		SyncValues();
+		bIsUsingSlider = true;
+	}
+
+	/**
+	* Called when the slider stops moving.
+	*/
+	void OnEndSliderMovement_X(NumericType NewValue)
+	{
+		bIsUsingSlider = false;
+		if (LastSliderCommittedValue_X != NewValue)
+		{
+			OnNumericCommitted_X.ExecuteIfBound(NewValue, ETextCommit::Default);
+			LastSliderCommittedValue_X = NewValue;
+		}
+	}
+
+	/**
+	* Called when the slider begins to move.
+	*/
+	void OnBeginSliderMovement_Y()
+	{
+		SyncValues();
+		bIsUsingSlider = true;
+	}
+
+	/**
+	* Called when the slider stops moving.
+	*/
+	void OnEndSliderMovement_Y(NumericType NewValue)
+	{
+		bIsUsingSlider = false;
+		if (LastSliderCommittedValue_Y != NewValue)
+		{
+			OnNumericCommitted_Y.ExecuteIfBound(NewValue, ETextCommit::Default);
+			LastSliderCommittedValue_Y = NewValue;
+		}
+	}
+
+	void SyncValues()
+	{
+		SliderValue_X = LastSliderCommittedValue_X = GetTypeInValue_X().GetValue();
+		SliderValue_Y = LastSliderCommittedValue_Y = GetTypeInValue_Y().GetValue();
+	}
+
+	NumericType GetValueType(const FString& InString) const
+	{
+		static_assert(std::is_floating_point_v<NumericType>);
+
+		if constexpr (std::is_same_v<float, NumericType>)
+		{
+			return FCString::Atof(*InString);
+		}
+		else if constexpr (std::is_same_v<double, NumericType>)
+		{
+			return FCString::Atod(*InString);
+		}
+
+		return NumericType{};
+	}
+
+	// Get value for X text box
+	TOptional<NumericType> GetTypeInValue_X() const
+	{
+		return GetValueType(VisibleText_X.Get());
+	}
+
+	// Get value for Y text box
+	TOptional<NumericType> GetTypeInValue_Y() const
+	{
+		return GetValueType(VisibleText_Y.Get());
+	}
+
+	TAttribute<FString> VisibleText_X;
+	TAttribute<FString> VisibleText_Y;
+
+	FProperty* PinProperty;
+	FOnNumericValueCommitted OnNumericCommitted_X;
+	FOnNumericValueCommitted OnNumericCommitted_Y;
+
+	NumericType LastSliderCommittedValue_X;
+	NumericType LastSliderCommittedValue_Y;
+	NumericType SliderValue_X;
+	NumericType SliderValue_Y;
+	bool bIsUsingSlider = false;
+};
+
+#undef LOCTEXT_NAMESPACE
diff --git a/Engine/Source/Editor/GraphEditor/Public/KismetPins/SVector4Slider.h b/Engine/Source/Editor/GraphEditor/Public/KismetPins/SVector4Slider.h
new file mode 100644
--- /dev/null
+++ b/Engine/Source/Editor/GraphEditor/Public/KismetPins/SVector4Slider.h
@@ -0,0 +1,397 @@
+// Copyright Epic Games, Inc. All Rights Reserved.
+
+#pragma once
+
+#include "CoreMinimal.h"
+#include "Widgets/Input/SNumericEntryBox.h"
+#include "Widgets/SBoxPanel.h"
+#include "Widgets/SCompoundWidget.h"
+#include "NumericPropertyParams.h"
+#include "Editor.h"
+
+#define LOCTEXT_NAMESPACE "VectorTextBox"
+
+// Class implementation to create 3 editable text boxes to represent vector/rotator graph pin
+template <typename NumericType>
+class SVector4Slider : public SCompoundWidget
+{
+public:
+
+	// Notification for numeric value committed
+	DECLARE_DELEGATE_TwoParams(FOnNumericValueCommitted, NumericType, ETextCommit::Type);
+
+	SLATE_BEGIN_ARGS(SVector4Slider) {}
+		SLATE_ATTRIBUTE(FString, VisibleText_0)
+		SLATE_ATTRIBUTE(FString, VisibleText_1)
+		SLATE_ATTRIBUTE(FString, VisibleText_2)
+		SLATE_ATTRIBUTE(FString, VisibleText_3)
+		SLATE_EVENT(FOnNumericValueCommitted, OnNumericCommitted_Box_0)
+		SLATE_EVENT(FOnNumericValueCommitted, OnNumericCommitted_Box_1)
+		SLATE_EVENT(FOnNumericValueCommitted, OnNumericCommitted_Box_2)
+		SLATE_EVENT(FOnNumericValueCommitted, OnNumericCommitted_Box_3)
+	SLATE_END_ARGS()
+	
+	// Construct editable text boxes with the appropriate getter & setter functions along with tool tip text
+	void Construct(const FArguments& InArgs, FProperty* InProperty)
+	{
+		PinProperty = InProperty;
+		VisibleText_0 = InArgs._VisibleText_0;
+		VisibleText_1 = InArgs._VisibleText_1;
+		VisibleText_2 = InArgs._VisibleText_2;
+		VisibleText_3 = InArgs._VisibleText_3;
+		const FLinearColor LabelClr = FLinearColor(1.f, 1.f, 1.f, 0.4f);
+
+		OnNumericCommitted_0 = InArgs._OnNumericCommitted_Box_0;
+		OnNumericCommitted_1 = InArgs._OnNumericCommitted_Box_1;
+		OnNumericCommitted_2 = InArgs._OnNumericCommitted_Box_2;
+		OnNumericCommitted_3 = InArgs._OnNumericCommitted_Box_3;
+
+		const typename TNumericPropertyParams<NumericType>::FMetaDataGetter MetaDataGetter = TNumericPropertyParams<NumericType>::FMetaDataGetter::CreateLambda([&](const FName& Key)
+			{
+				return PinProperty->GetMetaData(Key);
+			});
+
+		TNumericPropertyParams<NumericType> NumericPropertyParams(PinProperty, PinProperty ? MetaDataGetter : nullptr);
+
+		const bool bAllowSpin = !(PinProperty && PinProperty->GetBoolMetaData("NoSpinbox"));
+
+		this->ChildSlot
+			[
+				SNew(SVerticalBox)
+				+ SVerticalBox::Slot()
+				.AutoHeight()
+				.Padding(0)
+				[
+					SNew(SHorizontalBox)
+					+ SHorizontalBox::Slot()
+					.AutoWidth().Padding(2).HAlign(HAlign_Fill)
+					[
+						// Create Text box 0 
+						SNew(SNumericEntryBox<NumericType>)
+						.LabelVAlign(VAlign_Center)
+						.Label()
+						[
+							SNew(STextBlock)
+							.Font(FAppStyle::GetFontStyle("Graph.VectorEditableTextBox"))
+							.Text(LOCTEXT("VectorNodeXAxisValueLabel", "X"))
+							.ColorAndOpacity(LabelClr)
+						]
+						.MinValue(NumericPropertyParams.MinValue)
+						.MaxValue(NumericPropertyParams.MaxValue)
+						.MinSliderValue(NumericPropertyParams.MinSliderValue)
+						.MaxSliderValue(NumericPropertyParams.MaxSliderValue)
+						.SliderExponent(NumericPropertyParams.SliderExponent)
+						.Delta(NumericPropertyParams.Delta)
+						.Value(this, &SVector4Slider::GetTypeInValue_0)
+						.LinearDeltaSensitivity(NumericPropertyParams.GetLinearDeltaSensitivityAttribute())
+						.AllowWheel(bAllowSpin)
+						.WheelStep(NumericPropertyParams.WheelStep)
+						.AllowSpin(bAllowSpin)
+						.OnValueCommitted(InArgs._OnNumericCommitted_Box_0)
+						.OnValueChanged(this, &SVector4Slider::OnValueChanged_0)
+						.OnBeginSliderMovement(this, &SVector4Slider::OnBeginSliderMovement_0)
+						.OnEndSliderMovement(this, &SVector4Slider::OnEndSliderMovement_0)
+						.Font(FAppStyle::GetFontStyle("Graph.VectorEditableTextBox"))
+						.UndeterminedString(LOCTEXT("MultipleValues", "Multiple Values"))
+						.ToolTipText(LOCTEXT("VectorNodeXAxisValueLabel_ToolTip", "X value"))
+						.EditableTextBoxStyle(&FAppStyle::GetWidgetStyle<FEditableTextBoxStyle>("Graph.VectorEditableTextBox"))
+						.BorderForegroundColor(FLinearColor::White)
+						.BorderBackgroundColor(FLinearColor::White)
+					]
+					+ SHorizontalBox::Slot()
+					.AutoWidth().Padding(2).HAlign(HAlign_Fill)
+					[
+						// Create Text box 1
+						SNew(SNumericEntryBox<NumericType>)
+						.LabelVAlign(VAlign_Center)
+						.Label()
+						[
+							SNew(STextBlock)
+							.Font(FAppStyle::GetFontStyle("Graph.VectorEditableTextBox"))
+							.Text(LOCTEXT("VectorNodeYAxisValueLabel", "Y"))
+							.ColorAndOpacity(LabelClr)
+						]
+						.MinValue(NumericPropertyParams.MinValue)
+						.MaxValue(NumericPropertyParams.MaxValue)
+						.MinSliderValue(NumericPropertyParams.MinSliderValue)
+						.MaxSliderValue(NumericPropertyParams.MaxSliderValue)
+						.SliderExponent(NumericPropertyParams.SliderExponent)
+						.Delta(NumericPropertyParams.Delta)
+						.Value(this, &SVector4Slider::GetTypeInValue_1)
+						.LinearDeltaSensitivity(NumericPropertyParams.GetLinearDeltaSensitivityAttribute())
+						.AllowWheel(bAllowSpin)
+						.WheelStep(NumericPropertyParams.WheelStep)
+						.AllowSpin(bAllowSpin)
+						.OnValueCommitted(InArgs._OnNumericCommitted_Box_1)
+						.OnValueChanged(this, &SVector4Slider::OnValueChanged_1)
+						.OnBeginSliderMovement(this, &SVector4Slider::OnBeginSliderMovement_1)
+						.OnEndSliderMovement(this, &SVector4Slider::OnEndSliderMovement_1)
+						.Font(FAppStyle::GetFontStyle("Graph.VectorEditableTextBox"))
+						.UndeterminedString(LOCTEXT("MultipleValues", "Multiple Values"))
+						.ToolTipText(LOCTEXT("VectorNodeYAxisValueLabel_ToolTip", "Y value"))
+						.EditableTextBoxStyle(&FAppStyle::GetWidgetStyle<FEditableTextBoxStyle>("Graph.VectorEditableTextBox"))
+						.BorderForegroundColor(FLinearColor::White)
+						.BorderBackgroundColor(FLinearColor::White)
+					]
+					+ SHorizontalBox::Slot()
+					.AutoWidth().Padding(2).HAlign(HAlign_Fill)
+					[
+						// Create Text box 2
+						SNew(SNumericEntryBox<NumericType>)
+						.LabelVAlign(VAlign_Center)
+						.Label()
+						[
+							SNew(STextBlock)
+							.Font(FAppStyle::GetFontStyle("Graph.VectorEditableTextBox"))
+							.Text(LOCTEXT("VectorNodeZAxisValueLabel", "Z"))
+							.ColorAndOpacity(LabelClr)
+						]
+						.MinValue(NumericPropertyParams.MinValue)
+						.MaxValue(NumericPropertyParams.MaxValue)
+						.MinSliderValue(NumericPropertyParams.MinSliderValue)
+						.MaxSliderValue(NumericPropertyParams.MaxSliderValue)
+						.SliderExponent(NumericPropertyParams.SliderExponent)
+						.Delta(NumericPropertyParams.Delta)
+						.Value(this, &SVector4Slider::GetTypeInValue_2)
+						.LinearDeltaSensitivity(NumericPropertyParams.GetLinearDeltaSensitivityAttribute())
+						.AllowWheel(bAllowSpin)
+						.WheelStep(NumericPropertyParams.WheelStep)
+						.AllowSpin(bAllowSpin)
+						.OnValueCommitted(InArgs._OnNumericCommitted_Box_2)
+						.OnValueChanged(this, &SVector4Slider::OnValueChanged_2)
+						.OnBeginSliderMovement(this, &SVector4Slider::OnBeginSliderMovement_2)
+						.OnEndSliderMovement(this, &SVector4Slider::OnEndSliderMovement_2)
+						.Font(FAppStyle::GetFontStyle("Graph.VectorEditableTextBox"))
+						.UndeterminedString(LOCTEXT("MultipleValues", "Multiple Values"))
+						.ToolTipText(LOCTEXT("VectorNodeZAxisValueLabel_ToolTip", "Z value"))
+						.EditableTextBoxStyle(&FAppStyle::GetWidgetStyle<FEditableTextBoxStyle>("Graph.VectorEditableTextBox"))
+						.BorderForegroundColor(FLinearColor::White)
+						.BorderBackgroundColor(FLinearColor::White)
+					]
+					+ SHorizontalBox::Slot()
+					.AutoWidth().Padding(2).HAlign(HAlign_Fill)
+					[
+						// Create Text box 3
+						SNew(SNumericEntryBox<NumericType>)
+						.LabelVAlign(VAlign_Center)
+						.Label()
+						[
+							SNew(STextBlock)
+							.Font(FAppStyle::GetFontStyle("Graph.VectorEditableTextBox"))
+							.Text(LOCTEXT("VectorNodeWAxisValueLabel", "W"))
+							.ColorAndOpacity(LabelClr)
+						]
+						.MinValue(NumericPropertyParams.MinValue)
+						.MaxValue(NumericPropertyParams.MaxValue)
+						.MinSliderValue(NumericPropertyParams.MinSliderValue)
+						.MaxSliderValue(NumericPropertyParams.MaxSliderValue)
+						.SliderExponent(NumericPropertyParams.SliderExponent)
+						.Delta(NumericPropertyParams.Delta)
+						.Value(this, &SVector4Slider::GetTypeInValue_3)
+						.LinearDeltaSensitivity(NumericPropertyParams.GetLinearDeltaSensitivityAttribute())
+						.AllowWheel(bAllowSpin)
+						.WheelStep(NumericPropertyParams.WheelStep)
+						.AllowSpin(bAllowSpin)
+						.OnValueCommitted(InArgs._OnNumericCommitted_Box_3)
+						.OnValueChanged(this, &SVector4Slider::OnValueChanged_3)
+						.OnBeginSliderMovement(this, &SVector4Slider::OnBeginSliderMovement_3)
+						.OnEndSliderMovement(this, &SVector4Slider::OnEndSliderMovement_3)
+						.Font(FAppStyle::GetFontStyle("Graph.VectorEditableTextBox"))
+						.UndeterminedString(LOCTEXT("MultipleValues", "Multiple Values"))
+						.ToolTipText(LOCTEXT("VectorNodeWAxisValueLabel_ToolTip", "W value"))
+						.EditableTextBoxStyle(&FAppStyle::GetWidgetStyle<FEditableTextBoxStyle>("Graph.VectorEditableTextBox"))
+						.BorderForegroundColor(FLinearColor::White)
+						.BorderBackgroundColor(FLinearColor::White)
+					]
+				]
+			];
+		}
+
+private:
+
+	/**
+	* Called when the slider begins to move.
+	*/
+	void OnBeginSliderMovement_0()
+	{
+		SyncValues();
+		bIsUsingSlider = true;
+	}
+
+	/**
+	* Called when the slider stops moving.
+	*/
+	void OnEndSliderMovement_0(NumericType NewValue)
+	{
+		bIsUsingSlider = false;
+		if (LastSliderCommittedValue_0 != NewValue)
+		{
+			OnNumericCommitted_0.ExecuteIfBound(NewValue, ETextCommit::Default);
+			LastSliderCommittedValue_0 = NewValue;
+		}
+	}
+
+	/**
+	* Called when the slider begins to move.
+	*/
+	void OnBeginSliderMovement_1()
+	{
+		SyncValues();
+		bIsUsingSlider = true;
+	}
+
+	/**
+	* Called when the slider stops moving.
+	*/
+	void OnEndSliderMovement_1(NumericType NewValue)
+	{
+		bIsUsingSlider = false;
+		if (LastSliderCommittedValue_1 != NewValue)
+		{
+			OnNumericCommitted_1.ExecuteIfBound(NewValue, ETextCommit::Default);
+			LastSliderCommittedValue_1 = NewValue;
+		}
+	}
+
+	/**
+	* Called when the slider begins to move.
+	*/
+	void OnBeginSliderMovement_2()
+	{
+		SyncValues();
+		bIsUsingSlider = true;
+	}
+
+	/**
+	* Called when the slider stops moving.
+	*/
+	void OnEndSliderMovement_2(NumericType NewValue)
+	{
+		bIsUsingSlider = false;
+		if (LastSliderCommittedValue_2 != NewValue)
+		{
+			OnNumericCommitted_2.ExecuteIfBound(NewValue, ETextCommit::Default);
+			LastSliderCommittedValue_2 = NewValue;
+		}
+	}
+
+	/**
+	* Called when the slider begins to move.
+	*/
+	void OnBeginSliderMovement_3()
+	{
+		SyncValues();
+		bIsUsingSlider = true;
+	}
+
+	/**
+	* Called when the slider stops moving.
+	*/
+	void OnEndSliderMovement_3(NumericType NewValue)
+	{
+		bIsUsingSlider = false;
+		if (LastSliderCommittedValue_3 != NewValue)
+		{
+			OnNumericCommitted_3.ExecuteIfBound(NewValue, ETextCommit::Default);
+			LastSliderCommittedValue_3 = NewValue;
+		}
+	}
+
+	/**
+	* Called when the slider value changes
+	*/
+	void OnValueChanged_0(NumericType NewValue)
+	{
+		SliderValue_0 = NewValue;
+	}
+
+	/**
+	* Called when the slider value changes
+	*/
+	void OnValueChanged_1(NumericType NewValue)
+	{
+		SliderValue_1 = NewValue;
+	}
+
+	/**
+	* Called when the slider value changes
+	*/
+	void OnValueChanged_2(NumericType NewValue)
+	{
+		SliderValue_2 = NewValue;
+	}
+
+	/**
+	* Called when the slider value changes
+	*/
+	void OnValueChanged_3(NumericType NewValue)
+	{
+		SliderValue_3 = NewValue;
+	}
+
+	void SyncValues()
+	{
+		SliderValue_0 = LastSliderCommittedValue_0 = GetTypeInValue_0().GetValue();
+		SliderValue_1 = LastSliderCommittedValue_1 = GetTypeInValue_1().GetValue();
+		SliderValue_2 = LastSliderCommittedValue_2 = GetTypeInValue_2().GetValue();
+		SliderValue_3 = LastSliderCommittedValue_3 = GetTypeInValue_3().GetValue();
+	}
+	NumericType GetValueType(const FString& InString) const
+	{
+		static_assert(std::is_floating_point_v<NumericType>);
+
+		if constexpr (std::is_same_v<float, NumericType>)
+		{
+			return FCString::Atof(*InString);
+		}
+		else if constexpr (std::is_same_v<double, NumericType>)
+		{
+			return FCString::Atod(*InString);
+		}
+
+		return NumericType{};
+	}
+
+	TOptional<NumericType> GetTypeInValue_0() const
+	{
+		return bIsUsingSlider ? SliderValue_0 : GetValueType(VisibleText_0.Get());
+	}
+
+	TOptional<NumericType> GetTypeInValue_1() const
+	{
+		return bIsUsingSlider ? SliderValue_1 : GetValueType(VisibleText_1.Get());
+	}
+
+	TOptional<NumericType> GetTypeInValue_2() const
+	{
+		return bIsUsingSlider ? SliderValue_2 : GetValueType(VisibleText_2.Get());
+	}
+
+	TOptional<NumericType> GetTypeInValue_3() const
+	{
+		return bIsUsingSlider ? SliderValue_3 : GetValueType(VisibleText_3.Get());
+	}
+
+	TAttribute<FString> VisibleText_0;
+	TAttribute<FString> VisibleText_1;
+	TAttribute<FString> VisibleText_2;
+	TAttribute<FString> VisibleText_3;
+	FProperty* PinProperty;
+
+	FOnNumericValueCommitted OnNumericCommitted_0;
+	FOnNumericValueCommitted OnNumericCommitted_1;
+	FOnNumericValueCommitted OnNumericCommitted_2;
+	FOnNumericValueCommitted OnNumericCommitted_3;
+
+	NumericType LastSliderCommittedValue_0;
+	NumericType LastSliderCommittedValue_1;
+	NumericType LastSliderCommittedValue_2;
+	NumericType LastSliderCommittedValue_3;
+	NumericType SliderValue_0;
+	NumericType SliderValue_1;
+	NumericType SliderValue_2;
+	NumericType SliderValue_3;
+	bool bIsUsingSlider = false;
+};
+
+#undef LOCTEXT_NAMESPACE
diff --git a/Engine/Source/Editor/GraphEditor/Public/KismetPins/SVectorSlider.h b/Engine/Source/Editor/GraphEditor/Public/KismetPins/SVectorSlider.h
new file mode 100644
--- /dev/null
+++ b/Engine/Source/Editor/GraphEditor/Public/KismetPins/SVectorSlider.h
@@ -0,0 +1,328 @@
+// Copyright Epic Games, Inc. All Rights Reserved.
+
+#pragma once
+
+#include "CoreMinimal.h"
+#include "Widgets/Input/SNumericEntryBox.h"
+#include "Widgets/SBoxPanel.h"
+#include "NumericPropertyParams.h"
+#include "Editor.h"
+#include "Widgets/SCompoundWidget.h"
+
+#define LOCTEXT_NAMESPACE "Vector2DTextBox"
+
+// Class implementation to create 2 editable text boxes to represent vector2D graph pin
+template <typename NumericType>
+class SVectorSlider : public SCompoundWidget
+{
+public:
+
+	// Notification for numeric value committed
+	DECLARE_DELEGATE_TwoParams(FOnNumericValueCommitted, NumericType, ETextCommit::Type);
+
+	SLATE_BEGIN_ARGS( SVectorSlider ){}	
+		SLATE_ATTRIBUTE( FString, VisibleText_0 )
+		SLATE_ATTRIBUTE( FString, VisibleText_1 )
+		SLATE_ATTRIBUTE( FString, VisibleText_2)
+		SLATE_EVENT(FOnNumericValueCommitted, OnNumericCommitted_Box_0)
+		SLATE_EVENT(FOnNumericValueCommitted, OnNumericCommitted_Box_1)
+		SLATE_EVENT(FOnNumericValueCommitted, OnNumericCommitted_Box_2)
+	SLATE_END_ARGS()
+
+	// Construct editable text boxes with the appropriate getter & setter functions along with tool tip text
+	void Construct( const FArguments& InArgs,const bool bInIsRotator, FProperty* InProperty)
+	{
+		bIsRotator = bInIsRotator;
+		const bool bUseRPY = false;
+
+		PinProperty = InProperty;
+		VisibleText_0 = InArgs._VisibleText_0;
+		VisibleText_1 = InArgs._VisibleText_1;
+		VisibleText_2 = InArgs._VisibleText_2;
+		const FLinearColor LabelClr = FLinearColor(1.f, 1.f, 1.f, 0.4f);
+
+		OnNumericCommitted_0 = InArgs._OnNumericCommitted_Box_0;
+		OnNumericCommitted_1 = InArgs._OnNumericCommitted_Box_1;
+		OnNumericCommitted_2 = InArgs._OnNumericCommitted_Box_2;
+
+		const typename TNumericPropertyParams<NumericType>::FMetaDataGetter MetaDataGetter = TNumericPropertyParams<NumericType>::FMetaDataGetter::CreateLambda([&](const FName& Key)
+			{
+				return PinProperty->GetMetaData(Key);
+			});
+
+		TNumericPropertyParams<NumericType> NumericPropertyParams(PinProperty, PinProperty ? MetaDataGetter : nullptr);
+
+		const bool bAllowSpin = !(PinProperty && PinProperty->GetBoolMetaData("NoSpinbox"));
+
+		this->ChildSlot
+		[
+			SNew(SVerticalBox)
+			+SVerticalBox::Slot()
+			.AutoHeight()
+			.Padding(0)
+			[
+				SNew(SHorizontalBox)
+				+SHorizontalBox::Slot()
+				.AutoWidth().Padding(2) .HAlign(HAlign_Fill)
+				[
+					// Create Text box 0 
+					SNew( SNumericEntryBox<NumericType> )
+					.LabelVAlign(VAlign_Center)
+					.Label()
+					[
+						SNew( STextBlock )
+						.Font( FAppStyle::GetFontStyle( "Graph.VectorEditableTextBox" ) )
+						.Text( bIsRotator && bUseRPY ? LOCTEXT("VectorNodeRollValueLabel", "R") : LOCTEXT("VectorNodeXAxisValueLabel", "X") )
+						.ColorAndOpacity( LabelClr )
+					]
+					.MinValue(NumericPropertyParams.MinValue)
+					.MaxValue(NumericPropertyParams.MaxValue)
+					.MinSliderValue(NumericPropertyParams.MinSliderValue)
+					.MaxSliderValue(NumericPropertyParams.MaxSliderValue)
+					.SliderExponent(NumericPropertyParams.SliderExponent)
+					.Delta(NumericPropertyParams.Delta)
+					.Value( this, &SVectorSlider::GetTypeInValue_0 )
+					.LinearDeltaSensitivity(NumericPropertyParams.GetLinearDeltaSensitivityAttribute())
+					.AllowWheel(bAllowSpin)
+					.WheelStep(NumericPropertyParams.WheelStep)
+					.AllowSpin(bAllowSpin)
+					.OnValueCommitted( InArgs._OnNumericCommitted_Box_0 )
+					.OnValueChanged(this, &SVectorSlider::OnValueChanged_0)
+					.OnBeginSliderMovement(this, &SVectorSlider::OnBeginSliderMovement_0)
+					.OnEndSliderMovement(this, &SVectorSlider::OnEndSliderMovement_0)
+					.Font( FAppStyle::GetFontStyle( "Graph.VectorEditableTextBox" ) )
+					.UndeterminedString( LOCTEXT("MultipleValues", "Multiple Values") )
+					.ToolTipText(bIsRotator ? LOCTEXT("VectorNodeRollValueLabel_ToolTip", "Roll value (around X)") : LOCTEXT("VectorNodeXAxisValueLabel_ToolTip", "X value"))
+					.EditableTextBoxStyle( &FAppStyle::GetWidgetStyle<FEditableTextBoxStyle>( "Graph.VectorEditableTextBox" ))
+					.BorderForegroundColor( FLinearColor::White )
+					.BorderBackgroundColor( FLinearColor::White )
+				]
+				+SHorizontalBox::Slot()
+				.AutoWidth().Padding(2) .HAlign(HAlign_Fill)
+				[
+					// Create Text box 1
+					SNew( SNumericEntryBox<NumericType> )
+					.LabelVAlign(VAlign_Center)
+					.Label()
+					[
+						SNew( STextBlock )
+						.Font( FAppStyle::GetFontStyle( "Graph.VectorEditableTextBox" ) )
+						.Text( bIsRotator && bUseRPY ? LOCTEXT("VectorNodePitchValueLabel", "P") : LOCTEXT("VectorNodeYAxisValueLabel", "Y") )
+						.ColorAndOpacity( LabelClr )
+					]
+					.MinValue(NumericPropertyParams.MinValue)
+					.MaxValue(NumericPropertyParams.MaxValue)
+					.MinSliderValue(NumericPropertyParams.MinSliderValue)
+					.MaxSliderValue(NumericPropertyParams.MaxSliderValue)
+					.SliderExponent(NumericPropertyParams.SliderExponent)
+					.Delta(NumericPropertyParams.Delta)
+					.Value( this, &SVectorSlider::GetTypeInValue_1 )
+					// LinearDeltaSensitivity needs to be left unset if not provided, rather than being set to some default
+					.LinearDeltaSensitivity(NumericPropertyParams.GetLinearDeltaSensitivityAttribute())
+					.AllowWheel(bAllowSpin)
+					.WheelStep(NumericPropertyParams.WheelStep)
+					.AllowSpin(bAllowSpin)
+					.OnValueCommitted( InArgs._OnNumericCommitted_Box_1 )
+					.OnValueChanged(this, &SVectorSlider::OnValueChanged_1)
+					.OnBeginSliderMovement(this, &SVectorSlider::OnBeginSliderMovement_1)
+					.OnEndSliderMovement(this, &SVectorSlider::OnEndSliderMovement_1)
+					.Font( FAppStyle::GetFontStyle( "Graph.VectorEditableTextBox" ) )
+					.UndeterminedString( LOCTEXT("MultipleValues", "Multiple Values") )
+					.ToolTipText(bIsRotator ? LOCTEXT("VectorNodePitchValueLabel_ToolTip", "Pitch value (around Y)") : LOCTEXT("VectorNodeYAxisValueLabel_ToolTip", "Y value"))
+					.EditableTextBoxStyle( &FAppStyle::GetWidgetStyle<FEditableTextBoxStyle>( "Graph.VectorEditableTextBox" ))
+					.BorderForegroundColor( FLinearColor::White )
+					.BorderBackgroundColor( FLinearColor::White )
+				]
+				+ SHorizontalBox::Slot()
+					.AutoWidth().Padding(2).HAlign(HAlign_Fill)
+					[
+						// Create Text box 1
+						SNew(SNumericEntryBox<NumericType>)
+						.LabelVAlign(VAlign_Center)
+					.Label()
+					[
+						SNew(STextBlock)
+						.Font(FAppStyle::GetFontStyle("Graph.VectorEditableTextBox"))
+						.Text(bIsRotator && bUseRPY ? LOCTEXT("VectorNodeYawValueLabel", "Y") : LOCTEXT("VectorNodeZAxisValueLabel", "Z"))
+						.ColorAndOpacity(LabelClr)
+					]
+					.MinValue(NumericPropertyParams.MinValue)
+					.MaxValue(NumericPropertyParams.MaxValue)
+					.MinSliderValue(NumericPropertyParams.MinSliderValue)
+					.MaxSliderValue(NumericPropertyParams.MaxSliderValue)
+					.SliderExponent(NumericPropertyParams.SliderExponent)
+					.Delta(NumericPropertyParams.Delta)
+					.Value(this, &SVectorSlider::GetTypeInValue_2)
+					// LinearDeltaSensitivity needs to be left unset if not provided, rather than being set to some default
+					.LinearDeltaSensitivity(NumericPropertyParams.GetLinearDeltaSensitivityAttribute())
+					.AllowWheel(bAllowSpin)
+					.WheelStep(NumericPropertyParams.WheelStep)
+					.AllowSpin(bAllowSpin)
+					.OnValueCommitted(InArgs._OnNumericCommitted_Box_2)
+					.OnValueChanged(this, &SVectorSlider::OnValueChanged_2)
+					.OnBeginSliderMovement(this, &SVectorSlider::OnBeginSliderMovement_2)
+					.OnEndSliderMovement(this, &SVectorSlider::OnEndSliderMovement_2)
+					.Font(FAppStyle::GetFontStyle("Graph.VectorEditableTextBox"))
+					.UndeterminedString(LOCTEXT("MultipleValues", "Multiple Values"))
+					.ToolTipText(bIsRotator ? LOCTEXT("VectorNodeYawValueLabel_Tooltip", "Yaw value (around Z)") : LOCTEXT("VectorNodeZAxisValueLabel_ToolTip", "Z value"))
+					.EditableTextBoxStyle(&FAppStyle::GetWidgetStyle<FEditableTextBoxStyle>("Graph.VectorEditableTextBox"))
+					.BorderForegroundColor(FLinearColor::White)
+					.BorderBackgroundColor(FLinearColor::White)
+					]
+			]
+		];
+	}
+
+private:
+
+	/**
+	* Called when the slider begins to move.
+	*/
+	void OnBeginSliderMovement_0()
+	{
+		SyncValues();
+		bIsUsingSlider = true;
+	}
+
+	/**
+	* Called when the slider stops moving.
+	*/
+	void OnEndSliderMovement_0(NumericType NewValue)
+	{
+		bIsUsingSlider = false;
+		if (LastSliderCommittedValue_0 != NewValue)
+		{
+			OnNumericCommitted_0.ExecuteIfBound(NewValue, ETextCommit::Default);
+			LastSliderCommittedValue_0 = NewValue;
+		}
+	}
+
+	/**
+	* Called when the slider begins to move.
+	*/
+	void OnBeginSliderMovement_1()
+	{
+		SyncValues();
+		bIsUsingSlider = true;
+	}
+
+	/**
+	* Called when the slider stops moving.
+	*/
+	void OnEndSliderMovement_1(NumericType NewValue)
+	{
+		bIsUsingSlider = false;
+		if (LastSliderCommittedValue_1 != NewValue)
+		{
+			OnNumericCommitted_1.ExecuteIfBound(NewValue, ETextCommit::Default);
+			LastSliderCommittedValue_1 = NewValue;
+		}
+	}
+
+	/**
+	* Called when the slider begins to move.
+	*/
+	void OnBeginSliderMovement_2()
+	{
+		SyncValues();
+		bIsUsingSlider = true;
+	}
+
+	/**
+	* Called when the slider stops moving.
+	*/
+	void OnEndSliderMovement_2(NumericType NewValue)
+	{
+		bIsUsingSlider = false;
+		if (LastSliderCommittedValue_2 != NewValue)
+		{
+			OnNumericCommitted_2.ExecuteIfBound(NewValue, ETextCommit::Default);
+			LastSliderCommittedValue_2 = NewValue;
+		}
+	}
+
+	/**
+	* Called when the slider value changes
+	*/
+	void OnValueChanged_0(NumericType NewValue)
+	{
+		SliderValue_0 = NewValue;
+	}
+
+	/**
+	* Called when the slider value changes
+	*/
+	void OnValueChanged_1(NumericType NewValue)
+	{
+		SliderValue_1 = NewValue;
+	}
+
+	/**
+	* Called when the slider value changes
+	*/
+	void OnValueChanged_2(NumericType NewValue)
+	{
+		SliderValue_2 = NewValue;
+	}
+
+	void SyncValues()
+	{
+		SliderValue_0 = LastSliderCommittedValue_0 = GetTypeInValue_0().GetValue();
+		SliderValue_1 = LastSliderCommittedValue_1 = GetTypeInValue_1().GetValue();
+		SliderValue_2 = LastSliderCommittedValue_2 = GetTypeInValue_2().GetValue();
+	}
+
+	NumericType GetValueType(const FString& InString) const
+	{
+		static_assert(std::is_floating_point_v<NumericType>);
+
+		if constexpr (std::is_same_v<float, NumericType>)
+		{
+			return FCString::Atof(*InString);
+		}
+		else if constexpr (std::is_same_v<double, NumericType>)
+		{
+			return FCString::Atod(*InString);
+		}
+
+		return NumericType{};
+	}
+
+	// Get value for X text box
+	TOptional<NumericType> GetTypeInValue_0() const
+	{
+		return bIsUsingSlider ? SliderValue_0 :GetValueType(VisibleText_0.Get());
+	}
+
+	// Get value for Y text box
+	TOptional<NumericType> GetTypeInValue_1() const
+	{
+		return bIsUsingSlider ? SliderValue_1 : GetValueType(VisibleText_1.Get());
+	}
+	// Get value for Y text box
+	TOptional<NumericType> GetTypeInValue_2() const
+	{
+		return bIsUsingSlider ? SliderValue_2 : GetValueType(VisibleText_2.Get());
+	}
+
+
+	TAttribute<FString> VisibleText_0;
+	TAttribute<FString> VisibleText_1;
+	TAttribute<FString> VisibleText_2;
+
+	FProperty* PinProperty;
+	FOnNumericValueCommitted OnNumericCommitted_0;
+	FOnNumericValueCommitted OnNumericCommitted_1;
+	FOnNumericValueCommitted OnNumericCommitted_2;
+
+	NumericType LastSliderCommittedValue_0;
+	NumericType LastSliderCommittedValue_1;
+	NumericType LastSliderCommittedValue_2;
+	NumericType SliderValue_0;
+	NumericType SliderValue_1;
+	NumericType SliderValue_2;
+	bool bIsUsingSlider = false;
+	bool bIsRotator;
+};
+
+#undef LOCTEXT_NAMESPACE
diff --git a/Engine/Source/Editor/MaterialEditor/Private/MaterialEditorGraphPanelPinFactory.cpp b/Engine/Source/Editor/MaterialEditor/Private/MaterialEditorGraphPanelPinFactory.cpp
--- a/Engine/Source/Editor/MaterialEditor/Private/MaterialEditorGraphPanelPinFactory.cpp
+++ b/Engine/Source/Editor/MaterialEditor/Private/MaterialEditorGraphPanelPinFactory.cpp
@@ -9,11 +9,6 @@
 #include "KismetPins/SGraphPinBool.h"
 #include "KismetPins/SGraphPinColor.h"
 #include "KismetPins/SGraphPinExec.h"
-#include "KismetPins/SGraphPinInteger.h"
-#include "KismetPins/SGraphPinNum.h"
-#include "KismetPins/SGraphPinVector.h"
-#include "KismetPins/SGraphPinVector2D.h"
-#include "KismetPins/SGraphPinVector4.h"
 #include "MaterialGraph/MaterialGraphSchema.h"
 #include "MaterialPins/SGraphPinMaterialInput.h"
 #include "Misc/Optional.h"
@@ -25,6 +20,14 @@
 #include "UObject/WeakObjectPtr.h"
 #include "UObject/WeakObjectPtrTemplates.h"
 #include "Widgets/DeclarativeSyntaxSupport.h"
+#include "MaterialGraph/MaterialGraphNode.h"
+#include "Materials/MaterialExpression.h"
+#include "Materials/MaterialExpressionScalarParameter.h"
+#include "KismetPins/SGraphPinNumSlider.h"
+#include "KismetPins/SGraphPinVector2DSlider.h"
+#include "KismetPins/SGraphPinVectorSlider.h"
+#include "KismetPins/SGraphPinVector4Slider.h"
+#include "KismetPins/SGraphPinIntegerSlider.h"
 
 TSharedPtr<class SGraphPin> FMaterialEditorGraphPanelPinFactory::CreatePin(class UEdGraphPin* InPin) const
 {
@@ -34,23 +37,29 @@ TSharedPtr<class SGraphPin> FMaterialEditorGraphPanelPinFactory::CreatePin(class
 		{
 			return SNew(SGraphPinExec, InPin);
 		}
-		else if (InPin->PinType.PinCategory == MaterialGraphSchema->PC_MaterialInput)
-		{
-			return SNew(SGraphPinMaterialInput, InPin);
-		}
 		else
 		{
+			// get metadata and pass it along
+			UObject* NodeOwner = CastChecked<UMaterialGraphNode_Base>(InPin->GetOuter())->GetMaterialNodeOwner();
+
+			TArray<FProperty*> InputProperties;
+			if (NodeOwner->IsA<UMaterialExpression>())
+			{
+				InputProperties = CastChecked<UMaterialExpression>(NodeOwner)->GetInputPinProperty(InPin->SourceIndex);
+			}
+			FProperty* InProperty = InputProperties.Num() > 0 ? InputProperties[0] : nullptr;
+
 			if (InPin->PinType.PinSubCategory == MaterialGraphSchema->PSC_Red || InPin->PinType.PinSubCategory == MaterialGraphSchema->PSC_Float)
 			{
-				return SNew(SGraphPinNum<double>, InPin);
+				return SNew(SGraphPinNumSlider<double>, InPin, InProperty);	
 			}
 			else if (InPin->PinType.PinSubCategory == MaterialGraphSchema->PSC_RG)
 			{
-				return SNew(SGraphPinVector2D<float>, InPin);
+				return SNew(SGraphPinVector2DSlider<float>, InPin, InProperty);
 			}
 			else if (InPin->PinType.PinSubCategory == MaterialGraphSchema->PSC_RGB)
 			{
-				return SNew(SGraphPinVector<float>, InPin);
+				return SNew(SGraphPinVectorSlider<float>, InPin, InProperty);
 			}
 			else if (InPin->PinType.PinSubCategory == MaterialGraphSchema->PSC_RGBA)
 			{
@@ -58,11 +67,11 @@ TSharedPtr<class SGraphPin> FMaterialEditorGraphPanelPinFactory::CreatePin(class
 			}
 			else if (InPin->PinType.PinSubCategory == MaterialGraphSchema->PSC_Vector4)
 			{
-				return SNew(SGraphPinVector4<float>, InPin);
+				return SNew(SGraphPinVector4Slider<float>, InPin, InProperty);
 			}
 			else if (InPin->PinType.PinSubCategory == MaterialGraphSchema->PSC_Int)
 			{
-				return SNew(SGraphPinInteger, InPin);
+				return SNew(SGraphPinIntegerSlider, InPin, InProperty);
 			}
 			else if (InPin->PinType.PinSubCategory == MaterialGraphSchema->PSC_Byte)
 			{
@@ -73,13 +82,17 @@ TSharedPtr<class SGraphPin> FMaterialEditorGraphPanelPinFactory::CreatePin(class
 				}
 				else
 				{
-					return SNew(SGraphPinInteger, InPin);
+					return SNew(SGraphPinIntegerSlider, InPin, InProperty);
 				}
 			}
 			else if (InPin->PinType.PinSubCategory == MaterialGraphSchema->PSC_Bool)
 			{
 				return SNew(SGraphPinBool, InPin);
 			}
+			else if (InPin->PinType.PinCategory == MaterialGraphSchema->PC_MaterialInput)
+			{
+				return SNew(SGraphPinMaterialInput, InPin);
+			}
 		}
 
 		return SNew(SGraphPin, InPin);
diff --git a/Engine/Source/Editor/UnrealEd/Classes/MaterialGraph/MaterialGraphNode_Root.h b/Engine/Source/Editor/UnrealEd/Classes/MaterialGraph/MaterialGraphNode_Root.h
--- a/Engine/Source/Editor/UnrealEd/Classes/MaterialGraph/MaterialGraphNode_Root.h
+++ b/Engine/Source/Editor/UnrealEd/Classes/MaterialGraph/MaterialGraphNode_Root.h
@@ -18,6 +18,8 @@ class UMaterialGraphNode_Root : public UMaterialGraphNode_Base
 	UPROPERTY()
 	TObjectPtr<class UMaterial> Material;
 
+	void UpdateInputUseConstant(UEdGraphPin* Pin, bool bUseConstant);
+
 	//~ Begin UEdGraphNode Interface.
 	virtual FText GetNodeTitle(ENodeTitleType::Type TitleType) const override;
 	virtual FLinearColor GetNodeTitleColor() const override;
@@ -26,6 +28,7 @@ class UMaterialGraphNode_Root : public UMaterialGraphNode_Base
 	virtual bool CanDuplicateNode() const override { return false; }
 	virtual void PostPlacedNewNode() override;
 	virtual TSharedPtr<SGraphNode> CreateVisualWidget() override;
+	virtual void PinDefaultValueChanged(UEdGraphPin* Pin) override;
 	//~ End UEdGraphNode Interface.
 
 	//~ Begin UMaterialGraphNode_Base Interface
diff --git a/Engine/Source/Editor/UnrealEd/Private/MaterialGraph.cpp b/Engine/Source/Editor/UnrealEd/Private/MaterialGraph.cpp
--- a/Engine/Source/Editor/UnrealEd/Private/MaterialGraph.cpp
+++ b/Engine/Source/Editor/UnrealEd/Private/MaterialGraph.cpp
@@ -619,6 +619,8 @@ void UMaterialGraph::LinkMaterialExpressionsFromGraph()
 					{
 						MaterialInput.Expression = NULL;
 					}
+
+					RootNode->UpdateInputUseConstant(Pin, MaterialInput.Expression == nullptr);
 				}
 			}
 		}
diff --git a/Engine/Source/Editor/UnrealEd/Private/MaterialGraphNode_Root.cpp b/Engine/Source/Editor/UnrealEd/Private/MaterialGraphNode_Root.cpp
--- a/Engine/Source/Editor/UnrealEd/Private/MaterialGraphNode_Root.cpp
+++ b/Engine/Source/Editor/UnrealEd/Private/MaterialGraphNode_Root.cpp
@@ -24,6 +24,45 @@ UMaterialGraphNode_Root::UMaterialGraphNode_Root(const FObjectInitializer& Objec
 {
 }
 
+void UMaterialGraphNode_Root::UpdateInputUseConstant(UEdGraphPin* Pin, bool bUseConstant)
+{
+	const UMaterialGraph* MaterialGraph = CastChecked<UMaterialGraph>(GetGraph());
+	const FMaterialInputInfo& MaterialInput = MaterialGraph->MaterialInputs[Pin->SourceIndex];
+
+	UMaterialEditorOnlyData* EditorOnlyData = Material->GetEditorOnlyData();
+	EMaterialProperty Property = MaterialInput.GetProperty();
+	switch (Property)
+	{
+	case MP_EmissiveColor:		EditorOnlyData->EmissiveColor.UseConstant = bUseConstant; break;
+	case MP_Opacity:			EditorOnlyData->Opacity.UseConstant = bUseConstant; break;
+	case MP_OpacityMask:		EditorOnlyData->OpacityMask.UseConstant = bUseConstant; break;
+	case MP_BaseColor:			EditorOnlyData->BaseColor.UseConstant = bUseConstant; break;
+	case MP_Metallic:			EditorOnlyData->Metallic.UseConstant = bUseConstant; break;
+	case MP_Specular:			EditorOnlyData->Specular.UseConstant = bUseConstant; break;
+	case MP_Roughness:			EditorOnlyData->Roughness.UseConstant = bUseConstant; break;
+	case MP_Anisotropy:			EditorOnlyData->Anisotropy.UseConstant = bUseConstant; break;
+	case MP_Normal:				EditorOnlyData->Normal.UseConstant = bUseConstant; break;
+	case MP_Tangent:			EditorOnlyData->Tangent.UseConstant = bUseConstant; break;
+	case MP_WorldPositionOffset:EditorOnlyData->WorldPositionOffset.UseConstant = bUseConstant; break;
+	case MP_SubsurfaceColor:	EditorOnlyData->SubsurfaceColor.UseConstant = bUseConstant; break;
+	case MP_CustomData0:		EditorOnlyData->ClearCoat.UseConstant = bUseConstant; break;
+	case MP_CustomData1:		EditorOnlyData->ClearCoatRoughness.UseConstant = bUseConstant; break;
+	case MP_AmbientOcclusion:	EditorOnlyData->AmbientOcclusion.UseConstant = bUseConstant; break;
+	case MP_Refraction:			EditorOnlyData->Refraction.UseConstant = bUseConstant; break;
+	case MP_PixelDepthOffset:	EditorOnlyData->PixelDepthOffset.UseConstant = bUseConstant; break;
+	case MP_ShadingModel:		break; // TODO, see notes in CreateInputPins
+	case MP_FrontMaterial:		break; // TODO, see notes in CreateInputPins
+	case MP_SurfaceThickness:	EditorOnlyData->SurfaceThickness.UseConstant = bUseConstant; break;
+	case MP_Displacement:		EditorOnlyData->Displacement.UseConstant = bUseConstant; break;
+	default:
+		if (Property >= MP_CustomizedUVs0 && Property <= MP_CustomizedUVs7)
+		{
+			EditorOnlyData->CustomizedUVs[Property - MP_CustomizedUVs0].UseConstant = bUseConstant; break;
+		}
+		break;
+	}
+}
+
 FText UMaterialGraphNode_Root::GetNodeTitle(ENodeTitleType::Type TitleType) const
 {
 	return FMaterialEditorUtilities::GetOriginalObjectName(this->GetGraph());
@@ -53,6 +92,47 @@ TSharedPtr<SGraphNode> UMaterialGraphNode_Root::CreateVisualWidget()
 	return SNew(SGraphNodeMaterialResult, this);
 }
 
+void UMaterialGraphNode_Root::PinDefaultValueChanged(UEdGraphPin* Pin)
+{
+	const UMaterialGraph* MaterialGraph = CastChecked<UMaterialGraph>(GetGraph());
+	const FMaterialInputInfo& MaterialInput = MaterialGraph->MaterialInputs[Pin->SourceIndex];
+
+	UMaterialEditorOnlyData* EditorOnlyData = Material->GetEditorOnlyData();
+	EMaterialProperty Property = MaterialInput.GetProperty();
+	switch (Property)
+	{
+	case MP_EmissiveColor:		EditorOnlyData->EmissiveColor.DefaultValueChanged(Pin->DefaultValue); break;
+	case MP_Opacity:			EditorOnlyData->Opacity.DefaultValueChanged(Pin->DefaultValue); break;
+	case MP_OpacityMask:		EditorOnlyData->OpacityMask.DefaultValueChanged(Pin->DefaultValue); break;
+	case MP_BaseColor:			EditorOnlyData->BaseColor.DefaultValueChanged(Pin->DefaultValue); break;
+	case MP_Metallic:			EditorOnlyData->Metallic.DefaultValueChanged(Pin->DefaultValue); break;
+	case MP_Specular:			EditorOnlyData->Specular.DefaultValueChanged(Pin->DefaultValue); break;
+	case MP_Roughness:			EditorOnlyData->Roughness.DefaultValueChanged(Pin->DefaultValue); break;
+	case MP_Anisotropy:			EditorOnlyData->Anisotropy.DefaultValueChanged(Pin->DefaultValue); break;
+	case MP_Normal:				EditorOnlyData->Normal.DefaultValueChanged(Pin->DefaultValue); break;
+	case MP_Tangent:			EditorOnlyData->Tangent.DefaultValueChanged(Pin->DefaultValue); break;
+	case MP_WorldPositionOffset:EditorOnlyData->WorldPositionOffset.DefaultValueChanged(Pin->DefaultValue); break;
+	case MP_SubsurfaceColor:	EditorOnlyData->SubsurfaceColor.DefaultValueChanged(Pin->DefaultValue); break;
+	case MP_CustomData0:		EditorOnlyData->ClearCoat.DefaultValueChanged(Pin->DefaultValue); break;
+	case MP_CustomData1:		EditorOnlyData->ClearCoatRoughness.DefaultValueChanged(Pin->DefaultValue); break;
+	case MP_AmbientOcclusion:	EditorOnlyData->AmbientOcclusion.DefaultValueChanged(Pin->DefaultValue); break;
+	case MP_Refraction:			EditorOnlyData->Refraction.DefaultValueChanged(Pin->DefaultValue); break;
+	case MP_PixelDepthOffset:	EditorOnlyData->PixelDepthOffset.DefaultValueChanged(Pin->DefaultValue); break;
+	case MP_ShadingModel:		break; // TODO, see notes in CreateInputPins
+	case MP_FrontMaterial:		break; // TODO, see notes in CreateInputPins
+	case MP_SurfaceThickness:	EditorOnlyData->SurfaceThickness.DefaultValueChanged(Pin->DefaultValue); break;
+	case MP_Displacement:		EditorOnlyData->Displacement.DefaultValueChanged(Pin->DefaultValue); break; 
+	default:
+		if (Property >= MP_CustomizedUVs0 && Property <= MP_CustomizedUVs7)
+		{
+			EditorOnlyData->CustomizedUVs[Property - MP_CustomizedUVs0].DefaultValueChanged(Pin->DefaultValue); break;
+		}
+		break;
+	}
+
+	FMaterialEditorUtilities::UpdateMaterialAfterGraphChange(MaterialGraph);
+}
+
 UObject* UMaterialGraphNode_Root::GetMaterialNodeOwner() const
 {
 	return Material;
@@ -114,15 +194,35 @@ void UMaterialGraphNode_Root::CreateInputPins()
 		NewPin->PinFriendlyName = LOCTEXT("Space", " ");
 	}
 
+	UMaterialEditorOnlyData* EditorOnlyData = Material->GetEditorOnlyData();
 	for (int32 Index = 0; Index < MaterialGraph->MaterialInputs.Num(); ++Index)
 	{
 		const FMaterialInputInfo& MaterialInput = MaterialGraph->MaterialInputs[Index];
 		EMaterialProperty Property = MaterialInput.GetProperty();
 	
-		FName MaterialInputName;
-		if (Property == MP_Refraction)
+		FName MaterialInputName = *MaterialInput.GetName().ToString();
+		FName PinSubCategory;
+		FString RefractionMethodStr;
+		switch (Property)
 		{
-			FString RefractionMethodStr;
+		case MP_Metallic:
+		case MP_Specular:
+		case MP_Roughness:
+		case MP_Anisotropy:
+		case MP_Opacity:
+		case MP_OpacityMask:
+		case MP_CustomData0:
+		case MP_CustomData1:
+		case MP_AmbientOcclusion:
+		case MP_PixelDepthOffset:
+		case MP_Displacement:
+		case MP_SurfaceThickness:
+			PinSubCategory = UMaterialGraphSchema::PSC_Red;
+			break;
+			
+		case MP_Refraction:
+			PinSubCategory = UMaterialGraphSchema::PSC_Red;
+			
 			switch (MaterialGraph->Material->RefractionMethod)
 			{
 			case ERefractionMode::RM_None:
@@ -138,18 +238,71 @@ void UMaterialGraphNode_Root::CreateInputPins()
 				RefractionMethodStr = TEXT("2D Offset");
 				break;
 			default:
-				RefractionMethodStr = TEXT("UNKOWN ERefractionMode");
+				RefractionMethodStr = TEXT("UNKNOWN ERefractionMode");
 			}
 			
 			MaterialInputName = *FString::Printf(TEXT("%s (%s)"), *MaterialInput.GetName().ToString(), *RefractionMethodStr);
-		}
-		else
-		{
-			MaterialInputName = *MaterialInput.GetName().ToString();
+			break;
+		case MP_Normal:
+		case MP_Tangent:
+		case MP_WorldPositionOffset:
+			PinSubCategory = UMaterialGraphSchema::PSC_RGB;
+			break;
+		case MP_BaseColor:
+		case MP_EmissiveColor:
+		case MP_SubsurfaceColor:
+			PinSubCategory = UMaterialGraphSchema::PSC_RGBA;
+			break;
+
+		case MP_ShadingModel: // TODO FMaterialInput<uint32>
+		case MP_FrontMaterial: // TODO FMaterialInput<uint32>
+			// TODO: Not sure if we want to simply show a Num Pin UI here.
+			// Skipping these for now.
+			// If these are based on Enum types, Enum pin construction requires a UEnum object to create UI with
+			// which these constants don't have.
+			break;
+			
+		default:
+			if (Property >= MP_CustomizedUVs0 && Property <= MP_CustomizedUVs7)
+			{
+				PinSubCategory = UMaterialGraphSchema::PSC_RG;
+			}
+			break;
 		}
 
-		UEdGraphPin* InputPin = CreatePin(EGPD_Input, UMaterialGraphSchema::PC_MaterialInput, *FString::Printf(TEXT("%d"), (int32)Property), MaterialInputName);
+		UEdGraphPin* InputPin = CreatePin(EGPD_Input, UMaterialGraphSchema::PC_MaterialInput, PinSubCategory/*, *FString::Printf(TEXT("%d"), (int32)Property)*/, MaterialInputName);
 		InputPin->SourceIndex = Index;
+
+		switch (Property)
+		{
+		case MP_EmissiveColor:		InputPin->DefaultValue = EditorOnlyData->EmissiveColor.GetDefaultValue(); break;
+		case MP_Opacity:			InputPin->DefaultValue = EditorOnlyData->Opacity.GetDefaultValue(); break;
+		case MP_OpacityMask:		InputPin->DefaultValue = EditorOnlyData->OpacityMask.GetDefaultValue(); break;
+		case MP_BaseColor:			InputPin->DefaultValue = EditorOnlyData->BaseColor.GetDefaultValue(); break;
+		case MP_Metallic:			InputPin->DefaultValue = EditorOnlyData->Metallic.GetDefaultValue(); break;
+		case MP_Specular:			InputPin->DefaultValue = EditorOnlyData->Specular.GetDefaultValue(); break;
+		case MP_Roughness:			InputPin->DefaultValue = EditorOnlyData->Roughness.GetDefaultValue(); break;
+		case MP_Anisotropy:			InputPin->DefaultValue = EditorOnlyData->Anisotropy.GetDefaultValue(); break;
+		case MP_Normal:				InputPin->DefaultValue = EditorOnlyData->Normal.GetDefaultValue(); break;
+		case MP_Tangent:			InputPin->DefaultValue = EditorOnlyData->Tangent.GetDefaultValue(); break;
+		case MP_WorldPositionOffset:InputPin->DefaultValue = EditorOnlyData->WorldPositionOffset.GetDefaultValue(); break;
+		case MP_SubsurfaceColor:	InputPin->DefaultValue = EditorOnlyData->SubsurfaceColor.GetDefaultValue(); break;
+		case MP_CustomData0:		InputPin->DefaultValue = EditorOnlyData->ClearCoat.GetDefaultValue(); break;
+		case MP_CustomData1:		InputPin->DefaultValue = EditorOnlyData->ClearCoatRoughness.GetDefaultValue(); break;
+		case MP_AmbientOcclusion:	InputPin->DefaultValue = EditorOnlyData->AmbientOcclusion.GetDefaultValue(); break;
+		case MP_Refraction:			InputPin->DefaultValue = EditorOnlyData->Refraction.GetDefaultValue(); break;
+		case MP_PixelDepthOffset:	InputPin->DefaultValue = EditorOnlyData->PixelDepthOffset.GetDefaultValue(); break;
+		case MP_ShadingModel:		break; // TODO
+		case MP_FrontMaterial:		break; // TODO
+		case MP_SurfaceThickness:	InputPin->DefaultValue = EditorOnlyData->SurfaceThickness.GetDefaultValue(); break;
+		case MP_Displacement:		InputPin->DefaultValue = EditorOnlyData->Displacement.GetDefaultValue(); break;
+		default:
+			if (Property >= MP_CustomizedUVs0 && Property <= MP_CustomizedUVs7)
+			{
+				InputPin->DefaultValue = EditorOnlyData->CustomizedUVs[Property - MP_CustomizedUVs0].GetDefaultValue(); break;
+			}
+			break;
+		}
 	}
 
 }
diff --git a/Engine/Source/Runtime/Engine/Private/Materials/MaterialShared.cpp b/Engine/Source/Runtime/Engine/Private/Materials/MaterialShared.cpp
--- a/Engine/Source/Runtime/Engine/Private/Materials/MaterialShared.cpp
+++ b/Engine/Source/Runtime/Engine/Private/Materials/MaterialShared.cpp
@@ -598,6 +598,7 @@ bool FSubstrateMaterialInput::Serialize(FArchive& Ar)
 	return SerializeMaterialInput<uint32>(Ar, *this);
 }
 
+
 bool FVectorMaterialInput::Serialize(FArchive& Ar)
 {
 	return SerializeMaterialInput<FVector3f>(Ar, *this);
@@ -613,6 +614,97 @@ bool FMaterialAttributesInput::Serialize(FArchive& Ar)
 	return SerializeExpressionInput(Ar, *this);
 }
 
+void FColorMaterialInput::DefaultValueChanged(const FString& DefaultValue)
+{
+#if WITH_EDITOR
+	FLinearColor Value;
+	Value.InitFromString(DefaultValue);
+	Constant = Value.QuantizeRound();
+	UseConstant = true;
+#endif
+}
+
+FString FColorMaterialInput::GetDefaultValue() const
+{
+	FString DefaultValue;
+#if WITH_EDITOR
+	if (UseConstant)
+	{
+		FLinearColor Color = Constant.ReinterpretAsLinear();
+		DefaultValue = Color.ToString();
+	}
+#endif
+	return DefaultValue;
+}
+
+void FScalarMaterialInput::DefaultValueChanged(const FString& DefaultValue)
+{
+#if WITH_EDITOR
+	Constant = FCString::Atof(*DefaultValue);
+	UseConstant = true;
+#endif
+}
+
+FString FScalarMaterialInput::GetDefaultValue() const
+{
+	FString DefaultValue;
+#if WITH_EDITOR
+	if (UseConstant)
+	{
+		DefaultValue = FString::SanitizeFloat(Constant);
+	}
+#endif
+	return DefaultValue;
+}
+
+void FVector2MaterialInput::DefaultValueChanged(const FString& DefaultValue)
+{
+#if WITH_EDITOR
+	FVector2f Value;
+	Value.InitFromString(DefaultValue);
+	Constant = Value;
+	UseConstant = true;
+#endif
+}
+
+FString FVector2MaterialInput::GetDefaultValue() const
+{
+	FString DefaultValue;
+#if WITH_EDITOR
+	if (UseConstant)
+	{
+		DefaultValue = FString(TEXT("(X=")) + FString::SanitizeFloat(Constant.X) + FString(TEXT(",Y=")) + FString::SanitizeFloat(Constant.Y) + FString(TEXT(")"));
+	}
+#endif
+	return DefaultValue;
+}
+
+void FVectorMaterialInput::DefaultValueChanged(const FString& DefaultValue)
+{
+#if WITH_EDITOR
+	//Parse string to split its contents separated by ','
+	TArray<FString> Elements;
+	DefaultValue.ParseIntoArray(Elements, TEXT(","), true);
+	check(Elements.Num() == 3);
+	Constant.X = FCString::Atof(*Elements[0]);
+	Constant.Y = FCString::Atof(*Elements[1]);
+	Constant.Z = FCString::Atof(*Elements[2]);
+	UseConstant = true;
+#endif
+}
+
+FString FVectorMaterialInput::GetDefaultValue() const
+{
+	FString DefaultValue;
+#if WITH_EDITOR
+	if (UseConstant)
+	{
+		DefaultValue = FString::SanitizeFloat(Constant.X) + FString(TEXT(",")) + FString::SanitizeFloat(Constant.Y) + FString(TEXT(",")) + FString::SanitizeFloat(Constant.Z);
+	}
+#endif
+	return DefaultValue;
+}
+
 #if WITH_EDITOR
 int32 FColorMaterialInput::CompileWithDefault(class FMaterialCompiler* Compiler, EMaterialProperty Property)
 {
diff --git a/Engine/Source/Runtime/Engine/Public/MaterialExpressionIO.h b/Engine/Source/Runtime/Engine/Public/MaterialExpressionIO.h
--- a/Engine/Source/Runtime/Engine/Public/MaterialExpressionIO.h
+++ b/Engine/Source/Runtime/Engine/Public/MaterialExpressionIO.h
@@ -192,6 +192,8 @@ struct FColorMaterialInput : FMaterialInput<FColor>
 #endif  // WITH_EDITOR
 	/** ICPPStructOps interface */
 	ENGINE_API bool Serialize(FArchive& Ar);
+	ENGINE_API void DefaultValueChanged(const FString& DefaultValue);
+	ENGINE_API FString GetDefaultValue() const;
 };
 
 template<>
@@ -213,6 +215,8 @@ struct FScalarMaterialInput : FMaterialInput<float>
 #endif  // WITH_EDITOR
 	/** ICPPStructOps interface */
 	ENGINE_API bool Serialize(FArchive& Ar);
+	ENGINE_API void DefaultValueChanged(const FString& DefaultValue);
+	ENGINE_API FString GetDefaultValue() const;
 };
 
 template<>
@@ -233,6 +237,8 @@ struct FShadingModelMaterialInput : FMaterialInput<uint32>
 #endif  // WITH_EDITOR
 	/** ICPPStructOps interface */
 	ENGINE_API bool Serialize(FArchive& Ar);
+	ENGINE_API void DefaultValueChanged(const FString& DefaultValue);
+	ENGINE_API FString GetDefaultValue() const;
 };
 
 template<>
@@ -252,6 +258,8 @@ struct FSubstrateMaterialInput : FMaterialInput<uint32> // Still giving it a def
 #endif  // WITH_EDITOR
 	/** ICPPStructOps interface */
 	ENGINE_API bool Serialize(FArchive& Ar);
+	ENGINE_API void DefaultValueChanged(const FString& DefaultValue);
+    ENGINE_API FString GetDefaultValue() const;
 };
 
 template<>
@@ -271,6 +279,8 @@ struct FVectorMaterialInput : FMaterialInput<FVector3f>
 #endif  // WITH_EDITOR
 	/** ICPPStructOps interface */
 	ENGINE_API bool Serialize(FArchive& Ar);
+	ENGINE_API void DefaultValueChanged(const FString& DefaultValue);
+	ENGINE_API FString GetDefaultValue() const;
 };
 
 template<>
@@ -291,6 +301,8 @@ struct FVector2MaterialInput : FMaterialInput<FVector2f>
 #endif  // WITH_EDITOR
 	/** ICPPStructOps interface */
 	ENGINE_API bool Serialize(FArchive& Ar);
+	ENGINE_API void DefaultValueChanged(const FString& DefaultValue);
+	ENGINE_API FString GetDefaultValue() const;
 };
 
 template<>
