Tooltip on Links for Material Expressions

diff --git a/Engine/Source/Editor/GraphEditor/Private/ConnectionDrawingPolicy.cpp b/Engine/Source/Editor/GraphEditor/Private/ConnectionDrawingPolicy.cpp
--- a/Engine/Source/Editor/GraphEditor/Private/ConnectionDrawingPolicy.cpp
+++ b/Engine/Source/Editor/GraphEditor/Private/ConnectionDrawingPolicy.cpp
@@ -2,6 +2,7 @@
 
 #include "ConnectionDrawingPolicy.h"
 #include "Rendering/DrawElements.h"
+#include "Widgets/SToolTip.h"
 #include "Framework/Application/SlateApplication.h"
 
 DEFINE_LOG_CATEGORY(LogConnectionDrawingPolicy);
@@ -625,3 +626,9 @@ void FGraphSplineOverlapResult::GetPinWidgets(const class SGraphPanel& InGraphPa
 		OutPin2 = Pin2Handle.FindInGraphPanel(InGraphPanel);
 	}
 }
+
+TSharedPtr<IToolTip> FConnectionDrawingPolicy::GetConnectionToolTip(const SGraphPanel& GraphPanel, const FGraphSplineOverlapResult& OverlapData) const
+{
+	TSharedPtr<SGraphPin> BestPinWidget = OverlapData.GetBestPinWidget(GraphPanel);
+	return BestPinWidget->GetToolTip();
+}
\ No newline at end of file
diff --git a/Engine/Source/Editor/GraphEditor/Private/MaterialGraphConnectionDrawingPolicy.cpp b/Engine/Source/Editor/GraphEditor/Private/MaterialGraphConnectionDrawingPolicy.cpp
--- a/Engine/Source/Editor/GraphEditor/Private/MaterialGraphConnectionDrawingPolicy.cpp
+++ b/Engine/Source/Editor/GraphEditor/Private/MaterialGraphConnectionDrawingPolicy.cpp
@@ -18,6 +18,9 @@
 #include "Templates/SharedPointer.h"
 #include "UObject/NameTypes.h"
 #include "SGraphSubstrateMaterial.h"
+#include "Widgets/SToolTip.h"
+#include "Widgets/SBoxPanel.h"
+#include "Widgets/Text/STextBlock.h"
 
 class FSlateRect;
 
@@ -37,8 +40,6 @@ FMaterialGraphConnectionDrawingPolicy::FMaterialGraphConnectionDrawingPolicy(int
 	HoverDeemphasisDarkFraction = 0.4f;
 }
 
-
-
 bool FMaterialGraphConnectionDrawingPolicy::FindPinCenter(UEdGraphPin* Pin, FVector2D& OutCenter) const
 {
 	if (const TSharedPtr<SGraphPin>* pPinWidget = PinToPinWidgetMap.Find(Pin))
@@ -230,3 +231,56 @@ void FMaterialGraphConnectionDrawingPolicy::DetermineWiringStyle(UEdGraphPin* Ou
 		ApplyHoverDeemphasis(OutputPin, InputPin, /*inout*/ Params.WireThickness, /*inout*/ Params.WireColor);
 	}
 }
+
+TSharedPtr<IToolTip> FMaterialGraphConnectionDrawingPolicy::GetConnectionToolTip(const SGraphPanel& GraphPanel,	const FGraphSplineOverlapResult& OverlapData) const
+{
+	TSharedPtr<SGraphPin> Pin1Widget;
+	TSharedPtr<SGraphPin> Pin2Widget;
+	OverlapData.GetPinWidgets(GraphPanel, Pin1Widget, Pin2Widget);
+
+	const FText LeftText = FText::Format(NSLOCTEXT("Unreal", "PinConnectionTooltipLeft", "<< {0}"), GetNodePinInfo(Pin1Widget));
+	const FText RightText = FText::Format(NSLOCTEXT("Unreal", "PinConnectionTooltipRight", "{0} >>"), GetNodePinInfo(Pin2Widget));
+
+	return SNew(SToolTip)
+		[
+			SNew(SVerticalBox)
+			+ SVerticalBox::Slot()
+			.HAlign(HAlign_Left)
+			.AutoHeight()
+			.Padding(0.0f, 0.0f, 0.0f, 5.0f)
+			[
+				SNew(STextBlock)
+					.Margin(FMargin(0.0f, 0.0f, 4.0f, 0.0f))
+					.Justification(ETextJustify::Left)
+					.Text(LeftText)
+			]
+			+ SVerticalBox::Slot()
+			.HAlign(HAlign_Right)
+			.AutoHeight()
+			.Padding(0.0f, 5.0f, 0.0f, 0.0f)
+			[
+				SNew(STextBlock)
+					.Margin(FMargin(4.0f, 0.0f, 0.0f, 0.0f))
+					.Justification(ETextJustify::Right)
+					.Text(RightText)
+			]
+		];
+}
+
+FText FMaterialGraphConnectionDrawingPolicy::GetNodePinInfo(const TSharedPtr<SGraphPin>& PinWidget) const
+{
+	const UEdGraphPin* PinObj = PinWidget->GetPinObj();
+	const UEdGraphNode* EdNode = PinObj->GetOwningNode();
+	 
+	FString NodeTitle = EdNode->GetNodeTitle(ENodeTitleType::ListView).ToString();
+	NodeTitle.RemoveFromStart(TEXT("Material Expression "));
+	
+	if (EdNode->GetCanRenameNode())
+	{
+		FText NodeEditableName = EdNode->GetNodeTitle(ENodeTitleType::EditableTitle);
+		NodeTitle = NodeTitle + TEXT(" (") + NodeEditableName.ToString() + TEXT(")");
+	}
+	const FText PinName = PinObj->GetDisplayName().IsEmptyOrWhitespace() ? FText::FromName(PinObj->PinName) : PinObj->GetDisplayName();
+	
+	return FText::Format(NSLOCTEXT("Unreal", "PinConnectionTooltipPartial", "{0}\r\n{1}"), FText::FromString(NodeTitle), /*NodeType,*/ PinName);
+}
diff --git a/Engine/Source/Editor/GraphEditor/Private/MaterialGraphConnectionDrawingPolicy.h b/Engine/Source/Editor/GraphEditor/Private/MaterialGraphConnectionDrawingPolicy.h
--- a/Engine/Source/Editor/GraphEditor/Private/MaterialGraphConnectionDrawingPolicy.h
+++ b/Engine/Source/Editor/GraphEditor/Private/MaterialGraphConnectionDrawingPolicy.h
@@ -35,7 +35,12 @@ protected:
 public:
 	FMaterialGraphConnectionDrawingPolicy(int32 InBackLayerID, int32 InFrontLayerID, float ZoomFactor, const FSlateRect& InClippingRect, FSlateWindowElementList& InDrawElements, UEdGraph* InGraphObj);
 
+	virtual TSharedPtr<IToolTip> GetConnectionToolTip(const SGraphPanel& GraphPanel, const FGraphSplineOverlapResult& OverlapData) const override;
+	
 	// FConnectionDrawingPolicy interface
 	virtual void DetermineWiringStyle(UEdGraphPin* OutputPin, UEdGraphPin* InputPin, /*inout*/ FConnectionParams& Params) override;
 	// End of FConnectionDrawingPolicy interface
+
+private:
+	FText GetNodePinInfo(const TSharedPtr<SGraphPin>& PinWidget) const;
 };
diff --git a/Engine/Source/Editor/GraphEditor/Private/SGraphPanel.cpp b/Engine/Source/Editor/GraphEditor/Private/SGraphPanel.cpp
--- a/Engine/Source/Editor/GraphEditor/Private/SGraphPanel.cpp
+++ b/Engine/Source/Editor/GraphEditor/Private/SGraphPanel.cpp
@@ -620,7 +620,12 @@ int32 SGraphPanel::OnPaint( const FPaintArgs& Args, const FGeometry& AllottedGeo
 			}
 
 			// Update the spline hover state
-			const_cast<SGraphPanel*>(this)->OnSplineHoverStateChanged(OverlapData);
+			if (const_cast<SGraphPanel*>(this)->OnSplineHoverStateChanged(OverlapData))
+			{
+				
+				// if hover state changed, we update the tooltip text based on the connection drawing policy
+				const_cast<SGraphPanel*>(this)->SetToolTip(ConnectionDrawingPolicy->GetConnectionToolTip(*this, OverlapData));
+			}
 		}
 
 		delete ConnectionDrawingPolicy;
@@ -666,7 +671,7 @@ int32 SGraphPanel::OnPaint( const FPaintArgs& Args, const FGeometry& AllottedGeo
 	return MaxLayerId;
 }
 
-void SGraphPanel::OnSplineHoverStateChanged(const FGraphSplineOverlapResult& NewSplineHoverState)
+bool SGraphPanel::OnSplineHoverStateChanged(const FGraphSplineOverlapResult& NewSplineHoverState)
 {
 	TSharedPtr<SGraphPin> OldPin1Widget;
 	TSharedPtr<SGraphPin> OldPin2Widget;
@@ -684,11 +689,15 @@ void SGraphPanel::OnSplineHoverStateChanged(const FGraphSplineOverlapResult& New
 	if (OldPin1Widget.IsValid() && OldPin1Widget != NewPin1Widget && OldPin1Widget != NewPin2Widget)
 	{
 		OldPin1Widget->OnMouseLeave(LastPointerEvent);
+		// reset connection tooltip if hover outside spline
+		SetToolTipText(FText());
 	}
 
 	if (OldPin2Widget.IsValid() && OldPin2Widget != NewPin1Widget && OldPin2Widget != NewPin2Widget)
 	{
 		OldPin2Widget->OnMouseLeave(LastPointerEvent);
+		// reset connection tooltip if hover outside spline
+		SetToolTipText(FText());
 	}
 
 	// Handle enter hovering on the pins
@@ -711,6 +720,7 @@ void SGraphPanel::OnSplineHoverStateChanged(const FGraphSplineOverlapResult& New
 		//@TODO: Source this parameter from the graph rendering settings once it is there (see code in ApplyHoverDeemphasis)
 		TimeWhenMouseEnteredPin -= 0.75f;
 	}
+	return bChangedHover;
 }
 
 bool SGraphPanel::SupportsKeyboardFocus() const
@@ -741,16 +751,6 @@ void SGraphPanel::OnArrangeChildren( const FGeometry& AllottedGeometry, FArrange
 	ArrangedChildren.Append(MyArrangedChildren);
 }
 
-TSharedPtr<IToolTip> SGraphPanel::GetToolTip()
-{
-	if (SGraphPin* BestPinFromHoveredSpline = GetBestPinFromHoveredSpline())
-	{
-		return BestPinFromHoveredSpline->GetToolTip();
-	}
-
-	return SNodePanel::GetToolTip();
-}
-
 void SGraphPanel::UpdateSelectedNodesPositions(FVector2D PositionIncrement)
 {
 	FScopedTransaction Transaction(NSLOCTEXT("GraphEditor", "NudgeNodeAction", "Nudge Node"));
diff --git a/Engine/Source/Editor/GraphEditor/Public/ConnectionDrawingPolicy.h b/Engine/Source/Editor/GraphEditor/Public/ConnectionDrawingPolicy.h
--- a/Engine/Source/Editor/GraphEditor/Public/ConnectionDrawingPolicy.h
+++ b/Engine/Source/Editor/GraphEditor/Public/ConnectionDrawingPolicy.h
@@ -134,7 +134,7 @@ protected:
 	TArray<UEdGraphNode*> SelectedGraphNodes;
 public:
 	virtual ~FConnectionDrawingPolicy() {}
-
+	
 	FConnectionDrawingPolicy(int32 InBackLayerID, int32 InFrontLayerID, float InZoomFactor, const FSlateRect& InClippingRect, FSlateWindowElementList& InDrawElements);
 
 	// Update the drawing policy with the set of hovered pins (which can be empty)
@@ -185,6 +185,8 @@ public:
 
 	virtual bool IsConnectionCulled( const FArrangedWidget& StartLink, const FArrangedWidget& EndLink ) const;
 
+	virtual TSharedPtr<IToolTip> GetConnectionToolTip(const SGraphPanel& GraphPanel, const FGraphSplineOverlapResult& OverlapData) const;
+	
 protected:
 	// Helper function used by Draw(). Called before DrawPinGeometries to populate PinToPinWidgetMap
 	virtual void BuildPinToPinWidgetMap(TMap<TSharedRef<SWidget>, FArrangedWidget>& InPinGeometries);
diff --git a/Engine/Source/Editor/GraphEditor/Public/SGraphPanel.h b/Engine/Source/Editor/GraphEditor/Public/SGraphPanel.h
--- a/Engine/Source/Editor/GraphEditor/Public/SGraphPanel.h
+++ b/Engine/Source/Editor/GraphEditor/Public/SGraphPanel.h
@@ -138,7 +138,6 @@ public:
 	virtual FReply OnKeyDown( const FGeometry& MyGeometry, const FKeyEvent& InKeyEvent ) override;
 	virtual bool SupportsKeyboardFocus() const override;
 	virtual void OnArrangeChildren( const FGeometry& AllottedGeometry, FArrangedChildren& ArrangedChildren ) const override;
-	virtual TSharedPtr<IToolTip> GetToolTip() override;
 	// End of SWidget interface
 
 	// SNodePanel interface
@@ -369,7 +368,7 @@ private:
 	void UpdateSelectedNodesPositions(FVector2D PositionIncrement);
 
 	/** Handle updating the spline hover state */
-	void OnSplineHoverStateChanged(const FGraphSplineOverlapResult& NewSplineHoverState);
+	bool OnSplineHoverStateChanged(const FGraphSplineOverlapResult& NewSplineHoverState);
 
 	/** Returns the pin that we're considering as hovered if we are hovering over a spline; may be null */
 	class SGraphPin* GetBestPinFromHoveredSpline() const;
