Tooltip on Links for Material Expressions

diff --git a/Engine/Source/Editor/GraphEditor/Private/MaterialGraphConnectionDrawingPolicy.cpp b/Engine/Source/Editor/GraphEditor/Private/MaterialGraphConnectionDrawingPolicy.cpp
--- a/Engine/Source/Editor/GraphEditor/Private/MaterialGraphConnectionDrawingPolicy.cpp
+++ b/Engine/Source/Editor/GraphEditor/Private/MaterialGraphConnectionDrawingPolicy.cpp
@@ -238,8 +238,8 @@ TSharedPtr<IToolTip> FMaterialGraphConnectionDrawingPolicy::GetConnectionToolTip
 	TSharedPtr<SGraphPin> Pin2Widget;
 	OverlapData.GetPinWidgets(GraphPanel, Pin1Widget, Pin2Widget);
 
-	const FText LeftText = FText::Format(NSLOCTEXT("Unreal", "PinConnectionTooltipLeft", "<< {0}"), GetNodePinInfo(Pin1Widget));
-	const FText RightText = FText::Format(NSLOCTEXT("Unreal", "PinConnectionTooltipRight", "{0} >>"), GetNodePinInfo(Pin2Widget));
+	const FText LeftText = Pin1Widget != nullptr ? FText::Format(NSLOCTEXT("Unreal", "PinConnectionTooltipLeft", "<< {0}"), GetNodePinInfo(Pin1Widget)) : FText();
+	const FText RightText = Pin2Widget != nullptr ? FText::Format(NSLOCTEXT("Unreal", "PinConnectionTooltipRight", "{0} >>"), GetNodePinInfo(Pin2Widget)) : FText();
 
 	return SNew(SToolTip)
 		[
