Bug fix for inline node editing where entering 0 wasn't working as expected in some cases.

diff --git a/Engine/Source/Editor/GraphEditor/Public/KismetPins/SGraphPinNumSlider.h b/Engine/Source/Editor/GraphEditor/Public/KismetPins/SGraphPinNumSlider.h
--- a/Engine/Source/Editor/GraphEditor/Public/KismetPins/SGraphPinNumSlider.h
+++ b/Engine/Source/Editor/GraphEditor/Public/KismetPins/SGraphPinNumSlider.h
@@ -41,6 +41,9 @@ protected:
 		
 		const bool bAllowSpin = !(PinProperty && PinProperty->GetBoolMetaData("NoSpinbox"));
 
+		// Save last committed value to compare when value changes
+		LastSliderCommittedValue = GetNumericValue().GetValue();
+		
 		return SNew(SBox)
 			.MinDesiredWidth(18)
 			.MaxDesiredWidth(400)
