Tooltip crash fix for editors not having a "BestPinWidget"

diff --git a/Engine/Source/Editor/GraphEditor/Private/ConnectionDrawingPolicy.cpp b/Engine/Source/Editor/GraphEditor/Private/ConnectionDrawingPolicy.cpp
--- a/Engine/Source/Editor/GraphEditor/Private/ConnectionDrawingPolicy.cpp
+++ b/Engine/Source/Editor/GraphEditor/Private/ConnectionDrawingPolicy.cpp
@@ -1,6 +1,7 @@
 // Copyright Epic Games, Inc. All Rights Reserved.
 
 #include "ConnectionDrawingPolicy.h"
+#include "SGraphPanel.h"
 #include "Rendering/DrawElements.h"
 #include "Widgets/SToolTip.h"
 #include "Framework/Application/SlateApplication.h"
@@ -629,6 +630,9 @@ void FGraphSplineOverlapResult::GetPinWidgets(const class SGraphPanel& InGraphPa
 
 TSharedPtr<IToolTip> FConnectionDrawingPolicy::GetConnectionToolTip(const SGraphPanel& GraphPanel, const FGraphSplineOverlapResult& OverlapData) const
 {
-	TSharedPtr<SGraphPin> BestPinWidget = OverlapData.GetBestPinWidget(GraphPanel);
-	return BestPinWidget->GetToolTip();
+	if (SGraphPin* BestPinFromHoveredSpline = OverlapData.GetBestPinWidget(GraphPanel).Get())
+	{
+		return BestPinFromHoveredSpline->GetToolTip();
+	}
+	return const_cast<SGraphPanel&>(GraphPanel).GetToolTip();
 }
\ No newline at end of file
