PR #10992: Hover shortcut to use "Straighten Connections" on an outgoing pin

diff --git a/Engine/Source/Editor/GraphEditor/Private/SGraphPanel.cpp b/Engine/Source/Editor/GraphEditor/Private/SGraphPanel.cpp
--- a/Engine/Source/Editor/GraphEditor/Private/SGraphPanel.cpp
+++ b/Engine/Source/Editor/GraphEditor/Private/SGraphPanel.cpp
@@ -1941,6 +1941,8 @@ private:
 
 void SGraphPanel::StraightenConnections()
 {
+	bool bHasAlignedNodes = false;
+	
 	FConnectionAligner Aligner;
 	for (auto& It : SelectionManager.SelectedNodes)
 	{
@@ -1974,12 +1976,59 @@ void SGraphPanel::StraightenConnections()
 					
 					if (PinWidget.IsValid() && LinkedPinWidget.IsValid())
 					{
+						bHasAlignedNodes = true;
 						Aligner.DefineConnection(SourceNode, PinWidget, DestNode, LinkedPinWidget);
 					}
 				}
 			}
 		}
 	}
+	
+	// If we aren't aligning selected nodes, try to align a hovered Single Pin (non-knot) connected nodes.
+	if (!bHasAlignedNodes && CurrentHoveredPins.Num() > 0)
+	{
+		UEdGraphPin* SourcePin = nullptr;
+		for (const FEdGraphPinReference& CurrentHoverPin : CurrentHoveredPins)
+		{
+			int32 InputPinIndex = INDEX_NONE;
+			int32 OutputPinIndex = INDEX_NONE;
+			UEdGraphNode* InKnot = CurrentHoverPin.Get()->GetOwningNodeUnchecked();
+			bool bIsKnot = (InKnot != nullptr && InKnot->ShouldDrawNodeAsControlPointOnly(InputPinIndex, OutputPinIndex) == true &&
+				InputPinIndex >= 0 && OutputPinIndex >= 0);
+	
+			//only use the actual node pins and not knot pins
+			if (!bIsKnot)
+			{
+				SourcePin = CurrentHoverPin.Get();
+			}
+		}
+
+		if (SourcePin)
+		{
+			UEdGraphNode* SourceNode = SourcePin->GetOwningNode();
+			if (SourceNode)
+			{
+				UEdGraphPin* DestPin = (SourcePin->LinkedTo.Num() == 1) ? SourcePin->LinkedTo[0] : nullptr;
+				UEdGraphNode* DestNode = DestPin ? DestPin->GetOwningNode() : nullptr;
+				if (DestPin && DestNode)
+				{
+					TSharedRef<SNode>* SrcNodePtr = NodeToWidgetLookup.Find(SourceNode);
+					TSharedRef<SNode>* DstNodePtr = NodeToWidgetLookup.Find(DestNode);
+
+					if (SrcNodePtr && DstNodePtr)
+					{
+						TSharedPtr<SGraphPin> PinWidget = StaticCastSharedRef<SGraphNode>(*SrcNodePtr)->FindWidgetForPin(SourcePin);
+						TSharedPtr<SGraphPin> LinkedPinWidget = StaticCastSharedRef<SGraphNode>(*DstNodePtr)->FindWidgetForPin(DestPin);
+			
+						if (PinWidget.IsValid() && LinkedPinWidget.IsValid())
+						{
+							Aligner.DefineConnection(SourceNode, PinWidget, DestNode, LinkedPinWidget);
+						}
+					}
+				}
+			}
+		}
+	}
 
 	Aligner.Process();
 }
