use SGraphPinNumSlider instead of STG_GraphPinNum

diff --git a/Engine/Source/Editor/GraphEditor/Public/KismetPins/SGraphPinNumSlider.h b/Engine/Source/Editor/GraphEditor/Public/KismetPins/SGraphPinNumSlider.h
--- a/Engine/Source/Editor/GraphEditor/Public/KismetPins/SGraphPinNumSlider.h
+++ b/Engine/Source/Editor/GraphEditor/Public/KismetPins/SGraphPinNumSlider.h
@@ -17,13 +17,20 @@ template <typename NumericType>
 class SGraphPinNumSlider : public SGraphPin
 {
 public:
-	SLATE_BEGIN_ARGS(SGraphPinNumSlider) 
+	SLATE_BEGIN_ARGS(SGraphPinNumSlider)
+		:_MinDesiredBoxWidth(18.0f)
+		,_ShouldShowDisabledWhenConnected(false)
 	{}
+	
+		SLATE_ARGUMENT(float, MinDesiredBoxWidth)
+		SLATE_ARGUMENT(bool, ShouldShowDisabledWhenConnected)
 	SLATE_END_ARGS()
 
 	void Construct(const FArguments& InArgs, UEdGraphPin* InGraphPinObj, FProperty* InProperty)
 	{
 		PinProperty = InProperty;
+		MinDesiredBoxWidth = InArgs._MinDesiredBoxWidth;
+		bShouldShowDisabledWhenConnected = InArgs._ShouldShowDisabledWhenConnected;
 		SGraphPin::Construct(SGraphPin::FArguments(), InGraphPinObj);
 
 	}
@@ -45,11 +52,10 @@ protected:
 		LastSliderCommittedValue = GetNumericValue().GetValue();
 		
 		return SNew(SBox)
-			.MinDesiredWidth(18)
+			.MinDesiredWidth(MinDesiredBoxWidth)
 			.MaxDesiredWidth(400)
 			[
 				SNew(SNumericEntryBox<NumericType>)
-				.MinDesiredValueWidth(100.0f)
 				.EditableTextBoxStyle(FAppStyle::Get(), "Graph.EditableTextBox")
 				.BorderForegroundColor(FSlateColor::UseForeground())
 				.Visibility(this, &SGraphPinNumSlider::GetDefaultValueVisibility)
@@ -117,9 +123,43 @@ protected:
 		return bIsUsingSlider ? SliderValue : Num;
 	}
 
+	virtual EVisibility GetDefaultValueVisibility() const override
+	{
+		// If this is only for showing default value, always show
+		if (bOnlyShowDefaultValue)
+		{
+			return EVisibility::Visible;
+		}
+
+		// First ask schema
+		UEdGraphPin* GraphPin = GetPinObj();
+		const UEdGraphSchema* Schema = (GraphPin && !GraphPin->IsPendingKill()) ? GraphPin->GetSchema() : nullptr;
+		if (Schema == nullptr || Schema->ShouldHidePinDefaultValue(GraphPin))
+		{
+			return EVisibility::Collapsed;
+		}
+
+		if (GraphPin->bNotConnectable && !GraphPin->bOrphanedPin)
+		{
+			// The only reason this pin exists is to show something, so do so
+			return EVisibility::Visible;
+		}
+
+		if (GraphPin->Direction == EGPD_Output)
+		{
+			//@TODO: Should probably be a bLiteralOutput flag or a Schema call
+			return EVisibility::Collapsed;
+		}
+		else
+		{
+			return IsConnected() && !bShouldShowDisabledWhenConnected ? EVisibility::Collapsed : EVisibility::Visible;
+		}
+	}
 private:
 	FProperty* PinProperty;
 	NumericType LastSliderCommittedValue;
 	NumericType SliderValue;
+	float MinDesiredBoxWidth = 0;
 	bool bIsUsingSlider = false;
+	bool bShouldShowDisabledWhenConnected = false;
 };
