Fix BT decorator and service lists (and also some custom licensee graph implementations) put the first and second level category on the same level.

A very long time ago CL 2411381 introduced a change to make the default Base Indent Level of all Graph Action Menus "1", basically meaning that both "0" and "1" indent levels would be placed at the root of the menu. Presumably they've needed it, because in BPs, categories on the left side are custom headers, and they didn't want the children (functions, variables, etc) to have an extra indent.
However, they've made this change globally, basically breaking Graph Action Menus that didn't have category customization and needed the first level to be indented under the root level.
Some systems have overcome this by defining a custom expander with a Base Indent Level of 0 - getting back the original, expected behavior. That was however just hiding the underlying issue.
This change reverts back the Base Indent Level to 0, and sets it to 1 instead only for those menus, which have a customized category header. This should fix all graph menus where the second category level was not indented (i.e. BT Action Menu and Decorator Context Menu).

diff --git a/Engine/Plugins/Animation/DeformerGraph/Source/OptimusEditor/Private/Widgets/SOptimusEditorGraphExplorer.cpp b/Engine/Plugins/Animation/DeformerGraph/Source/OptimusEditor/Private/Widgets/SOptimusEditorGraphExplorer.cpp
--- a/Engine/Plugins/Animation/DeformerGraph/Source/OptimusEditor/Private/Widgets/SOptimusEditorGraphExplorer.cpp
+++ b/Engine/Plugins/Animation/DeformerGraph/Source/OptimusEditor/Private/Widgets/SOptimusEditorGraphExplorer.cpp
@@ -211,6 +211,7 @@ void SOptimusEditorGraphExplorer::CreateWidgets()
 	    .OnCanRenameSelectedAction(this, &SOptimusEditorGraphExplorer::CanRequestRenameOnActionNode)
 	    .OnGetSectionTitle(this, &SOptimusEditorGraphExplorer::OnGetSectionTitle)
 	    .OnGetSectionWidget(this, &SOptimusEditorGraphExplorer::OnGetSectionWidget)
+		.DefaultRowExpanderBaseIndentLevel(1)
 	    .AlphaSortItems(false)
 	    .UseSectionStyling(true);
 
diff --git a/Engine/Source/Editor/BehaviorTreeEditor/Private/SBehaviorTreeBlackboardView.cpp b/Engine/Source/Editor/BehaviorTreeEditor/Private/SBehaviorTreeBlackboardView.cpp
--- a/Engine/Source/Editor/BehaviorTreeEditor/Private/SBehaviorTreeBlackboardView.cpp
+++ b/Engine/Source/Editor/BehaviorTreeEditor/Private/SBehaviorTreeBlackboardView.cpp
@@ -526,6 +526,7 @@ void SBehaviorTreeBlackboardView::Construct(const FArguments& InArgs, TSharedRef
 				.OnContextMenuOpening(this, &SBehaviorTreeBlackboardView::HandleContextMenuOpening, InCommandList)
 				.OnActionMatchesName(this, &SBehaviorTreeBlackboardView::HandleActionMatchesName)
 				.AlphaSortItems(GetDefault<UEditorPerProjectUserSettings>()->bDisplayBlackboardKeysInAlphabeticalOrder)
+				.DefaultRowExpanderBaseIndentLevel(1)
 				.AutoExpandActionMenu(true)
 			]
 		]
diff --git a/Engine/Source/Editor/GraphEditor/Private/SGraphActionMenu.cpp b/Engine/Source/Editor/GraphEditor/Private/SGraphActionMenu.cpp
--- a/Engine/Source/Editor/GraphEditor/Private/SGraphActionMenu.cpp
+++ b/Engine/Source/Editor/GraphEditor/Private/SGraphActionMenu.cpp
@@ -311,6 +311,7 @@ void SGraphActionMenu::Construct( const FArguments& InArgs, bool bIsReadOnly/* =
 	this->bIgnoreUIUpdate = false;
 	this->bUseSectionStyling = InArgs._UseSectionStyling;
 	this->bAllowPreselectedItemActivation = InArgs._bAllowPreselectedItemActivation;
+	this->DefaultRowExpanderBaseIndentLevel = InArgs._DefaultRowExpanderBaseIndentLevel;
 
 	this->bAutoExpandActionMenu = InArgs._AutoExpandActionMenu;
 	this->bShowFilterTextBox = InArgs._ShowFilterTextBox;
@@ -1022,7 +1023,7 @@ TSharedRef<ITableRow> SGraphActionMenu::MakeWidget( TSharedPtr<FGraphActionNode>
 	{
 		ExpanderWidget =
 			SNew(SExpanderArrow, TableRow)
-			.BaseIndentLevel(1);
+			.BaseIndentLevel(DefaultRowExpanderBaseIndentLevel);
 	}
 
 	RowContainer->AddSlot()
diff --git a/Engine/Source/Editor/GraphEditor/Public/SGraphActionMenu.h b/Engine/Source/Editor/GraphEditor/Public/SGraphActionMenu.h
--- a/Engine/Source/Editor/GraphEditor/Public/SGraphActionMenu.h
+++ b/Engine/Source/Editor/GraphEditor/Public/SGraphActionMenu.h
@@ -152,6 +152,7 @@ public:
 		, _ShowFilterTextBox(true)
 		, _UseSectionStyling(false)
 		, _bAllowPreselectedItemActivation(false)
+		, _DefaultRowExpanderBaseIndentLevel(0)
 		, _GraphObj(nullptr)
 		{ }
 
@@ -178,6 +179,7 @@ public:
 		SLATE_ARGUMENT( bool, ShowFilterTextBox )
 		SLATE_ARGUMENT( bool, UseSectionStyling )
 		SLATE_ARGUMENT( bool, bAllowPreselectedItemActivation )
+		SLATE_ARGUMENT( int32, DefaultRowExpanderBaseIndentLevel )
 		SLATE_ARGUMENT( TArray<UEdGraphPin*>, DraggedFromPins )
 		SLATE_ARGUMENT( UEdGraph*, GraphObj )
 
@@ -237,6 +239,8 @@ protected:
 	bool bUseSectionStyling;
 	/** Whether we allow pre-selected items to be activated with a left-click */
 	bool bAllowPreselectedItemActivation;
+	/** The BaseIndentLevel of the default-created row expander. Not used with OnCreateCustomRowExpander. */
+	int32 DefaultRowExpanderBaseIndentLevel;
 	
 	/** Delegate to call when action is selected */
 	FOnActionSelected OnActionSelected;
diff --git a/Engine/Source/Editor/Kismet/Private/SMyBlueprint.cpp b/Engine/Source/Editor/Kismet/Private/SMyBlueprint.cpp
--- a/Engine/Source/Editor/Kismet/Private/SMyBlueprint.cpp
+++ b/Engine/Source/Editor/Kismet/Private/SMyBlueprint.cpp
@@ -541,6 +541,7 @@ void SMyBlueprint::Construct(const FArguments& InArgs, TWeakPtr<FBlueprintEditor
 		.OnGetSectionTitle(this, &SMyBlueprint::OnGetSectionTitle)
 		.OnGetSectionWidget(this, &SMyBlueprint::OnGetSectionWidget)
 		.OnActionMatchesName(this, &SMyBlueprint::HandleActionMatchesName)
+		.DefaultRowExpanderBaseIndentLevel(1)
 		.AlphaSortItems(false)
 		.UseSectionStyling(true);
 
