Fix node slider tweaking not updating texture graph

diff --git a/Engine/Source/Editor/GraphEditor/Public/KismetPins/SGraphPinNumSlider.h b/Engine/Source/Editor/GraphEditor/Public/KismetPins/SGraphPinNumSlider.h
--- a/Engine/Source/Editor/GraphEditor/Public/KismetPins/SGraphPinNumSlider.h
+++ b/Engine/Source/Editor/GraphEditor/Public/KismetPins/SGraphPinNumSlider.h
@@ -81,6 +81,7 @@ protected:
 	void OnValueChanged( NumericType NewValue )
 	{
 		SliderValue = NewValue;
+		GraphPinObj->GetSchema()->TrySetDefaultValue(*GraphPinObj, *LexToSanitizedString(NewValue), false);
 	}
 
 	void OnValueCommitted( NumericType NewValue, ETextCommit::Type CommitInfo )
