Diff tool - When matching pins before and after a change, prefer pins that match by name and id, and relegate name-only matches as a fallback.
This fixes false positives when a pin is connected to several pins with the same name

diff --git a/Engine/Source/Editor/GraphEditor/Private/GraphDiffControl.cpp b/Engine/Source/Editor/GraphEditor/Private/GraphDiffControl.cpp
--- a/Engine/Source/Editor/GraphEditor/Private/GraphDiffControl.cpp
+++ b/Engine/Source/Editor/GraphEditor/Private/GraphDiffControl.cpp
@@ -468,16 +468,23 @@ static UEdGraphPin* FindOtherLink(TArray<UEdGraphPin*>& Links2, int32 OriginalIn
 {
 	// Sometimes the order of the pins is different between revisions, although the pins themselves are unchanged, so we have to look at all of them
 	UEdGraphNode* Node1 = PinToFind->GetOwningNode();
+	UEdGraphPin* BestMatch = Links2[OriginalIndex];
 	for (UEdGraphPin* Other : Links2)
 	{
 		UEdGraphNode* Node2 = Other->GetOwningNode();
-		const bool bPinsMatch = Other->GetName() == PinToFind->GetName() || Other->PinId == PinToFind->PinId;
-		if(FGraphDiffControl::IsNodeMatch(Node1, Node2) && bPinsMatch)
+		if(FGraphDiffControl::IsNodeMatch(Node1, Node2))
 		{
-			return Other;
+			if (Other->PinId == PinToFind->PinId)
+			{
+				return Other;
+			}
+			else if (Other->GetName() == PinToFind->GetName())
+			{
+				BestMatch = Other;
+			}
 		}
 	}
-	return Links2[OriginalIndex];
+	return BestMatch;
 }
 
 /** Determine if the LinkedTo pins are the same */
