Blueprint editor marching ants (execution path) are now displayed for function library functions

diff --git a/Engine/Source/Editor/GraphEditor/Private/BlueprintConnectionDrawingPolicy.cpp b/Engine/Source/Editor/GraphEditor/Private/BlueprintConnectionDrawingPolicy.cpp
--- a/Engine/Source/Editor/GraphEditor/Private/BlueprintConnectionDrawingPolicy.cpp
+++ b/Engine/Source/Editor/GraphEditor/Private/BlueprintConnectionDrawingPolicy.cpp
@@ -130,12 +130,7 @@ UBlueprint* FKismetConnectionDrawingPolicy::GetTargetBlueprint() const
 
 bool FKismetConnectionDrawingPolicy::CanBuildRoadmap(UBlueprint* TargetBP) const
 {
-	UObject* ActiveObject = nullptr;
-	if (TargetBP)
-	{
-		ActiveObject = TargetBP->GetObjectBeingDebugged();
-	}
-
+	UObject* ActiveObject = GetObjectBeingDebugged(TargetBP);
 	return ActiveObject != nullptr;
 }
 
@@ -156,7 +151,7 @@ void FKismetConnectionDrawingPolicy::BuildExecutionRoadmap()
 		return;
 	}
 
-	UObject* ActiveObject = TargetBP->GetObjectBeingDebugged();
+	UObject* ActiveObject = GetObjectBeingDebugged(TargetBP);
 	check(ActiveObject); // Due to CanBuildRoadmap
 
 	TArray<UEdGraphNode*> SequentialNodesInGraph;
@@ -401,6 +396,21 @@ bool FKismetConnectionDrawingPolicy::FindPinCenter(UEdGraphPin* Pin, FVector2D&
 	return false;
 }
 
+UObject* FKismetConnectionDrawingPolicy::GetObjectBeingDebugged(UBlueprint* TargetBP)
+{
+	UObject* ActiveObject = nullptr;
+	if (TargetBP)
+	{
+		ActiveObject = TargetBP->GetObjectBeingDebugged();
+		if (UClass* GeneratedClass = TargetBP->GeneratedClass;
+			ActiveObject == nullptr && BPTYPE_FunctionLibrary == TargetBP->BlueprintType && GeneratedClass)
+		{
+			ActiveObject = GeneratedClass->ClassDefaultObject;
+		}
+	}
+	return ActiveObject;
+}
+
 bool FKismetConnectionDrawingPolicy::GetAverageConnectedPosition(class UK2Node_Knot* Knot, EEdGraphPinDirection Direction, FVector2D& OutPos) const
 {
 	FVector2D Result = FVector2D::ZeroVector;
diff --git a/Engine/Source/Editor/GraphEditor/Public/BlueprintConnectionDrawingPolicy.h b/Engine/Source/Editor/GraphEditor/Public/BlueprintConnectionDrawingPolicy.h
--- a/Engine/Source/Editor/GraphEditor/Public/BlueprintConnectionDrawingPolicy.h
+++ b/Engine/Source/Editor/GraphEditor/Public/BlueprintConnectionDrawingPolicy.h
@@ -140,4 +140,5 @@ protected:
 	bool ShouldChangeTangentForKnot(class UK2Node_Knot* Knot);
 	bool GetAverageConnectedPosition(class UK2Node_Knot* Knot, EEdGraphPinDirection Direction, FVector2D& OutPos) const;
 	bool FindPinCenter(UEdGraphPin* Pin, FVector2D& OutCenter) const;
+	static UObject* GetObjectBeingDebugged(UBlueprint* TargetBP);
 };
diff --git a/Engine/Source/Editor/UnrealEd/Private/Kismet2/KismetDebugUtilities.cpp b/Engine/Source/Editor/UnrealEd/Private/Kismet2/KismetDebugUtilities.cpp
--- a/Engine/Source/Editor/UnrealEd/Private/Kismet2/KismetDebugUtilities.cpp
+++ b/Engine/Source/Editor/UnrealEd/Private/Kismet2/KismetDebugUtilities.cpp
@@ -254,6 +254,11 @@ void FKismetDebugUtilities::OnScriptException(const UObject* ActiveObject, const
 		const FBlueprintExceptionInfo* ExceptionInfo = &Info;
 		bool bResetObjectBeingDebuggedWhenFinished = false;
 		UObject* ObjectBeingDebugged = BlueprintObj->GetObjectBeingDebugged();
+		if (UClass* GeneratedClass = BlueprintObj->GeneratedClass; 
+			ObjectBeingDebugged == nullptr && BPTYPE_FunctionLibrary == BlueprintObj->BlueprintType && GeneratedClass)
+		{
+			ObjectBeingDebugged = GeneratedClass->ClassDefaultObject;
+		}
 
 		auto IsAPreviewOrInactiveObject = [](const UObject* InObject)
 		{
