Minor revision to SGraphPanel to decouple its refresh on updates to the underlying graph object from a dependency on a parent SGraphEditor.

This allows SGraphPanel widgets that live outside of an editor context (e.g. previews) to avoid node/pin references going stale if the underlying graph object is reconstructed.

diff --git a/Engine/Source/Editor/GraphEditor/Private/SGraphEditorImpl.cpp b/Engine/Source/Editor/GraphEditor/Private/SGraphEditorImpl.cpp
--- a/Engine/Source/Editor/GraphEditor/Private/SGraphEditorImpl.cpp
+++ b/Engine/Source/Editor/GraphEditor/Private/SGraphEditorImpl.cpp
@@ -178,45 +178,15 @@ FReply SGraphEditorImpl::OnKeyDown( const FGeometry& MyGeometry, const FKeyEvent
 
 void SGraphEditorImpl::NotifyGraphChanged()
 {
-	FEdGraphEditAction DefaultAction;
-	OnGraphChanged(DefaultAction);
+	GetCurrentGraph()->NotifyGraphChanged();
 }
 
 void SGraphEditorImpl::OnGraphChanged(const FEdGraphEditAction& InAction)
 {
-	if ( !bIsActiveTimerRegistered )
+	const bool bWasAddAction = (InAction.Action & GRAPHACTION_AddNode) != 0;
+	if (bWasAddAction)
 	{
-		const UEdGraphSchema* Schema = EdGraphObj->GetSchema();
-		const bool bSchemaRequiresFullRefresh = Schema->ShouldAlwaysPurgeOnModification();
-
-		const bool bWasAddAction = (InAction.Action & GRAPHACTION_AddNode) != 0;
-		const bool bWasSelectAction = (InAction.Action & GRAPHACTION_SelectNode) != 0;
-		const bool bWasRemoveAction = (InAction.Action & GRAPHACTION_RemoveNode) != 0;
-		const bool bWasEditAction = (InAction.Action & GRAPHACTION_EditNode) != 0;
-
-		// If we did a 'default action' (or some other action not handled by SGraphPanel::OnGraphChanged
-		// or if we're using a schema that always needs a full refresh, then purge the current nodes
-		// and queue an update:
-		if (bSchemaRequiresFullRefresh || 
-			(!bWasAddAction && !bWasSelectAction && !bWasRemoveAction && !bWasEditAction) )
-		{
-			GraphPanel->PurgeVisualRepresentation();
-			// Trigger the refresh
-			bIsActiveTimerRegistered = true;
-			RegisterActiveTimer(0.f, FWidgetActiveTimerDelegate::CreateSP(this, &SGraphEditorImpl::TriggerRefresh));
-		}
-
-		if (bWasAddAction)
-		{
-			NumNodesAddedSinceLastPointerPosition++;
-		}
-		else if (bWasEditAction && !bSchemaRequiresFullRefresh)
-		{
-			for (const UEdGraphNode* Node : InAction.Nodes)
-			{
-				RefreshNode(const_cast<UEdGraphNode&>(*Node));
-			}
-		}
+		NumNodesAddedSinceLastPointerPosition++;
 	}
 }
 
@@ -660,7 +630,6 @@ void SGraphEditorImpl::Construct( const FArguments& InArgs )
 	OnNavigateHistoryForward = InArgs._OnNavigateHistoryForward;
 	OnNodeSpawnedByKeymap = InArgs._GraphEvents.OnNodeSpawnedByKeymap;
 
-	bIsActiveTimerRegistered = false;
 	NumNodesAddedSinceLastPointerPosition = 0;
 
 	// Make sure that the editor knows about what kinds
@@ -915,14 +884,6 @@ void SGraphEditorImpl::Tick( const FGeometry& AllottedGeometry, const double InC
 	}
 }
 
-EActiveTimerReturnType SGraphEditorImpl::TriggerRefresh( double InCurrentTime, float InDeltaTime )
-{
-	GraphPanel->Update();
-
-	bIsActiveTimerRegistered = false;
-	return EActiveTimerReturnType::Stop;
-}
-
 void SGraphEditorImpl::OnClosedActionMenu()
 {
 	GraphPanel->OnStopMakingConnection(/*bForceStop=*/ true);
diff --git a/Engine/Source/Editor/GraphEditor/Private/SGraphEditorImpl.h b/Engine/Source/Editor/GraphEditor/Private/SGraphEditorImpl.h
--- a/Engine/Source/Editor/GraphEditor/Private/SGraphEditorImpl.h
+++ b/Engine/Source/Editor/GraphEditor/Private/SGraphEditorImpl.h
@@ -378,9 +378,6 @@ protected:
 	// End of SGraphEditorInterface
 
 private:
-	/** One-off active timer to ensure the graph is refreshes as needed */
-	EActiveTimerReturnType TriggerRefresh( double InCurrentTime, float InDeltaTime );
-
 	FText GetZoomText() const;
 
 	FSlateColor GetZoomTextColorAndOpacity() const;
@@ -412,7 +409,6 @@ private:
 	EActiveTimerReturnType HandleFocusEditorDeferred(double InCurrentTime, float InDeltaTime);
 private:
 	TWeakPtr<FActiveTimerHandle> FocusEditorTimer;
-	bool bIsActiveTimerRegistered;
 	uint32 NumNodesAddedSinceLastPointerPosition;
 };
 
diff --git a/Engine/Source/Editor/GraphEditor/Private/SGraphPanel.cpp b/Engine/Source/Editor/GraphEditor/Private/SGraphPanel.cpp
--- a/Engine/Source/Editor/GraphEditor/Private/SGraphPanel.cpp
+++ b/Engine/Source/Editor/GraphEditor/Private/SGraphPanel.cpp
@@ -2210,17 +2210,17 @@ void SGraphPanel::Update()
 			}
 		}
 
-		// find the last selection action, and execute it
-		for (int32 ActionIndex = UserActions.Num() - 1; ActionIndex >= 0; --ActionIndex)
+		// check the last selection action, and execute it
+		if (!UserSelectedNodes.IsEmpty())
 		{
-			if (UserActions[ActionIndex].Action & GRAPHACTION_SelectNode)
+			DeferredSelectionTargetObjects.Empty();
+			for (TWeakObjectPtr<UEdGraphNode>& NodePtr : UserSelectedNodes)
 			{
-				DeferredSelectionTargetObjects.Empty();
-				for (const UEdGraphNode* Node : UserActions[ActionIndex].Nodes)
+				if (NodePtr.IsValid())
 				{
+					UEdGraphNode* Node = NodePtr.Get();
 					DeferredSelectionTargetObjects.Add(Node);
 				}
-				break;
 			}
 		}
 	}
@@ -2230,8 +2230,8 @@ void SGraphPanel::Update()
 	}
 
 	// Clean out set of added nodes
-	UserAddedNodes.Empty();
-	UserActions.Empty();
+	UserAddedNodes.Reset();
+	UserSelectedNodes.Reset();
 
 	// Invoke any delegate methods
 	OnUpdateGraphPanel.ExecuteIfBound();
@@ -2348,28 +2348,51 @@ void SGraphPanel::OnEndPIE( const bool bIsSimulating )
 
 void SGraphPanel::OnGraphChanged(const FEdGraphEditAction& EditAction)
 {
-	const bool bWillPurge = GraphObj->GetSchema()->ShouldAlwaysPurgeOnModification();
-	if (bWillPurge)
+	const bool bShouldPurge = GraphObj->GetSchema()->ShouldAlwaysPurgeOnModification();
+	if (bShouldPurge || EditAction.Action == GRAPHACTION_Default)
 	{
+		if (!bVisualUpdatePending)
+		{
+			PurgeVisualRepresentation();
+
+			const auto RefreshPanelDelegateWrapper = [](double, float, TSharedRef<SGraphPanel> Parent) -> EActiveTimerReturnType
+			{
+				Parent->Update();
+				return EActiveTimerReturnType::Stop;
+			};
+
+			// Trigger the refresh
+			RegisterActiveTimer(0.f, FWidgetActiveTimerDelegate::CreateStatic(RefreshPanelDelegateWrapper, StaticCastSharedRef<SGraphPanel>(AsShared())));
+		}
+
 		if ((EditAction.Graph == GraphObj) &&
 			(EditAction.Nodes.Num() > 0) &&
 			EditAction.bUserInvoked)
 		{
-			int32 ActionIndex = UserActions.Num();
 			if (EditAction.Action & GRAPHACTION_AddNode)
 			{
+				UserAddedNodes.Append(EditAction.Nodes);
+			}
+			if (EditAction.Action & GRAPHACTION_SelectNode)
+			{
+				UserSelectedNodes.Reset();
 				for (const UEdGraphNode* Node : EditAction.Nodes)
 				{
-					UserAddedNodes.Add(Node, ActionIndex);
+					TWeakObjectPtr<UEdGraphNode> NodePtr = MakeWeakObjectPtr(const_cast<UEdGraphNode*>(Node));
+					UserSelectedNodes.Add(NodePtr);
 				}
 			}
-			UserActions.Add(EditAction);
 		}
 	}
 	else if ((EditAction.Graph == GraphObj) && (EditAction.Nodes.Num() > 0) )
 	{
+		// Ensure that any new non-default action(s) get handled here
+		constexpr int32 HandledActionsMask = (GRAPHACTION_AddNode | GRAPHACTION_EditNode | GRAPHACTION_SelectNode | GRAPHACTION_RemoveNode);
+		ensureMsgf((EditAction.Action & ~HandledActionsMask) == 0, TEXT("Unhandled actions: %08x"), EditAction.Action & ~HandledActionsMask);
+
 		// Remove action handled immediately by SGraphPanel::OnGraphChanged
 		const bool bWasAddAction = (EditAction.Action & GRAPHACTION_AddNode) != 0;
+		const bool bWasEditAction = (EditAction.Action & GRAPHACTION_EditNode) != 0;
 		const bool bWasSelectAction = (EditAction.Action & GRAPHACTION_SelectNode) != 0;
 		const bool bWasRemoveAction = (EditAction.Action & GRAPHACTION_RemoveNode) != 0;
 
@@ -2383,7 +2406,7 @@ void SGraphPanel::OnGraphChanged(const FEdGraphEditAction& EditAction)
 		// that the timer system requires (and we don't leverage):
 		if (bWasRemoveAction)
 		{
-			const auto RemoveNodeDelegateWrapper = [](double, float, SGraphPanel* Parent, TWeakObjectPtr<UEdGraphNode> NodePtr) -> EActiveTimerReturnType
+			const auto RemoveNodeDelegateWrapper = [](double, float, TSharedRef<SGraphPanel> Parent, TWeakObjectPtr<UEdGraphNode> NodePtr) -> EActiveTimerReturnType
 			{
 				if (NodePtr.IsValid())
 				{
@@ -2396,44 +2419,66 @@ void SGraphPanel::OnGraphChanged(const FEdGraphEditAction& EditAction)
 			for (const UEdGraphNode* Node : EditAction.Nodes)
 			{
 				TWeakObjectPtr<UEdGraphNode> NodePtr = MakeWeakObjectPtr(const_cast<UEdGraphNode*>(Node));
-				RegisterActiveTimer(0.f, FWidgetActiveTimerDelegate::CreateStatic(RemoveNodeDelegateWrapper, this, NodePtr));
+				RegisterActiveTimer(0.f, FWidgetActiveTimerDelegate::CreateStatic(RemoveNodeDelegateWrapper, StaticCastSharedRef<SGraphPanel>(AsShared()), NodePtr));
 			}
 		}
 		if (bWasAddAction)
 		{
-			const auto AddNodeDelegateWrapper = [](double, float, SGraphPanel* Parent, TWeakObjectPtr<UEdGraphNode> NodePtr, bool bForceUserAdded) -> EActiveTimerReturnType
+			const auto AddNodeDelegateWrapper = [](double, float, TSharedRef<SGraphPanel> Parent, TWeakObjectPtr<UEdGraphNode> NodePtr, bool bForceUserAdded) -> EActiveTimerReturnType
 			{
 				if (NodePtr.IsValid())
 				{
 					UEdGraphNode* Node = NodePtr.Get();
 					if(IsValid(Node))
 					{
-						Parent->RemoveNode(Node);
-						Parent->AddNode(Node, bForceUserAdded ? WasUserAdded : NotUserAdded);
+						if (Parent->bVisualUpdatePending)
+						{
+							if (bForceUserAdded)
+							{
+								Parent->UserAddedNodes.Add(Node);
+							}
+						}
+						else
+						{
+							Parent->RemoveNode(Node);
+							Parent->AddNode(Node, bForceUserAdded ? WasUserAdded : NotUserAdded);
+						}
 					}
 				}
 				return EActiveTimerReturnType::Stop;
 			};
 
+			TSharedRef<SGraphPanel> SharedThis = StaticCastSharedRef<SGraphPanel>(AsShared());
 			for (const UEdGraphNode* Node : EditAction.Nodes)
 			{
 				TWeakObjectPtr<UEdGraphNode> NodePtr = MakeWeakObjectPtr(const_cast<UEdGraphNode*>(Node));
-				RegisterActiveTimer(0.f, FWidgetActiveTimerDelegate::CreateStatic(AddNodeDelegateWrapper, this, NodePtr, EditAction.bUserInvoked));
+				RegisterActiveTimer(0.f, FWidgetActiveTimerDelegate::CreateStatic(AddNodeDelegateWrapper, SharedThis, NodePtr, EditAction.bUserInvoked));
 			}
 		}
 		if (bWasSelectAction)
 		{
-			const auto SelectNodeDelegateWrapper = [](double, float, SGraphPanel* Parent, TSet< TWeakObjectPtr<UEdGraphNode> > NodePtrs) -> EActiveTimerReturnType
+			const auto SelectNodeDelegateWrapper = [](double, float, TSharedRef<SGraphPanel> Parent, TSet< TWeakObjectPtr<UEdGraphNode> > NodePtrs, bool bForceUserAdded) -> EActiveTimerReturnType
 			{
-				Parent->DeferredSelectionTargetObjects.Empty();
-				for (TWeakObjectPtr<UEdGraphNode>& NodePtr : NodePtrs)
+				if (Parent->bVisualUpdatePending)
+				{
+					if (bForceUserAdded)
+					{
+						Parent->UserSelectedNodes = NodePtrs;
+					}
+				}
+				else
 				{
-					if (NodePtr.IsValid())
+					Parent->DeferredSelectionTargetObjects.Empty();
+					for (TWeakObjectPtr<UEdGraphNode>& NodePtr : NodePtrs)
 					{
-						UEdGraphNode* Node = NodePtr.Get();
-						Parent->DeferredSelectionTargetObjects.Add(Node);
+						if (NodePtr.IsValid())
+						{
+							UEdGraphNode* Node = NodePtr.Get();
+							Parent->DeferredSelectionTargetObjects.Add(Node);
+						}
 					}
 				}
+				
 				return EActiveTimerReturnType::Stop;
 			};
 
@@ -2444,7 +2489,14 @@ void SGraphPanel::OnGraphChanged(const FEdGraphEditAction& EditAction)
 				NodePtrSet.Add(NodePtr);
 			}
 
-			RegisterActiveTimer(0.f, FWidgetActiveTimerDelegate::CreateStatic(SelectNodeDelegateWrapper, this, NodePtrSet));
+			RegisterActiveTimer(0.f, FWidgetActiveTimerDelegate::CreateStatic(SelectNodeDelegateWrapper, StaticCastSharedRef<SGraphPanel>(AsShared()), NodePtrSet, EditAction.bUserInvoked));
+		}
+		if (bWasEditAction)
+		{
+			for (const UEdGraphNode* Node : EditAction.Nodes)
+			{
+				RefreshNode(const_cast<UEdGraphNode&>(*Node));
+			}
 		}
 	}
 }
diff --git a/Engine/Source/Editor/GraphEditor/Private/SGraphPreviewer.cpp b/Engine/Source/Editor/GraphEditor/Private/SGraphPreviewer.cpp
--- a/Engine/Source/Editor/GraphEditor/Private/SGraphPreviewer.cpp
+++ b/Engine/Source/Editor/GraphEditor/Private/SGraphPreviewer.cpp
@@ -31,7 +31,7 @@ EActiveTimerReturnType SGraphPreviewer::RefreshGraphTimer(const double InCurrent
 void SGraphPreviewer::Construct( const FArguments& InArgs, UEdGraph* InGraphObj )
 {
 	EdGraphObj = InGraphObj;
-	NeedsRefreshCounter = 2;
+	NeedsRefreshCounter = 0;
 
 	TSharedPtr<SOverlay> DisplayStack;
 
@@ -43,10 +43,11 @@ void SGraphPreviewer::Construct( const FArguments& InArgs, UEdGraph* InGraphObj
 		+SOverlay::Slot()
 		[
 			SAssignNew(GraphPanel, SGraphPanel)
-			.GraphObj( EdGraphObj )
-			.IsEditable( false )
+			.GraphObj(EdGraphObj)
+			.IsEditable(false)
 			.ShowGraphStateOverlay(InArgs._ShowGraphStateOverlay)
-			.InitialZoomToFit( true )
+			.InitialZoomToFit(true)
+			.OnUpdateGraphPanel(this, &SGraphPreviewer::OnUpdateGraphPanel)
 		]
 
 		// Bottom-right corner text indicating the type of tool
@@ -62,6 +63,7 @@ void SGraphPreviewer::Construct( const FArguments& InArgs, UEdGraph* InGraphObj
 		]
 	];
 
+	// Note - this will also invoke OnUpdateGraphPanel (see below)
 	GraphPanel->Update();
 
 	// Add the title bar if specified
@@ -73,6 +75,14 @@ void SGraphPreviewer::Construct( const FArguments& InArgs, UEdGraph* InGraphObj
 				InArgs._TitleBar.ToSharedRef()
 			];
 	}
-
-	RegisterActiveTimer(0.0f, FWidgetActiveTimerDelegate::CreateSP(this, &SGraphPreviewer::RefreshGraphTimer));
 }
+
+void SGraphPreviewer::OnUpdateGraphPanel()
+{
+	if (!NeedsRefreshCounter)
+	{
+		// As node's bounds don't get updated immediately, to truly zoom out to fit we need to tick a couple times
+		NeedsRefreshCounter = 2;
+		RegisterActiveTimer(0.0f, FWidgetActiveTimerDelegate::CreateSP(this, &SGraphPreviewer::RefreshGraphTimer));
+	}
+}
\ No newline at end of file
diff --git a/Engine/Source/Editor/GraphEditor/Public/SGraphPanel.h b/Engine/Source/Editor/GraphEditor/Public/SGraphPanel.h
--- a/Engine/Source/Editor/GraphEditor/Public/SGraphPanel.h
+++ b/Engine/Source/Editor/GraphEditor/Public/SGraphPanel.h
@@ -343,11 +343,11 @@ protected:
 	TAttribute<bool> ShowGraphStateOverlay;
 
 private:
-	/** Ordered list of user actions, as they came in */
-	TArray<FEdGraphEditAction> UserActions;
+	/** Set of nodes selected by the user, tracked while a visual update is pending */
+	TSet<TWeakObjectPtr<class UEdGraphNode>> UserSelectedNodes;
 
-	/** Map of recently added nodes for the panel (maps from added nodes to UserActions indices) */
-	TMap<const class UEdGraphNode*, int32> UserAddedNodes;
+	/** Set of user-added nodes for the panel, tracked while a visual update is pending */
+	TSet<const class UEdGraphNode*> UserAddedNodes;
 
 	/** Should the graph display all nodes in a read-only state (grayed)? This does not affect functionality of using them (IsEditable) */
 	TAttribute<bool> DisplayAsReadOnly;
diff --git a/Engine/Source/Editor/GraphEditor/Public/SGraphPreviewer.h b/Engine/Source/Editor/GraphEditor/Public/SGraphPreviewer.h
--- a/Engine/Source/Editor/GraphEditor/Public/SGraphPreviewer.h
+++ b/Engine/Source/Editor/GraphEditor/Public/SGraphPreviewer.h
@@ -29,6 +29,11 @@ public:
 	SLATE_END_ARGS()
 
 	void Construct(const FArguments& InArgs, UEdGraph* InGraphObj);
+
+protected:
+	/** Delegate handler for graph panel updates */
+	void OnUpdateGraphPanel();
+
 private:
 	/** Helper function used to refresh the graph */
 	EActiveTimerReturnType RefreshGraphTimer(const double InCurrentTime, const float InDeltaTime);
diff --git a/Engine/Source/Runtime/Engine/Private/EdGraph/EdGraph.cpp b/Engine/Source/Runtime/Engine/Private/EdGraph/EdGraph.cpp
--- a/Engine/Source/Runtime/Engine/Private/EdGraph/EdGraph.cpp
+++ b/Engine/Source/Runtime/Engine/Private/EdGraph/EdGraph.cpp
@@ -252,6 +252,7 @@ void UEdGraph::SelectNodeSet(TSet<const UEdGraphNode*> NodeSelection, bool bFrom
 
 	SelectionAction.Graph = this;
 	SelectionAction.Nodes = NodeSelection;
+	SelectionAction.bUserInvoked = bFromUI;
 
 	NotifyGraphChanged(SelectionAction);
 }
