GraphEditor: ensure the correct LOD is applied to graph node and pins before the first tick.
We rely on the LOD setting of the graph panel to determine the LOD of nodes and pins, but they only get a reference to the panel after their construction, so we need to refresh the LOD once we get a reference to the owning panel from the node. This is only needed before the first tick.
This is a prerequisite to fix the issue of graph nodes disappearing for a frame on undo/redo

diff --git a/Engine/Source/Editor/GraphEditor/Private/SGraphNode.cpp b/Engine/Source/Editor/GraphEditor/Private/SGraphNode.cpp
--- a/Engine/Source/Editor/GraphEditor/Private/SGraphNode.cpp
+++ b/Engine/Source/Editor/GraphEditor/Private/SGraphNode.cpp
@@ -554,6 +554,20 @@ void SGraphNode::SetOwner( const TSharedRef<SGraphPanel>& OwnerPanel )
 		this->RightNodeBox->ClearChildren();
 		CreatePinWidgets();
 	}
+
+	if (TitleLODBranchNode.IsValid())
+	{
+		TitleLODBranchNode->RefreshLODSlotContent();
+	}
+	
+	for (TSharedRef<SGraphPin> Pin : InputPins)
+	{
+		Pin->RefreshLOD();
+	}
+	for (TSharedRef<SGraphPin> Pin : OutputPins)
+	{
+		Pin->RefreshLOD();
+	}
 }
 
 /** @param NewPosition  The Node should be relocated to this position in the graph panel */
@@ -930,20 +944,19 @@ void SGraphNode::UpdateGraphNode()
 
 	SetDefaultTitleAreaWidget(DefaultTitleAreaWidget);
 
-	TSharedRef<SWidget> TitleAreaWidget = 
-		SNew(SLevelOfDetailBranchNode)
-		.UseLowDetailSlot(this, &SGraphNode::UseLowDetailNodeTitles)
-		.LowDetail()
-		[
-			SNew(SBorder)
-			.BorderImage( FAppStyle::GetBrush("Graph.Node.ColorSpill") )
-			.Padding( FMargin(75.0f, 22.0f) ) // Saving enough space for a 'typical' title so the transition isn't quite so abrupt
-			.BorderBackgroundColor( this, &SGraphNode::GetNodeTitleColor )
-		]
-		.HighDetail()
-		[
-			DefaultTitleAreaWidget
-		];
+	SAssignNew(TitleLODBranchNode, SLevelOfDetailBranchNode)
+	.UseLowDetailSlot(this, &SGraphNode::UseLowDetailNodeTitles)
+	.LowDetail()
+	[
+		SNew(SBorder)
+		.BorderImage( FAppStyle::GetBrush("Graph.Node.ColorSpill") )
+		.Padding( FMargin(75.0f, 22.0f) ) // Saving enough space for a 'typical' title so the transition isn't quite so abrupt
+		.BorderBackgroundColor( this, &SGraphNode::GetNodeTitleColor )
+	]
+	.HighDetail()
+	[
+		DefaultTitleAreaWidget
+	];
 
 	
 	if (!SWidget::GetToolTip().IsValid())
@@ -967,7 +980,7 @@ void SGraphNode::UpdateGraphNode()
 		.VAlign(VAlign_Top)
 		.Padding(Settings->GetNonPinNodeBodyPadding())
 		[
-			TitleAreaWidget
+			TitleLODBranchNode.ToSharedRef()
 		]
 
 		+SVerticalBox::Slot()
diff --git a/Engine/Source/Editor/GraphEditor/Private/SGraphPin.cpp b/Engine/Source/Editor/GraphEditor/Private/SGraphPin.cpp
--- a/Engine/Source/Editor/GraphEditor/Private/SGraphPin.cpp
+++ b/Engine/Source/Editor/GraphEditor/Private/SGraphPin.cpp
@@ -350,7 +350,7 @@ void SGraphPin::Construct(const FArguments& InArgs, UEdGraphPin* InPin)
 			.BorderImage(CachedImg_Pin_DiffOutline)
 			.BorderBackgroundColor(this, &SGraphPin::GetPinDiffColor)
 			[
-				SNew(SLevelOfDetailBranchNode)
+				SAssignNew(PinNameLODBranchNode, SLevelOfDetailBranchNode)
 				.UseLowDetailSlot(this, &SGraphPin::UseLowDetailPinNames)
 				.LowDetail()
 				[
@@ -385,6 +385,14 @@ TSharedRef<SWidget> SGraphPin::GetLabelWidget(const FName& InLabelStyle)
 		.ColorAndOpacity(this, &SGraphPin::GetPinTextColor);
 }
 
+void SGraphPin::RefreshLOD()
+{
+	if (PinNameLODBranchNode.IsValid())
+	{
+		PinNameLODBranchNode->RefreshLODSlotContent();
+	}
+}
+
 void SGraphPin::SetIsEditable(TAttribute<bool> InIsEditable)
 {
 	IsEditable = InIsEditable;
diff --git a/Engine/Source/Editor/GraphEditor/Public/SGraphNode.h b/Engine/Source/Editor/GraphEditor/Public/SGraphNode.h
--- a/Engine/Source/Editor/GraphEditor/Public/SGraphNode.h
+++ b/Engine/Source/Editor/GraphEditor/Public/SGraphNode.h
@@ -44,6 +44,7 @@ class IToolTip;
 class SGraphPanel;
 class SGraphPin;
 class SInlineEditableTextBlock;
+class SLevelOfDetailBranchNode;
 class SToolTip;
 class SVerticalBox;
 class SWidget;
@@ -481,4 +482,7 @@ protected:
 
 	/** Cached pointer to graph editor settings */
 	const class UGraphEditorSettings* Settings;
+
+private:
+	TSharedPtr<SLevelOfDetailBranchNode> TitleLODBranchNode;
 };
diff --git a/Engine/Source/Editor/GraphEditor/Public/SGraphPin.h b/Engine/Source/Editor/GraphEditor/Public/SGraphPin.h
--- a/Engine/Source/Editor/GraphEditor/Public/SGraphPin.h
+++ b/Engine/Source/Editor/GraphEditor/Public/SGraphPin.h
@@ -38,6 +38,7 @@ class SGraphPanel;
 class SGraphPin;
 class SHorizontalBox;
 class SImage;
+class SLevelOfDetailBranchNode;
 class SWidget;
 class SWrapBox;
 struct FGeometry;
@@ -151,6 +152,9 @@ public:
 	virtual TSharedPtr<IToolTip> GetToolTip() override;
 	// End of SWidget interface
 
+	/** Refresh the level of detail applied to the pin name */
+	void RefreshLOD();
+
 public:
 	UEdGraphPin* GetPinObj() const;
 
@@ -421,4 +425,7 @@ protected:
 
 	/** TRUE if the connections from this pin should be drawn at a lower opacity */
 	bool bFadeConnections;
+
+private:
+	TSharedPtr<SLevelOfDetailBranchNode> PinNameLODBranchNode;
 };
diff --git a/Engine/Source/Editor/KismetWidgets/Private/SLevelOfDetailBranchNode.cpp b/Engine/Source/Editor/KismetWidgets/Private/SLevelOfDetailBranchNode.cpp
--- a/Engine/Source/Editor/KismetWidgets/Private/SLevelOfDetailBranchNode.cpp
+++ b/Engine/Source/Editor/KismetWidgets/Private/SLevelOfDetailBranchNode.cpp
@@ -25,19 +25,24 @@ void SLevelOfDetailBranchNode::Construct(const FArguments& InArgs)
 	ChildSlotLowDetail = InArgs._LowDetail.Widget;
 	ChildSlotHighDetail = InArgs._HighDetail.Widget;
 
+	const int32 CurrentValue = ShowLowDetailAttr.Get() ? 1 : 0;
 	if (OnGetActiveDetailSlotContent.IsBound())
 	{
-		const int32 CurrentValue = ShowLowDetailAttr.Get() ? 1 : 0;
 		ChildSlot[OnGetActiveDetailSlotContent.Execute((CurrentValue != 0) ? false : true)];
 	}
 	else
 	{
-		ChildSlot[ChildSlotHighDetail];
+		ChildSlot[(CurrentValue != 0) ? ChildSlotLowDetail : ChildSlotHighDetail];
 	}
 
 }
 
 void SLevelOfDetailBranchNode::Tick( const FGeometry& AllottedGeometry, const double InCurrentTime, const float InDeltaTime )
+{
+	RefreshLODSlotContent();
+}
+
+void SLevelOfDetailBranchNode::RefreshLODSlotContent()
 {
 	const int32 CurrentValue = ShowLowDetailAttr.Get() ? 1 : 0;
 	if (CurrentValue != LastCachedValue)
@@ -47,16 +52,16 @@ void SLevelOfDetailBranchNode::Tick( const FGeometry& AllottedGeometry, const do
 		if (OnGetActiveDetailSlotContent.IsBound())
 		{
 			ChildSlot
-			[
-				OnGetActiveDetailSlotContent.Execute((CurrentValue != 0) ? false : true)
-			];
+				[
+					OnGetActiveDetailSlotContent.Execute((CurrentValue != 0) ? false : true)
+				];
 		}
 		else
 		{
 			ChildSlot
-			[
-				(CurrentValue != 0) ? ChildSlotLowDetail : ChildSlotHighDetail
-			];
+				[
+					(CurrentValue != 0) ? ChildSlotLowDetail : ChildSlotHighDetail
+				];
 		}
 	}
-}
+}
\ No newline at end of file
diff --git a/Engine/Source/Editor/KismetWidgets/Public/SLevelOfDetailBranchNode.h b/Engine/Source/Editor/KismetWidgets/Public/SLevelOfDetailBranchNode.h
--- a/Engine/Source/Editor/KismetWidgets/Public/SLevelOfDetailBranchNode.h
+++ b/Engine/Source/Editor/KismetWidgets/Public/SLevelOfDetailBranchNode.h
@@ -45,6 +45,9 @@ public:
 	virtual void Tick( const FGeometry& AllottedGeometry, const double InCurrentTime, const float InDeltaTime ) override;
 	// End of SWidget interface
 
+	// Determine whether we need to show the low-detail slot or high-detail slot
+	void RefreshLODSlotContent();
+
 protected:
 	// What kind of slot was shown last frame
 	int LastCachedValue;
