PR: Fix graph nodes flashing for a frame when using undo/redo
#10486

diff --git a/Engine/Source/Editor/GraphEditor/Private/SGraphPanel.cpp b/Engine/Source/Editor/GraphEditor/Private/SGraphPanel.cpp
--- a/Engine/Source/Editor/GraphEditor/Private/SGraphPanel.cpp
+++ b/Engine/Source/Editor/GraphEditor/Private/SGraphPanel.cpp
@@ -2135,6 +2135,26 @@ void SGraphPanel::AddNode(UEdGraphNode* Node, AddNodeBehavior Behavior)
 		NewNode->PlaySpawnEffect();
 		NewNode->RequestRenameOnSpawn();
 	}
+
+	// Note: We delay the creation of widgets for new nodes by a frame in `OnGraphChanged()`, using a Slate timer per node that later
+	// calls into this method. Slate timers are executed from within the Paint event, but before the actual OnPaint is called. This means we've
+	// just inserted the new node widget after this panel has already pre-passed the existing node widgets, and because it's now a child of the panel
+	// it'll also be painted this frame, despite not having been pre-passed (meaning it'll be stuck with a desired size of zero).
+	// Because the new node widget(s) get painted with zero size, pin connection wires are then be drawn based on the layout of these zero-sized node(s),
+	// resulting in a pretty obvious one-frame flash whenever you insert new nodes. It's particularly visible when using undo/redo,
+	// since a lot of nodes can be inserted at once. To avoid this flash of 'painting without pre-pass', we'll just manually pre-pass
+	// the new widget here so that when we go to paint it after this function returns it'll at least have some sizing information when we arrange it in our OnPaint().
+	// This is safe since graph widgets don't rely on any outer layout information for their metrics, and we don't size ourselves based on node widgets either.
+	// We also need to take a bit of care to pass through the same layout scale multiplier as Prepass_ChildLoop() would have so that the zoom level
+	// scale is used, otherwise you'd still get a single frame of jitter while the graph is zoomed out.
+	const int32 ChildIndex = Children.Num() - 1;
+	const float SelfLayoutScaleMultiplier = PrepassLayoutScaleMultiplier.Get(1.f);
+	const float ChildLayoutScaleMultiplier = bHasRelativeLayoutScale
+		? SelfLayoutScaleMultiplier * GetRelativeLayoutScale(ChildIndex, SelfLayoutScaleMultiplier)
+		: SelfLayoutScaleMultiplier;
+
+	NewNode->MarkPrepassAsDirty();
+	NewNode->SlatePrepass(ChildLayoutScaleMultiplier);
 }
 
 void SGraphPanel::RemoveNode(const UEdGraphNode* Node)
