Added slate argument to automatically select an action in the graph action menu when only one action exists.
Defaults to off to save perf. Useful for some Niagara cases

diff --git a/Engine/Source/Editor/GraphEditor/Private/GraphActionNode.cpp b/Engine/Source/Editor/GraphEditor/Private/GraphActionNode.cpp
--- a/Engine/Source/Editor/GraphEditor/Private/GraphActionNode.cpp
+++ b/Engine/Source/Editor/GraphEditor/Private/GraphActionNode.cpp
@@ -280,6 +280,19 @@ void FGraphActionNode::GetAllNodes(TArray< TSharedPtr<FGraphActionNode> >& OutNo
 	}
 }
 
+void FGraphActionNode::GetAllActionNodes(TArray<TSharedPtr<FGraphActionNode>>& OutNodeArray) const
+{	
+	for (TSharedPtr<FGraphActionNode> const& ChildNode : Children)
+	{
+		if(ChildNode->IsActionNode())
+		{
+			OutNodeArray.Add(ChildNode);
+		}
+		
+		ChildNode->GetAllActionNodes(OutNodeArray);
+	}
+}
+
 //------------------------------------------------------------------------------
 void FGraphActionNode::GetLeafNodes(TArray< TSharedPtr<FGraphActionNode> >& OutLeafArray) const
 {
diff --git a/Engine/Source/Editor/GraphEditor/Private/SGraphActionMenu.cpp b/Engine/Source/Editor/GraphEditor/Private/SGraphActionMenu.cpp
--- a/Engine/Source/Editor/GraphEditor/Private/SGraphActionMenu.cpp
+++ b/Engine/Source/Editor/GraphEditor/Private/SGraphActionMenu.cpp
@@ -311,6 +311,7 @@ void SGraphActionMenu::Construct( const FArguments& InArgs, bool bIsReadOnly/* =
 	this->bIgnoreUIUpdate = false;
 	this->bUseSectionStyling = InArgs._UseSectionStyling;
 	this->bAllowPreselectedItemActivation = InArgs._bAllowPreselectedItemActivation;
+	this->bAutomaticallySelectSingleAction = InArgs._bAutomaticallySelectSingleAction;
 	this->DefaultRowExpanderBaseIndentLevel = InArgs._DefaultRowExpanderBaseIndentLevel;
 
 	this->bAutoExpandActionMenu = InArgs._AutoExpandActionMenu;
@@ -406,6 +407,17 @@ void SGraphActionMenu::Construct( const FArguments& InArgs, bool bIsReadOnly/* =
 
 	// Get all actions.
 	RefreshAllActions(false);
+
+	if(bAutomaticallySelectSingleAction)
+	{
+		TArray<TSharedPtr<FGraphActionNode>> ActionNodes;
+		FilteredRootAction->GetAllActionNodes(ActionNodes);
+
+		if(ActionNodes.Num() == 1)
+		{
+			OnItemSelected(ActionNodes[0], ESelectInfo::Direct);
+		}
+	}
 }
 
 void SGraphActionMenu::RefreshAllActions(bool bPreserveExpansion, bool bHandleOnSelectionEvent/*=true*/)
diff --git a/Engine/Source/Editor/GraphEditor/Public/GraphActionNode.h b/Engine/Source/Editor/GraphEditor/Public/GraphActionNode.h
--- a/Engine/Source/Editor/GraphEditor/Public/GraphActionNode.h
+++ b/Engine/Source/Editor/GraphEditor/Public/GraphActionNode.h
@@ -87,6 +87,13 @@ public:
 	 */
 	void GetAllNodes(TArray< TSharedPtr<FGraphActionNode> >& OutNodeArray) const;
 
+	/**
+	 * Recursively collects all child/grandchild/decendent action nodes.
+	 * 
+	 * @param  OutNodeArray	The array to fill out with decendent nodes.
+	 */
+	void GetAllActionNodes(TArray<TSharedPtr<FGraphActionNode>>& OutNodeArray) const;
+
 	/**
 	 * Recursively collects all decendent action/separator nodes (leaves out 
 	 * branching category-nodes).
diff --git a/Engine/Source/Editor/GraphEditor/Public/SGraphActionMenu.h b/Engine/Source/Editor/GraphEditor/Public/SGraphActionMenu.h
--- a/Engine/Source/Editor/GraphEditor/Public/SGraphActionMenu.h
+++ b/Engine/Source/Editor/GraphEditor/Public/SGraphActionMenu.h
@@ -154,6 +154,7 @@ public:
 		, _bAllowPreselectedItemActivation(false)
 		, _DefaultRowExpanderBaseIndentLevel(0)
 		, _GraphObj(nullptr)
+		, _bAutomaticallySelectSingleAction(false)
 		{ }
 
 		SLATE_EVENT( FOnActionSelected, OnActionSelected )
@@ -182,6 +183,7 @@ public:
 		SLATE_ARGUMENT( int32, DefaultRowExpanderBaseIndentLevel )
 		SLATE_ARGUMENT( TArray<UEdGraphPin*>, DraggedFromPins )
 		SLATE_ARGUMENT( UEdGraph*, GraphObj )
+		SLATE_ARGUMENT( bool, bAutomaticallySelectSingleAction )
 
 	SLATE_END_ARGS()
 
@@ -239,6 +241,8 @@ protected:
 	bool bUseSectionStyling;
 	/** Whether we allow pre-selected items to be activated with a left-click */
 	bool bAllowPreselectedItemActivation;
+	/** Whether to automatically proceed with an action if it's the only one in the list. */
+	bool bAutomaticallySelectSingleAction;
 	/** The BaseIndentLevel of the default-created row expander. Not used with OnCreateCustomRowExpander. */
 	int32 DefaultRowExpanderBaseIndentLevel;
 	
