Reverting blueprint editor shortcut Alt+Shift+F to search locally by mapping it to the singular Find References action. The singular Find References action, while not exposed to UI for nodes that have Find References By Name/Class Member support, can still be triggered with shortcut and will trigger a local By Name or By [REDACTED_PERSON] search based on what the previous behavior was for that node type. For function nodes it will search by function name, while for variable nodes it will do an exact search since Find References already did that for a while.

However: function searches will do a quoted search by native function name. The previous behavior was unquoted search by node title (usually function display name). As a result, the previous behavior for Find References would fail in functions with special characters in their name. Now that the name is surrounded in quotes, all function names are supported. The new Find References behavior now searches for correct function name for parent call nodes, interface implementations, event overrides, where the previous behavior failed due to searching for node title.

diff --git a/Engine/Source/Editor/BlueprintGraph/Private/K2Node_CallFunction.cpp b/Engine/Source/Editor/BlueprintGraph/Private/K2Node_CallFunction.cpp
--- a/Engine/Source/Editor/BlueprintGraph/Private/K2Node_CallFunction.cpp
+++ b/Engine/Source/Editor/BlueprintGraph/Private/K2Node_CallFunction.cpp
@@ -1595,14 +1595,12 @@ FString UK2Node_CallFunction::GetFindReferenceSearchString_Impl(EGetFindReferenc
 			}
 		}
 	}
-	else
+
+	// Searching by function name: try to resolve function to return native name
+	if (const UFunction* Function = GetTargetFunction())
 	{
-		// Searching by function name: try to resolve function to return native name
-		if (const UFunction* Function = GetTargetFunction())
-		{
-			const FString NativeName = Function->GetName();
-			return NativeName;
-		}
+		const FString NativeName = Function->GetName();
+		return FString::Printf(TEXT("\"%s\""), *NativeName);
 	}
 
 	// Fallback behavior: search by node title
diff --git a/Engine/Source/Editor/BlueprintGraph/Private/K2Node_Event.cpp b/Engine/Source/Editor/BlueprintGraph/Private/K2Node_Event.cpp
--- a/Engine/Source/Editor/BlueprintGraph/Private/K2Node_Event.cpp
+++ b/Engine/Source/Editor/BlueprintGraph/Private/K2Node_Event.cpp
@@ -904,7 +904,7 @@ FString UK2Node_Event::GetFindReferenceSearchString_Impl(EGetFindReferenceSearch
 	// If searching by class member, try to construct search term from the UFunction.
 	// This may fail if the function was not found or for whatever reason, its owning 
 	// class is invalid. If it fails, proceed to search by name as fallback behavior.
-	if (EnumHasAnyFlags(InFlags, EGetFindReferenceSearchStringFlags::UseSearchSyntax))
+	if (EnumHasAnyFlags(InFlags, EGetFindReferenceSearchStringFlags::UseSearchSyntax) && !EnumHasAnyFlags(InFlags, EGetFindReferenceSearchStringFlags::Legacy))
 	{
 		// Resolve the function
 		if (const UFunction* Function = FFunctionFromNodeHelper::FunctionFromNode(this))
@@ -925,18 +925,19 @@ FString UK2Node_Event::GetFindReferenceSearchString_Impl(EGetFindReferenceSearch
 		{
 			// Search by native name
 			const FString FunctionNativeName = Function->GetName();
-			return FunctionNativeName;
+			return FString::Printf(TEXT("\"%s\""), *FunctionNativeName);
 		}
 		else
 		{
-			// If we fail to find the function, still want to search for its expected name
-			return EventReference.GetMemberName().ToString();
+			// If we fail to find the function, still want to search for its expected name, in quotes
+			return FString::Printf(TEXT("\"%s\""), *EventReference.GetMemberName().ToString());
 		}
 	}
 	else
 	{
-		// The function was not an override; its name is defined by this node
-		return CustomFunctionName.ToString();
+		// The function was not an override; its name is defined by this node.
+		// Return that name in quotes (treating special characters as part of name)
+		return FString::Printf(TEXT("\"%s\""), *CustomFunctionName.ToString());
 	}
 }
 
diff --git a/Engine/Source/Editor/BlueprintGraph/Private/K2Node_FunctionTerminator.cpp b/Engine/Source/Editor/BlueprintGraph/Private/K2Node_FunctionTerminator.cpp
--- a/Engine/Source/Editor/BlueprintGraph/Private/K2Node_FunctionTerminator.cpp
+++ b/Engine/Source/Editor/BlueprintGraph/Private/K2Node_FunctionTerminator.cpp
@@ -75,6 +75,17 @@ FString UK2Node_FunctionTerminator::GetFindReferenceSearchString_Impl(EGetFindRe
 			}
 		}
 	}
+	else
+	{
+		// When searching by name, return function native name in quotes.
+		// The quotes guarantee that the whole function name is used as single search term.
+		// This avoids function names with special characters being interpreted as operators.
+		if (const UFunction* Function = FFunctionFromNodeHelper::FunctionFromNode(this))
+		{
+			const FString NativeName = Function->GetName();
+			return FString::Printf(TEXT("\"%s\""), *NativeName);
+		}
+	}
 
 	// Fallback behavior: function was not resolved
 	return Super::GetFindReferenceSearchString_Impl(InFlags);
diff --git a/Engine/Source/Editor/BlueprintGraph/Private/K2Node_Variable.cpp b/Engine/Source/Editor/BlueprintGraph/Private/K2Node_Variable.cpp
--- a/Engine/Source/Editor/BlueprintGraph/Private/K2Node_Variable.cpp
+++ b/Engine/Source/Editor/BlueprintGraph/Private/K2Node_Variable.cpp
@@ -265,7 +265,8 @@ FLinearColor UK2Node_Variable::GetNodeTitleColor() const
 
 FString UK2Node_Variable::GetFindReferenceSearchString_Impl(EGetFindReferenceSearchStringFlags InFlags) const
 {
-	if (EnumHasAnyFlags(InFlags, EGetFindReferenceSearchStringFlags::UseSearchSyntax))
+	// Legacy behavior for variable nodes was to do an exact search
+	if (EnumHasAnyFlags(InFlags, EGetFindReferenceSearchStringFlags::UseSearchSyntax) || EnumHasAnyFlags(InFlags, EGetFindReferenceSearchStringFlags::Legacy))
 	{
 		if (VariableReference.IsLocalScope())
 		{
@@ -280,7 +281,7 @@ FString UK2Node_Variable::GetFindReferenceSearchString_Impl(EGetFindReferenceSea
 	}
 
 	// Simple query: just search for variable name
-	return VariableReference.GetMemberName().ToString();
+	return FString::Printf(TEXT("\"%s\""), *VariableReference.GetMemberName().ToString());
 }
 
 UK2Node::ERedirectType UK2Node_Variable::DoPinsMatchForReconstruction(const UEdGraphPin* NewPin, int32 NewPinIndex, const UEdGraphPin* OldPin, int32 OldPinIndex) const 
diff --git a/Engine/Source/Editor/GraphEditor/Private/GraphEditorActions.cpp b/Engine/Source/Editor/GraphEditor/Private/GraphEditorActions.cpp
--- a/Engine/Source/Editor/GraphEditor/Private/GraphEditorActions.cpp
+++ b/Engine/Source/Editor/GraphEditor/Private/GraphEditorActions.cpp
@@ -77,11 +77,11 @@ void FGraphEditorCommandsImpl::RegisterCommands()
 
 	UI_COMMAND( SelectReferenceInLevel, "Find Actor in Level", "Select the actor referenced by this node in the level", EUserInterfaceActionType::Button, FInputChord() )
 	UI_COMMAND( AssignReferencedActor, "Assign Selected Actor", "Assign the selected actor to be this node's referenced object", EUserInterfaceActionType::Button, FInputChord() )
-	UI_COMMAND( FindReferences, "Find References", "Find references of this item", EUserInterfaceActionType::Button, FInputChord() )
+	UI_COMMAND( FindReferences, "Find References", "Find references of this item", EUserInterfaceActionType::Button, FInputChord(EKeys::F, EModifierKey::Shift | EModifierKey::Alt) )
 	UI_COMMAND( FindReferencesByNameLocal, "By Name", "Find references of this item by name", EUserInterfaceActionType::Button, FInputChord() )
 	UI_COMMAND( FindReferencesByNameGlobal, "By Name (All)", "Find references of this item by name in all blueprints", EUserInterfaceActionType::Button, FInputChord() )
 	UI_COMMAND( FindReferencesByClassMemberLocal, "By Class Member", "Find references of this item by class member", EUserInterfaceActionType::Button, FInputChord() )
-	UI_COMMAND( FindReferencesByClassMemberGlobal, "By Class Member (All)", "Find references of this item by class member in all blueprints", EUserInterfaceActionType::Button, FInputChord(EKeys::F, EModifierKey::Shift | EModifierKey::Alt) )
+	UI_COMMAND( FindReferencesByClassMemberGlobal, "By Class Member (All)", "Find references of this item by class member in all blueprints", EUserInterfaceActionType::Button, FInputChord() )
 	UI_COMMAND( FindAndReplaceReferences, "Replace References", "Brings up a window to help find and replace all instances of this item", EUserInterfaceActionType::Button, FInputChord() )
 	
 	UI_COMMAND( GoToDefinition, "Goto Definition", "Jumps to the defintion of the selected node if available, e.g., C++ code for a native function or the graph for a Blueprint function.", EUserInterfaceActionType::Button, FInputChord(EKeys::G, EModifierKey::Alt) )
diff --git a/Engine/Source/Editor/Kismet/Private/BlueprintEditor.cpp b/Engine/Source/Editor/Kismet/Private/BlueprintEditor.cpp
--- a/Engine/Source/Editor/Kismet/Private/BlueprintEditor.cpp
+++ b/Engine/Source/Editor/Kismet/Private/BlueprintEditor.cpp
@@ -1654,7 +1654,7 @@ TSharedRef<SGraphEditor> FBlueprintEditor::CreateGraphEditorWidget(TSharedRef<FT
 				);
 
 			GraphEditorCommands->MapAction( FGraphEditorCommands::Get().FindReferences,
-				FExecuteAction::CreateSP( this, &FBlueprintEditor::OnFindReferences, /*bSearchAllBlueprints=*/false, EGetFindReferenceSearchStringFlags::None),
+				FExecuteAction::CreateSP( this, &FBlueprintEditor::OnFindReferences, /*bSearchAllBlueprints=*/false, EGetFindReferenceSearchStringFlags::Legacy),
 				FCanExecuteAction::CreateSP( this, &FBlueprintEditor::CanFindReferences )
 				);
 			
@@ -10222,12 +10222,6 @@ void FBlueprintEditor::OnFindReferences(bool bSearchAllBlueprints, const EGetFin
 				FString SearchTerm = SelectedNode->GetFindReferenceSearchString(Flags);
 				if (!SearchTerm.IsEmpty())
 				{
-					// If not using search syntax, surround the search term with quotes
-					if (!EnumHasAnyFlags(Flags, EGetFindReferenceSearchStringFlags::UseSearchSyntax))
-					{
-						SearchTerm = FString::Printf(TEXT("\"%s\""), *SearchTerm);
-					}
-
 					// Start the search
 					const bool bSetFindWithinBlueprint = !bSearchAllBlueprints;
 					SummonSearchUI(bSetFindWithinBlueprint, SearchTerm);
diff --git a/Engine/Source/Editor/Kismet/Private/SMyBlueprint.cpp b/Engine/Source/Editor/Kismet/Private/SMyBlueprint.cpp
--- a/Engine/Source/Editor/Kismet/Private/SMyBlueprint.cpp
+++ b/Engine/Source/Editor/Kismet/Private/SMyBlueprint.cpp
@@ -357,7 +357,7 @@ void SMyBlueprint::Construct(const FArguments& InArgs, TWeakPtr<FBlueprintEditor
 			FIsActionButtonVisible::CreateSP(this, &SMyBlueprint::CanImplementFunction) );
 		
 		CommandList->MapAction( FGraphEditorCommands::Get().FindReferences,
-			FExecuteAction::CreateSP(this, &SMyBlueprint::OnFindReference, /*bSearchAllBlueprints=*/false, EGetFindReferenceSearchStringFlags::None),
+			FExecuteAction::CreateSP(this, &SMyBlueprint::OnFindReference, /*bSearchAllBlueprints=*/false, EGetFindReferenceSearchStringFlags::Legacy),
 			FCanExecuteAction(),
 			FIsActionChecked(),
 			FIsActionButtonVisible::CreateSP(this, &SMyBlueprint::CanFindReference) );
diff --git a/Engine/Source/Editor/Kismet/Private/SSCSEditor.cpp b/Engine/Source/Editor/Kismet/Private/SSCSEditor.cpp
--- a/Engine/Source/Editor/Kismet/Private/SSCSEditor.cpp
+++ b/Engine/Source/Editor/Kismet/Private/SSCSEditor.cpp
@@ -3942,7 +3942,7 @@ void SSCSEditor::Construct( const FArguments& InArgs )
 		);
 
 	CommandList->MapAction( FGraphEditorCommands::Get().FindReferences,
-		FUIAction( FExecuteAction::CreateSP( this, &SSCSEditor::OnFindReferences, false, EGetFindReferenceSearchStringFlags::UseSearchSyntax ) )
+		FUIAction( FExecuteAction::CreateSP( this, &SSCSEditor::OnFindReferences, false, EGetFindReferenceSearchStringFlags::Legacy) )
 	);
 
 	CommandList->MapAction( FGraphEditorCommands::Get().FindReferencesByNameLocal,
diff --git a/Engine/Source/Editor/SubobjectEditor/Private/SSubobjectEditor.cpp b/Engine/Source/Editor/SubobjectEditor/Private/SSubobjectEditor.cpp
--- a/Engine/Source/Editor/SubobjectEditor/Private/SSubobjectEditor.cpp
+++ b/Engine/Source/Editor/SubobjectEditor/Private/SSubobjectEditor.cpp
@@ -2148,7 +2148,7 @@ void SSubobjectEditor::CreateCommandList()
 	);
 
 	CommandList->MapAction(FGraphEditorCommands::Get().GetFindReferences(),
-		FUIAction(FExecuteAction::CreateSP(this, &SSubobjectEditor::OnFindReferences, false, EGetFindReferenceSearchStringFlags::UseSearchSyntax ) )
+		FUIAction(FExecuteAction::CreateSP(this, &SSubobjectEditor::OnFindReferences, false, EGetFindReferenceSearchStringFlags::Legacy) )
 	);
 	
 	CommandList->MapAction( FGraphEditorCommands::Get().FindReferencesByNameLocal,
diff --git a/Engine/Source/Editor/UMGEditor/Private/WidgetBlueprintEditor.cpp b/Engine/Source/Editor/UMGEditor/Private/WidgetBlueprintEditor.cpp
--- a/Engine/Source/Editor/UMGEditor/Private/WidgetBlueprintEditor.cpp
+++ b/Engine/Source/Editor/UMGEditor/Private/WidgetBlueprintEditor.cpp
@@ -205,7 +205,7 @@ void FWidgetBlueprintEditor::InitWidgetBlueprintEditor(const EToolkitMode::Type
 		);
 
 	DesignerCommandList->MapAction(FGraphEditorCommands::Get().FindReferences,
-		FExecuteAction::CreateSP(this, &FWidgetBlueprintEditor::OnFindWidgetReferences, false, EGetFindReferenceSearchStringFlags::UseSearchSyntax),
+		FExecuteAction::CreateSP(this, &FWidgetBlueprintEditor::OnFindWidgetReferences, false, EGetFindReferenceSearchStringFlags::Legacy),
 		FCanExecuteAction::CreateSP(this, &FWidgetBlueprintEditor::CanFindWidgetReferences));
 	
 	DesignerCommandList->MapAction(FGraphEditorCommands::Get().FindReferencesByNameLocal,
diff --git a/Engine/Source/Runtime/Engine/Classes/EdGraph/EdGraphNode.h b/Engine/Source/Runtime/Engine/Classes/EdGraph/EdGraphNode.h
--- a/Engine/Source/Runtime/Engine/Classes/EdGraph/EdGraphNode.h
+++ b/Engine/Source/Runtime/Engine/Classes/EdGraph/EdGraphNode.h
@@ -272,8 +272,12 @@ struct FEdGraphNodeDeprecationResponse
 /** Options when getting a Find References search string from a node. */
 enum class EGetFindReferenceSearchStringFlags : int8
 {
+	/** No options. */
 	None = 0,
-	UseSearchSyntax = 1 << 0,
+	/** Use search behavior from before there were different Find References options. */
+	Legacy = 1 << 0,
+	/** Use search syntax to do exact search. */
+	UseSearchSyntax = 1 << 1,
 };
 
 ENUM_CLASS_FLAGS(EGetFindReferenceSearchStringFlags);
