Fix a regression that could prevent widgets from being destroyed after closing an editor window containing a graph panel.

This could lead to an inability to delete certain assets through the content browser, as panel widgets would continue to reference the underlying graph object through the slate application context with no way to clean them up. In turn, this would result in the asset persisting and being (incorrectly) noted as now existing in a "corrupt" state.

After this change, the asset deletion dialog will no longer report the asset as being referenced in this manner.

This change also minimizes the number of active timers registered to the panel widget for incoming node add/remove graph change events (which matches what was already being done for node select events).

diff --git a/Engine/Source/Editor/GraphEditor/Private/SGraphPanel.cpp b/Engine/Source/Editor/GraphEditor/Private/SGraphPanel.cpp
--- a/Engine/Source/Editor/GraphEditor/Private/SGraphPanel.cpp
+++ b/Engine/Source/Editor/GraphEditor/Private/SGraphPanel.cpp
@@ -2375,14 +2375,18 @@ void SGraphPanel::OnGraphChanged(const FEdGraphEditAction& EditAction)
 		{
 			PurgeVisualRepresentation();
 
-			const auto RefreshPanelDelegateWrapper = [](double, float, TSharedRef<SGraphPanel> Parent) -> EActiveTimerReturnType
+			const auto RefreshPanelDelegateWrapper = [](double, float, TWeakPtr<SGraphPanel> WeakParent) -> EActiveTimerReturnType
 			{
-				Parent->Update();
+				TSharedPtr<SGraphPanel> Parent = WeakParent.Pin();
+				if (Parent.IsValid())
+				{
+					Parent->Update();
+				}
 				return EActiveTimerReturnType::Stop;
 			};
 
 			// Trigger the refresh
-			RegisterActiveTimer(0.f, FWidgetActiveTimerDelegate::CreateStatic(RefreshPanelDelegateWrapper, StaticCastSharedRef<SGraphPanel>(AsShared())));
+			RegisterActiveTimer(0.f, FWidgetActiveTimerDelegate::CreateLambda(RefreshPanelDelegateWrapper, StaticCastWeakPtr<SGraphPanel>(AsWeak())));
 		}
 
 		if ((EditAction.Graph == GraphObj) &&
@@ -2426,75 +2430,98 @@ void SGraphPanel::OnGraphChanged(const FEdGraphEditAction& EditAction)
 		// that the timer system requires (and we don't leverage):
 		if (bWasRemoveAction)
 		{
-			const auto RemoveNodeDelegateWrapper = [](double, float, TSharedRef<SGraphPanel> Parent, TWeakObjectPtr<UEdGraphNode> NodePtr) -> EActiveTimerReturnType
+			const auto RemoveNodesDelegateWrapper = [](double, float, TWeakPtr<SGraphPanel> WeakParent, TSet< TWeakObjectPtr<UEdGraphNode> > NodePtrs) -> EActiveTimerReturnType
 			{
-				if (NodePtr.IsValid())
+				TSharedPtr<SGraphPanel> Parent = WeakParent.Pin();
+				if (Parent.IsValid())
 				{
-					UEdGraphNode* Node = NodePtr.Get();
-					Parent->RemoveNode(Node);
+					for (TWeakObjectPtr<UEdGraphNode>& NodePtr : NodePtrs)
+					{
+						if (NodePtr.IsValid())
+						{
+							UEdGraphNode* Node = NodePtr.Get();
+							Parent->RemoveNode(Node);
+						}
+					}
 				}
 				return EActiveTimerReturnType::Stop;
 			};
 
+			TSet< TWeakObjectPtr<UEdGraphNode> > NodePtrSet;
 			for (const UEdGraphNode* Node : EditAction.Nodes)
 			{
 				TWeakObjectPtr<UEdGraphNode> NodePtr = MakeWeakObjectPtr(const_cast<UEdGraphNode*>(Node));
-				RegisterActiveTimer(0.f, FWidgetActiveTimerDelegate::CreateStatic(RemoveNodeDelegateWrapper, StaticCastSharedRef<SGraphPanel>(AsShared()), NodePtr));
+				NodePtrSet.Add(NodePtr);
 			}
+
+			RegisterActiveTimer(0.f, FWidgetActiveTimerDelegate::CreateLambda(RemoveNodesDelegateWrapper, StaticCastWeakPtr<SGraphPanel>(AsWeak()), NodePtrSet));
 		}
 		if (bWasAddAction)
 		{
-			const auto AddNodeDelegateWrapper = [](double, float, TSharedRef<SGraphPanel> Parent, TWeakObjectPtr<UEdGraphNode> NodePtr, bool bForceUserAdded) -> EActiveTimerReturnType
+			const auto AddNodesDelegateWrapper = [](double, float, TWeakPtr<SGraphPanel> WeakParent, TSet< TWeakObjectPtr<UEdGraphNode> > NodePtrs, bool bForceUserAdded) -> EActiveTimerReturnType
 			{
-				if (NodePtr.IsValid())
+				TSharedPtr<SGraphPanel> Parent = WeakParent.Pin();
+				if (Parent.IsValid())
 				{
-					UEdGraphNode* Node = NodePtr.Get();
-					if(IsValid(Node))
+					for (TWeakObjectPtr<UEdGraphNode>& NodePtr : NodePtrs)
 					{
-						if (Parent->bVisualUpdatePending)
+						if (NodePtr.IsValid())
 						{
-							if (bForceUserAdded)
+							UEdGraphNode* Node = NodePtr.Get();
+							if (IsValid(Node))
 							{
-								Parent->UserAddedNodes.Add(Node);
+								if (Parent->bVisualUpdatePending)
+								{
+									if (bForceUserAdded)
+									{
+										Parent->UserAddedNodes.Add(Node);
+									}
+								}
+								else
+								{
+									Parent->RemoveNode(Node);
+									Parent->AddNode(Node, bForceUserAdded ? WasUserAdded : NotUserAdded);
+								}
 							}
 						}
-						else
-						{
-							Parent->RemoveNode(Node);
-							Parent->AddNode(Node, bForceUserAdded ? WasUserAdded : NotUserAdded);
-						}
 					}
 				}
 				return EActiveTimerReturnType::Stop;
 			};
 
-			TSharedRef<SGraphPanel> SharedThis = StaticCastSharedRef<SGraphPanel>(AsShared());
+			TSet< TWeakObjectPtr<UEdGraphNode> > NodePtrSet;
 			for (const UEdGraphNode* Node : EditAction.Nodes)
 			{
 				TWeakObjectPtr<UEdGraphNode> NodePtr = MakeWeakObjectPtr(const_cast<UEdGraphNode*>(Node));
-				RegisterActiveTimer(0.f, FWidgetActiveTimerDelegate::CreateStatic(AddNodeDelegateWrapper, SharedThis, NodePtr, EditAction.bUserInvoked));
+				NodePtrSet.Add(NodePtr);
 			}
+
+			RegisterActiveTimer(0.f, FWidgetActiveTimerDelegate::CreateLambda(AddNodesDelegateWrapper, StaticCastWeakPtr<SGraphPanel>(AsWeak()), NodePtrSet, EditAction.bUserInvoked));
 		}
 		if (bWasSelectAction)
 		{
-			const auto SelectNodeDelegateWrapper = [](double, float, TSharedRef<SGraphPanel> Parent, TSet< TWeakObjectPtr<UEdGraphNode> > NodePtrs, bool bForceUserAdded) -> EActiveTimerReturnType
+			const auto SelectNodeDelegateWrapper = [](double, float, TWeakPtr<SGraphPanel> WeakParent, TSet< TWeakObjectPtr<UEdGraphNode> > NodePtrs, bool bForceUserAdded) -> EActiveTimerReturnType
 			{
-				if (Parent->bVisualUpdatePending)
+				TSharedPtr<SGraphPanel> Parent = WeakParent.Pin();
+				if (Parent.IsValid())
 				{
-					if (bForceUserAdded)
+					if (Parent->bVisualUpdatePending)
 					{
-						Parent->UserSelectedNodes = NodePtrs;
+						if (bForceUserAdded)
+						{
+							Parent->UserSelectedNodes = NodePtrs;
+						}
 					}
-				}
-				else
-				{
-					Parent->DeferredSelectionTargetObjects.Empty();
-					for (TWeakObjectPtr<UEdGraphNode>& NodePtr : NodePtrs)
+					else
 					{
-						if (NodePtr.IsValid())
+						Parent->DeferredSelectionTargetObjects.Empty();
+						for (TWeakObjectPtr<UEdGraphNode>& NodePtr : NodePtrs)
 						{
-							UEdGraphNode* Node = NodePtr.Get();
-							Parent->DeferredSelectionTargetObjects.Add(Node);
+							if (NodePtr.IsValid())
+							{
+								UEdGraphNode* Node = NodePtr.Get();
+								Parent->DeferredSelectionTargetObjects.Add(Node);
+							}
 						}
 					}
 				}
@@ -2509,7 +2536,7 @@ void SGraphPanel::OnGraphChanged(const FEdGraphEditAction& EditAction)
 				NodePtrSet.Add(NodePtr);
 			}
 
-			RegisterActiveTimer(0.f, FWidgetActiveTimerDelegate::CreateStatic(SelectNodeDelegateWrapper, StaticCastSharedRef<SGraphPanel>(AsShared()), NodePtrSet, EditAction.bUserInvoked));
+			RegisterActiveTimer(0.f, FWidgetActiveTimerDelegate::CreateLambda(SelectNodeDelegateWrapper, StaticCastWeakPtr<SGraphPanel>(AsWeak()), NodePtrSet, EditAction.bUserInvoked));
 		}
 		if (bWasEditAction)
 		{
