Blueprints: Fix dropdown for UPARAM's GetOptions was not hidden when a pin was connected

diff --git a/Engine/Source/Editor/GraphEditor/Private/KismetPins/SGraphPinString.cpp b/Engine/Source/Editor/GraphEditor/Private/KismetPins/SGraphPinString.cpp
--- a/Engine/Source/Editor/GraphEditor/Private/KismetPins/SGraphPinString.cpp
+++ b/Engine/Source/Editor/GraphEditor/Private/KismetPins/SGraphPinString.cpp
@@ -186,6 +186,8 @@ TSharedPtr<SWidget> SGraphPinString::TryBuildComboBoxWidget()
 		[
 			SNew(SSearchableComboBox)
 				.OptionsSource(&ComboBoxOptions)
+				.Visibility(this, &SGraphPin::GetDefaultValueVisibility)
+				.IsEnabled(this, &SGraphPin::GetDefaultValueIsEditable)
 				.OnGenerateWidget(this, &SGraphPinString::GenerateComboBoxEntry)
 				.OnSelectionChanged(this, &SGraphPinString::HandleComboBoxSelectionChanged)
 				.ContentPadding(3.0f)
