Blueprint Editor:

-Fix Panning on Linux is slow
-Add bUseHighPrecisionMouseMovement (true by default) and bUseInterpolationWithManualPanning (false by default) as Graph Editor Settings, set to the current behaviour, but can be switched to allow for a better panning experience on Linux versions of the editor.

diff --git a/Engine/Config/Linux/LinuxEditorPerProjectUserSettings.ini b/Engine/Config/Linux/LinuxEditorPerProjectUserSettings.ini
new file mode 100644
--- /dev/null
+++ b/Engine/Config/Linux/LinuxEditorPerProjectUserSettings.ini
@@ -0,0 +1,23 @@
+
+; -------------------------------------------------------------------------------------------------
+; IMPORTANT:  EditorPerProjectUserSettings.ini has special behavior!
+;
+; This .ini file contains the defaults for many editor and engine "preferences".  These preferences
+; are saved to an EditorPerProjectUserSettings.ini file in the user's /<Game>/Saved/Config/ directory. 
+;
+; If you change the default for a preference in here, that change will *NOT* automatically 
+; propagate to users who already have saved preferences.  The idea is that we must preserve
+; their existing settings and never want to clobber entries in that file.
+;
+; If you add a new preference, the new key and value *WILL* be propagated to the user's
+; EditorPerProjectUserSettings.ini file.  After we load the user's existing settings, we'll merge in
+; any missing keys and values that are present in these defaults.
+;
+; One easy technique for forcing a "reset" of an already-existing preference is to simply
+; rename the variable for that preference.  The config system will treat it as a newly-added
+; key and will propagate the change to existing users' preferences.
+; -------------------------------------------------------------------------------------------------
+
+[/Script/GraphEditor.GraphEditorSettings]
+bUseHighPrecisionMouseMovement=False
+bUseInterpolationWithManualPanning=True
diff --git a/Engine/Source/Editor/GraphEditor/Private/GraphEditorSettings.cpp b/Engine/Source/Editor/GraphEditor/Private/GraphEditorSettings.cpp
--- a/Engine/Source/Editor/GraphEditor/Private/GraphEditorSettings.cpp
+++ b/Engine/Source/Editor/GraphEditor/Private/GraphEditorSettings.cpp
@@ -11,6 +11,8 @@
 UGraphEditorSettings::UGraphEditorSettings( const FObjectInitializer& ObjectInitializer )
 	: Super(ObjectInitializer)
 	, PanningMouseButton(EGraphPanningMouseButton::Right)
+	, bUseHighPrecisionMouseMovement(true)
+	, bUseInterpolationWithManualPanning(false)
 	, PaddingAbovePin(4.0f)
 	, PaddingBelowPin(4.0f)
 	, PaddingRightOfInput(10.0f)
diff --git a/Engine/Source/Editor/GraphEditor/Private/SNodePanel.cpp b/Engine/Source/Editor/GraphEditor/Private/SNodePanel.cpp
--- a/Engine/Source/Editor/GraphEditor/Private/SNodePanel.cpp
+++ b/Engine/Source/Editor/GraphEditor/Private/SNodePanel.cpp
@@ -617,6 +617,30 @@ FReply SNodePanel::OnMouseButtonDown( const FGeometry& MyGeometry, const FPointe
 
 	TotalMouseDelta = 0;
 
+	auto PerformMousePan = [this, &MyGeometry, &MouseEvent]()
+	{
+		// Cache current cursor position as zoom origin and software cursor position
+		ZoomStartOffset = MyGeometry.AbsoluteToLocal( MouseEvent.GetLastScreenSpacePosition() );
+		SoftwareCursorPosition = PanelCoordToGraphCoord( ZoomStartOffset );
+
+		FReply ReplyState = FReply::Handled();
+		ReplyState.CaptureMouse( SharedThis(this) );
+		if (GetDefault<UGraphEditorSettings>()->bUseHighPrecisionMouseMovement)
+		{
+			ReplyState.UseHighPrecisionMouseMovement( SharedThis(this) );
+		}
+
+		SoftwareCursorPosition = PanelCoordToGraphCoord( MyGeometry.AbsoluteToLocal( MouseEvent.GetScreenSpacePosition() ) );
+
+		if (!GetDefault<UGraphEditorSettings>()->bUseInterpolationWithManualPanning)
+		{
+			DeferredMovementTargetObject = nullptr; // clear any interpolation when you manually pan
+		}
+		CancelZoomToFit();
+
+		return ReplyState;
+	};
+
 	if ((bIsLeftMouseButtonEffecting && bIsRightMouseButtonDown)
 	||  (bIsRightMouseButtonEffecting && (bIsLeftMouseButtonDown || FSlateApplication::Get().IsUsingTrackpad())))
 	{
@@ -673,39 +697,11 @@ FReply SNodePanel::OnMouseButtonDown( const FGeometry& MyGeometry, const FPointe
 	}
 	else if (bIsRightMouseButtonEffecting && ( GetDefault<UGraphEditorSettings>()->PanningMouseButton == EGraphPanningMouseButton::Right || GetDefault<UGraphEditorSettings>()->PanningMouseButton == EGraphPanningMouseButton::Both ) )
 	{
-		// Cache current cursor position as zoom origin and software cursor position
-		ZoomStartOffset = MyGeometry.AbsoluteToLocal( MouseEvent.GetLastScreenSpacePosition() );
-		SoftwareCursorPosition = PanelCoordToGraphCoord( ZoomStartOffset );
-
-		FReply ReplyState = FReply::Handled();
-		ReplyState.CaptureMouse( SharedThis(this) );
-		ReplyState.UseHighPrecisionMouseMovement( SharedThis(this) );
-
-		SoftwareCursorPosition = PanelCoordToGraphCoord( MyGeometry.AbsoluteToLocal( MouseEvent.GetScreenSpacePosition() ) );
-
-		DeferredMovementTargetObject = nullptr; // clear any interpolation when you manually pan
-		CancelZoomToFit();
-
-		// RIGHT BUTTON is for dragging and Context Menu.
-		return ReplyState;
+		return PerformMousePan();
 	}
 	else if (bIsMiddleMouseButtonEffecting && (GetDefault<UGraphEditorSettings>()->PanningMouseButton == EGraphPanningMouseButton::Middle || GetDefault<UGraphEditorSettings>()->PanningMouseButton == EGraphPanningMouseButton::Both))
 	{
-		// Cache current cursor position as zoom origin and software cursor position
-		ZoomStartOffset = MyGeometry.AbsoluteToLocal(MouseEvent.GetLastScreenSpacePosition());
-		SoftwareCursorPosition = PanelCoordToGraphCoord(ZoomStartOffset);
-
-		FReply ReplyState = FReply::Handled();
-		ReplyState.CaptureMouse(SharedThis(this));
-		ReplyState.UseHighPrecisionMouseMovement(SharedThis(this));
-
-		SoftwareCursorPosition = PanelCoordToGraphCoord(MyGeometry.AbsoluteToLocal(MouseEvent.GetScreenSpacePosition()));
-
-		DeferredMovementTargetObject = nullptr; // clear any interpolation when you manually pan
-		CancelZoomToFit();
-
-		// MIDDLE BUTTON is for dragging only.
-		return ReplyState;
+		return PerformMousePan();
 	}
 	else if ( bIsLeftMouseButtonEffecting )
 	{
diff --git a/Engine/Source/Editor/GraphEditor/Public/GraphEditorSettings.h b/Engine/Source/Editor/GraphEditor/Public/GraphEditorSettings.h
--- a/Engine/Source/Editor/GraphEditor/Public/GraphEditorSettings.h
+++ b/Engine/Source/Editor/GraphEditor/Public/GraphEditorSettings.h
@@ -52,6 +52,14 @@ public:
 	UPROPERTY(config, EditAnywhere, Category = GeneralStyle)
 	EGraphPanningMouseButton PanningMouseButton;
 
+	/** Whether to use high precision mouse movement for more accurate mouse movement without mouse ballistics (defaults to true) */
+	UPROPERTY(config, EditAnywhere, Category = GeneralStyle)
+	bool bUseHighPrecisionMouseMovement;
+
+	/** Whether to interpolate movement when manually panning a graph for smoother movement (defaults to false) */
+	UPROPERTY(config, EditAnywhere, Category = GeneralStyle)
+	bool bUseInterpolationWithManualPanning;
+
 	/** The amount of padding above a pin (defaults to 4) */
 	UPROPERTY(config, EditAnywhere, Category=GeneralStyle, AdvancedDisplay)
 	float PaddingAbovePin;
