[CrashFix] A udn user found a crash in LinkedToDifferent in GraphDiffControl.cpp. This change adds an ensure to catch and recover from it

diff --git a/Engine/Source/Editor/GraphEditor/Private/GraphDiffControl.cpp b/Engine/Source/Editor/GraphEditor/Private/GraphDiffControl.cpp
--- a/Engine/Source/Editor/GraphEditor/Private/GraphDiffControl.cpp
+++ b/Engine/Source/Editor/GraphEditor/Private/GraphDiffControl.cpp
@@ -495,7 +495,15 @@ static bool LinkedToDifferent(UEdGraphPin* OldPin, UEdGraphPin* NewPin, const TA
 	for(int32 i = 0;i<Size;++i)
 	{
 		UEdGraphPin* OldLinkedPin = OldLinks[i];
+		if (!ensureMsgf(OldLinkedPin, TEXT("Found unexpected nullptr in OldLinks. Node graph was either written with null links or pin resolution failed - see UEdGraphPin::ResolveAllPinReferences")))
+		{
+			continue;
+		}
 		UEdGraphPin* NewLinkedPin = FindOtherLink(NewLinks, i, OldLinkedPin);
+		if (!ensureMsgf(NewLinkedPin, TEXT("Found unexpected nullptr in NewLinks. Node graph was either written with null links or pin resolution failed - see UEdGraphPin::ResolveAllPinReferences")))
+		{
+			continue;
+		}
 
 		UEdGraphNode* OldNode = OldLinkedPin->GetOwningNode();
 		UEdGraphNode* NewNode = NewLinkedPin->GetOwningNode();
@@ -530,6 +538,14 @@ static bool ArePinsDifferent(const TArray<UEdGraphPin*>& OldPins, TArray<UEdGrap
 	{
 		UEdGraphPin* OldPin = OldPins[i];
 		UEdGraphPin* NewPin = NewPins[i];
+		if (!ensureMsgf(OldPin, TEXT("Found unexpected nullptr in OldPins. Node graph was either written with null links or pin resolution failed - see UEdGraphPin::ResolveAllPinReferences")))
+		{
+			continue;
+		}
+		if (!ensureMsgf(NewPin, TEXT("Found unexpected nullptr in NewPins. Node graph was either written with null links or pin resolution failed - see UEdGraphPin::ResolveAllPinReferences")))
+		{
+			continue;
+		}
 
 		const UEdGraphSchema* Schema = OldPin->GetSchema();
 
