Remove heading from "Find References" submenu

diff --git a/Engine/Source/Editor/GraphEditor/Private/GraphEditorActions.cpp b/Engine/Source/Editor/GraphEditor/Private/GraphEditorActions.cpp
--- a/Engine/Source/Editor/GraphEditor/Private/GraphEditorActions.cpp
+++ b/Engine/Source/Editor/GraphEditor/Private/GraphEditorActions.cpp
@@ -175,7 +175,7 @@ const FGraphEditorCommandsImpl& FGraphEditorCommands::Get()
 
 void FGraphEditorCommands::BuildFindReferencesMenu(FMenuBuilder& MenuBuilder)
 {
-	MenuBuilder.BeginSection("FindReferences", LOCTEXT("FindReferences", "Find References"));
+	MenuBuilder.BeginSection("FindReferences");
 	{
 		MenuBuilder.AddMenuEntry(FGraphEditorCommands::Get().FindReferencesByNameLocal);
 		MenuBuilder.AddMenuEntry(FGraphEditorCommands::Get().FindReferencesByNameGlobal);
