Rework SGraphActionMenu such that there is no need to walk all existing nodes during menu building, GetLeafNodes is instead called only during keyboard navigation. This speeds up asynchronous (and synchronous) action menu construction. FGraphActionNode now tracks the number of 'leaf' nodes that have been added to it, SGraphActionMenu now tracks the selected node, rather than its leaf index (which is unstable and would require expensive recalculation)

diff --git a/Engine/Source/Editor/GraphEditor/Private/GraphActionNode.cpp b/Engine/Source/Editor/GraphEditor/Private/GraphActionNode.cpp
--- a/Engine/Source/Editor/GraphEditor/Private/GraphActionNode.cpp
+++ b/Engine/Source/Editor/GraphEditor/Private/GraphActionNode.cpp
@@ -179,6 +179,7 @@ FGraphActionNode::FGraphActionNode(int32 InGrouping, int32 InSectionID)
 	, Grouping(InGrouping)
 	, bPendingRenameRequest(false)
 	, InsertOrder(0)
+	, TotalLeafs(0)
 {
 }
 
@@ -189,6 +190,7 @@ FGraphActionNode::FGraphActionNode(TArray< TSharedPtr<FEdGraphSchemaAction> > co
 	, Actions(ActionList)
 	, bPendingRenameRequest(false)
 	, InsertOrder(0)
+	, TotalLeafs(0)
 {
 }
 
@@ -198,17 +200,27 @@ TSharedPtr<FGraphActionNode> FGraphActionNode::AddChild(FGraphActionListBuilderB
 	const TArray<FString>& CategoryStack = ActionSet.GetCategoryChain();
 
 	TSharedPtr<FGraphActionNode> ActionNode = FGraphActionNode::NewActionNode(ActionSet.Actions);
+	if (!ActionNode->IsCategoryNode() && !ActionNode->IsSectionHeadingNode())
+	{
+		++TotalLeafs;
+	}
+
 	AddChildRecursively(CategoryStack, 0, ActionNode);
 	
 	return ActionNode;
 }
 
 //------------------------------------------------------------------------------
-void FGraphActionNode::AddChildAlphabetical(FGraphActionListBuilderBase::ActionGroup const& ActionSet)
+TSharedPtr<FGraphActionNode> FGraphActionNode::AddChildAlphabetical(FGraphActionListBuilderBase::ActionGroup const& ActionSet)
 {
 	TSharedPtr<FGraphActionNode> ActionNode = FGraphActionNode::NewActionNode(ActionSet.Actions);
 	check(ActionNode->SectionID == INVALID_SECTION_ID); // this method does not support sections, those should be built statically
 
+	if (!ActionNode->IsCategoryNode() && !ActionNode->IsSectionHeadingNode())
+	{
+		++TotalLeafs;
+	}
+
 	// if a divider hasn't been created for the grouping, create one:
 	AddChildGrouping(ActionNode, this->AsShared());
 
@@ -229,6 +241,7 @@ void FGraphActionNode::AddChildAlphabetical(FGraphActionListBuilderBase::ActionG
 
 	// finally insert the leaf:
 	OwningCategory->InsertChildAlphabetical(ActionNode);
+	return ActionNode;
 }
 
 //------------------------------------------------------------------------------
@@ -311,6 +324,11 @@ void FGraphActionNode::GetLeafNodes(TArray< TSharedPtr<FGraphActionNode> >& OutL
 	}
 }
 
+int32 FGraphActionNode::GetTotalLeafNodes() const 
+{ 
+	return TotalLeafs; 
+}
+
 //------------------------------------------------------------------------------
 void FGraphActionNode::ExpandAllChildren(TSharedPtr< STreeView< TSharedPtr<FGraphActionNode> > > TreeView, bool bRecursive/*= true*/)
 {
@@ -334,6 +352,7 @@ void FGraphActionNode::ExpandAllChildren(TSharedPtr< STreeView< TSharedPtr<FGrap
 //------------------------------------------------------------------------------
 void FGraphActionNode::ClearChildren()
 {
+	TotalLeafs = 0;
 	Children.Empty();
 	CategoryNodes.Empty();
 	ChildGroupings.Empty();
diff --git a/Engine/Source/Editor/GraphEditor/Private/SGraphActionMenu.cpp b/Engine/Source/Editor/GraphEditor/Private/SGraphActionMenu.cpp
--- a/Engine/Source/Editor/GraphEditor/Private/SGraphActionMenu.cpp
+++ b/Engine/Source/Editor/GraphEditor/Private/SGraphActionMenu.cpp
@@ -307,7 +307,7 @@ void SGraphActionMenu::Construct( const FArguments& InArgs, bool bIsReadOnly/* =
 {
 	this->SelectedSuggestionScore = TNumericLimits<float>::Lowest();
 	this->SelectedSuggestionSourceIndex = INDEX_NONE;
-	this->SelectedSuggestion = INDEX_NONE;
+	this->SelectedAction = TSharedPtr<FGraphActionNode>();
 	this->bIgnoreUIUpdate = false;
 	this->bUseSectionStyling = InArgs._UseSectionStyling;
 	this->bAllowPreselectedItemActivation = InArgs._bAllowPreselectedItemActivation;
@@ -426,7 +426,7 @@ void SGraphActionMenu::RefreshAllActions(bool bPreserveExpansion, bool bHandleOn
 
 	// Save Selection (of only the first selected thing)
 	TArray< TSharedPtr<FGraphActionNode> > SelectedNodes = TreeView->GetSelectedItems();
-	TSharedPtr<FGraphActionNode> SelectedAction = SelectedNodes.Num() > 0 ? SelectedNodes[0] : nullptr;
+	TSharedPtr<FGraphActionNode> CurrentSelectedAction = SelectedNodes.Num() > 0 ? SelectedNodes[0] : nullptr;
 
 	if (OnGetActionList.IsBound())
 	{
@@ -445,19 +445,19 @@ void SGraphActionMenu::RefreshAllActions(bool bPreserveExpansion, bool bHandleOn
 	GenerateFilteredItems(bPreserveExpansion);
 
 	// Re-apply selection #0 if possible
-	if (SelectedAction.IsValid())
+	if (CurrentSelectedAction.IsValid())
 	{
 		// Clear the selection, we will be re-selecting the previous action
 		TreeView->ClearSelection();
 
 		if(bHandleOnSelectionEvent)
 		{
-			SelectItemByName(*SelectedAction->GetDisplayName().ToString(), ESelectInfo::OnMouseClick, SelectedAction->SectionID, SelectedNodes[0]->IsCategoryNode());
+			SelectItemByName(*CurrentSelectedAction->GetDisplayName().ToString(), ESelectInfo::OnMouseClick, CurrentSelectedAction->SectionID, SelectedNodes[0]->IsCategoryNode());
 		}
 		else
 		{
 			// If we do not want to handle the selection, set it directly so it will reselect the item but not handle the event.
-			SelectItemByName(*SelectedAction->GetDisplayName().ToString(), ESelectInfo::Direct, SelectedAction->SectionID, SelectedNodes[0]->IsCategoryNode());
+			SelectItemByName(*CurrentSelectedAction->GetDisplayName().ToString(), ESelectInfo::Direct, CurrentSelectedAction->SectionID, SelectedNodes[0]->IsCategoryNode());
 		}
 	}
 }
@@ -733,8 +733,8 @@ void SGraphActionMenu::UpdateForNewActions(int32 IdxStart)
 {
 	check(bAlphaSortItems && bSortItemsRecursively);
 
-	FScoreResults Results = ScoreAndAddActions(IdxStart);
-	UpdateActiveSelection(Results);
+	ScoreAndAddActions(IdxStart);
+	MarkActiveSuggestion();
 
 	if (ShouldExpandNodes())
 	{
@@ -751,6 +751,7 @@ void SGraphActionMenu::GenerateFilteredItems(bool bPreserveExpansion)
 
 	SelectedSuggestionScore = TNumericLimits<float>::Lowest();
 	SelectedSuggestionSourceIndex = INDEX_NONE;
+	SelectedAction = TSharedPtr<FGraphActionNode>();
 
 	// First, save off current expansion state
 	TSet< TSharedPtr<FGraphActionNode> > OldExpansionState;
@@ -774,13 +775,13 @@ void SGraphActionMenu::GenerateFilteredItems(bool bPreserveExpansion)
 		}
 	}
 	
-	FScoreResults Results = ScoreAndAddActions();
+	ScoreAndAddActions();
 
 	FilteredRootAction->SortChildren(bAlphaSortItems, bSortItemsRecursively);
 
 	TreeView->RequestTreeRefresh();
 
-	UpdateActiveSelection(Results);
+	MarkActiveSuggestion();
 
 	if (ShouldExpandNodes())
 	{
@@ -815,11 +816,6 @@ bool SGraphActionMenu::CanRenameNode(TWeakPtr<FGraphActionNode> InNode) const
 
 void SGraphActionMenu::OnFilterTextChanged( const FText& InFilterText )
 {
-	// Reset the selection if the string is empty
-	if( InFilterText.IsEmpty() == true )
-	{
-		SelectedSuggestion = INDEX_NONE;
-	}
 	GenerateFilteredItems(false);
 }
 
@@ -841,9 +837,9 @@ bool SGraphActionMenu::TryToSpawnActiveSuggestion()
 		OnItemSelected( SelectionList[0], ESelectInfo::OnKeyPress );
 		return true;
 	}
-	else if (FilteredActionNodes.Num() == 1)
+	else if (GetTotalLeafNodes() == 1)
 	{
-		OnItemSelected( FilteredActionNodes[0], ESelectInfo::OnKeyPress );
+		OnItemSelected( GetFirstAction(), ESelectInfo::OnKeyPress);
 		return true;
 	}
 
@@ -1132,11 +1128,11 @@ bool SGraphActionMenu::OnMouseButtonDownEvent( TWeakPtr<FEdGraphSchemaAction> In
 		TSharedPtr<FGraphActionNode> SelectedNode;
 		if (SelectionList.Num() == 1)
 		{	
-			SelectedNode = SelectionList[0];			
+			SelectedNode = SelectionList[0];
 		}
-		else if (FilteredActionNodes.Num() == 1)
+		else if (GetTotalLeafNodes() == 1)
 		{
-			SelectedNode = FilteredActionNodes[0];			
+			SelectedNode = GetFirstAction();
 		}
 		if (SelectedNode.IsValid() && SelectedNode->HasValidAction())
 		{
@@ -1166,36 +1162,36 @@ FReply SGraphActionMenu::OnKeyDown( const FGeometry& MyGeometry, const FKeyEvent
 	else if (!FilterTextBox->GetText().IsEmpty())
 	{
 		// Needs to be done here in order not to eat up the text navigation key events when list isn't populated
-		if (FilteredActionNodes.Num() == 0)
+		if (GetTotalLeafNodes() == 0)
 		{
 			return FReply::Unhandled();
 		}
 
 		if (KeyEvent.GetKey() == EKeys::Up)
 		{
-			SelectedSuggestion = FMath::Max(0, SelectedSuggestion - 1);
+			SelectPreviousAction();
 		}
 		else if (KeyEvent.GetKey() == EKeys::Down)
 		{
-			SelectedSuggestion = FMath::Min(FilteredActionNodes.Num() - 1, SelectedSuggestion + 1);
+			SelectNextAction();
 		}
 		else if (KeyEvent.GetKey() == EKeys::PageUp)
 		{
 			const int32 NumItemsInAPage = 15; // arbitrary jump because we can't get at the visible item count from here
-			SelectedSuggestion = FMath::Max(0, SelectedSuggestion - NumItemsInAPage);
+			SelectPreviousAction(NumItemsInAPage);
 		}
 		else if (KeyEvent.GetKey() == EKeys::PageDown)
 		{
 			const int32 NumItemsInAPage = 15; // arbitrary jump because we can't get at the visible item count from here
-			SelectedSuggestion = FMath::Min(FilteredActionNodes.Num() - 1, SelectedSuggestion + NumItemsInAPage);
+			SelectNextAction(NumItemsInAPage);
 		}
 		else if (KeyEvent.GetKey() == EKeys::Home && KeyEvent.IsControlDown())
 		{
-			SelectedSuggestion = 0;
+			SelectFirstAction();
 		}
 		else if (KeyEvent.GetKey() == EKeys::End && KeyEvent.IsControlDown())
 		{
-			SelectedSuggestion = FilteredActionNodes.Num() - 1;
+			SelectLastAction();
 		}
 		else
 		{
@@ -1218,12 +1214,10 @@ void SGraphActionMenu::MarkActiveSuggestion()
 {
 	TGuardValue<bool> PreventSelectionFromTriggeringCommit(bIgnoreUIUpdate, true);
 
-	if (SelectedSuggestion >= 0)
+	if (SelectedAction.IsValid())
 	{
-		TSharedPtr<FGraphActionNode>& ActionToSelect = FilteredActionNodes[SelectedSuggestion];
-
-		TreeView->SetSelection(ActionToSelect);
-		TreeView->RequestScrollIntoView(ActionToSelect);
+		TreeView->SetSelection(SelectedAction);
+		TreeView->RequestScrollIntoView(SelectedAction);
 	}
 	else
 	{
@@ -1281,7 +1275,7 @@ void SGraphActionMenu::OnSetExpansionRecursive(TSharedPtr<FGraphActionNode> InTr
 	}
 }
 
-SGraphActionMenu::FScoreResults SGraphActionMenu::ScoreAndAddActions(int32 StartingIndex)
+void SGraphActionMenu::ScoreAndAddActions(int32 StartingIndex)
 {
 	// Trim and sanitized the filter text (so that it more likely matches the action descriptions)
 	FString TrimmedFilterString = FText::TrimPrecedingAndTrailing(GetFilterText()).ToString();
@@ -1337,6 +1331,7 @@ SGraphActionMenu::FScoreResults SGraphActionMenu::ScoreAndAddActions(int32 Start
 			continue;
 		}
 
+		bool bAddedAndSelected = false;
 		if (bRequiresFiltering)
 		{
 			// Get the 'weight' of this in relation to the filter
@@ -1346,62 +1341,100 @@ SGraphActionMenu::FScoreResults SGraphActionMenu::ScoreAndAddActions(int32 Start
 			{
 				BestMatchCount = EachWeight;
 				BestMatchIndex = CurTypeIndex;
+
+				// as our currently best scoring entry, add and select the node:
+				if (bIsPartialBuild)
+				{
+					SelectedAction = FilteredRootAction->AddChildAlphabetical(CurrentAction);
+				}
+				else
+				{
+					SelectedAction = FilteredRootAction->AddChild(CurrentAction);
+				}
+				bAddedAndSelected = true;
 			}
 		}
 
-		if (bIsPartialBuild)
+		if (!bAddedAndSelected) // if the node was not added and selected, then just add it to the root:
 		{
-			FilteredRootAction->AddChildAlphabetical(CurrentAction);
-		}
-		else
-		{
-			FilteredRootAction->AddChild(CurrentAction);
+			if (bIsPartialBuild)
+			{
+				FilteredRootAction->AddChildAlphabetical(CurrentAction);
+			}
+			else
+			{
+				FilteredRootAction->AddChild(CurrentAction);
+			}
 		}
 	}
 
-	return {BestMatchIndex, BestMatchCount};
+	SelectedSuggestionScore = BestMatchCount;
+	SelectedSuggestionSourceIndex = BestMatchIndex;
+}
+
+void SGraphActionMenu::SelectPreviousAction(int32 Num)
+{
+	// search backwards Num entries for a previous action, stop if we reach the first action:
+	int32 SelectedIndex = INDEX_NONE;
+	const TArray< TSharedPtr<FGraphActionNode> >& CurrentFilteredActionNodes = GetFilteredActionNodes(&SelectedIndex);
+	SelectedIndex = FMath::Max(0, SelectedIndex - Num);
+	SelectedAction = CurrentFilteredActionNodes.Num() > 0 ? CurrentFilteredActionNodes[SelectedIndex] : TSharedPtr<FGraphActionNode>();;
+}
+
+void SGraphActionMenu::SelectNextAction(int32 Num)
+{
+	// search forwards Num entries for a next action, stop if we reach the first action:
+	int32 SelectedIndex = INDEX_NONE;
+	const TArray< TSharedPtr<FGraphActionNode> >& CurrentFilteredActionNodes = GetFilteredActionNodes(&SelectedIndex);
+	SelectedIndex = FMath::Min(CurrentFilteredActionNodes.Num() - 1, SelectedIndex + Num);
+	SelectedAction = CurrentFilteredActionNodes.Num() > 0 ? CurrentFilteredActionNodes[SelectedIndex] : TSharedPtr<FGraphActionNode>();;
+}
+
+void SGraphActionMenu::SelectFirstAction()
+{
+	SelectedAction = GetFirstAction();
+}
+
+void SGraphActionMenu::SelectLastAction()
+{
+	// find the last unfiltered action:
+	const TArray< TSharedPtr<FGraphActionNode> >& CurrentFilteredActionNodes = GetFilteredActionNodes();
+	SelectedAction = CurrentFilteredActionNodes.Num() > 0 ? CurrentFilteredActionNodes.Last() : TSharedPtr<FGraphActionNode>();
+}
+
+TSharedPtr<FGraphActionNode> SGraphActionMenu::GetFirstAction()
+{
+	const TArray< TSharedPtr<FGraphActionNode> >& Nodes = GetFilteredActionNodes();
+	return Nodes.Num() > 0 ? FilteredActionNodes[0] : TSharedPtr<FGraphActionNode>();
 }
 
-void SGraphActionMenu::UpdateActiveSelection(SGraphActionMenu::FScoreResults ForResults)
+const TArray< TSharedPtr<FGraphActionNode> >& SGraphActionMenu::GetFilteredActionNodes(int32* OutSelectedIndex)
 {
-	int32 BestMatchIndex = ForResults.BestMatchIndex;
-	float BestMatchCount = ForResults.BestMatchScore;
-	// Update the filtered list (needs to be done in a separate pass because the list is sorted as items are inserted)
+	// We could cache this, but for now I'm calculating it every time it is requested - this
+	// possibility of caching is the reason none of these methods are const
 	FilteredActionNodes.Reset();
 	FilteredRootAction->GetLeafNodes(FilteredActionNodes);
-
-	// If theres a BestMatchIndex find it in the actions nodes and select it (maybe this should check the current selected suggestion first ?)
-	if (BestMatchIndex != INDEX_NONE)
+	ensureMsgf(GetTotalLeafNodes() == FilteredActionNodes.Num(), TEXT("FilteredActionNodes and GetTotalLeafNodes should match"));
+	// find selected item's leaf index:
+	if (OutSelectedIndex && SelectedAction)
 	{
-		FGraphActionListBuilderBase::ActionGroup& FilterSelectAction = AllActions->GetAction(BestMatchIndex);
-		if (FilterSelectAction.Actions[0].IsValid() == true)
+		for (int32 Idx = 0; Idx < FilteredActionNodes.Num(); ++Idx)
 		{
-			for (int32 iNode = 0; iNode < FilteredActionNodes.Num(); iNode++)
+			if (FilteredActionNodes[Idx] == SelectedAction)
 			{
-				if (FilteredActionNodes[iNode].Get()->GetPrimaryAction() == FilterSelectAction.Actions[0])
-				{
-					SelectedSuggestion = iNode;
-					SelectedSuggestionScore = BestMatchCount;
-					SelectedSuggestionSourceIndex = BestMatchIndex;
-				}
+				*OutSelectedIndex = Idx;
 			}
 		}
 	}
+	return FilteredActionNodes;
+}
 
-	// Make sure the selected suggestion stays within the filtered list
-	if ((SelectedSuggestion >= 0) && (FilteredActionNodes.Num() > 0))
-	{
-		//@TODO: Should try to actually maintain the highlight on the same item if it survived the filtering
-		SelectedSuggestion = FMath::Clamp<int32>(SelectedSuggestion, 0, FilteredActionNodes.Num() - 1);
-		MarkActiveSuggestion();
-	}
-	else
-	{
-		SelectedSuggestionScore = TNumericLimits<float>::Lowest();
-		SelectedSuggestionSourceIndex = INDEX_NONE;
-		SelectedSuggestion = INDEX_NONE;
-	}
+int32 SGraphActionMenu::GetTotalLeafNodes() const
+{
+	return FilteredRootAction->GetTotalLeafNodes();
 }
+
+
 /////////////////////////////////////////////////////
 
 #undef LOCTEXT_NAMESPACE
diff --git a/Engine/Source/Editor/GraphEditor/Public/GraphActionNode.h b/Engine/Source/Editor/GraphEditor/Public/GraphActionNode.h
--- a/Engine/Source/Editor/GraphEditor/Public/GraphActionNode.h
+++ b/Engine/Source/Editor/GraphEditor/Public/GraphActionNode.h
@@ -61,8 +61,9 @@ public:
 	 * the provided ActionSet. Inserts in alphabetical order.
 	 *
 	 * @param  ActionSet	A list of actions that you want the node to execute when picked.
+	 * @return The new action node.
 	 */
-	void AddChildAlphabetical(FGraphActionListBuilderBase::ActionGroup const& ActionSet);
+	TSharedPtr<FGraphActionNode> AddChildAlphabetical(FGraphActionListBuilderBase::ActionGroup const& ActionSet);
 
 	TSharedPtr<FGraphActionNode> AddSection(int32 Grouping, int32 InSectionID);
 
@@ -102,6 +103,9 @@ public:
 	 */
 	void GetLeafNodes(TArray< TSharedPtr<FGraphActionNode> >& OutLeafArray) const;
 
+	/** Returns the number of leaf nodes */
+	int32 GetTotalLeafNodes() const;
+
 	/**
 	 * Takes the tree view and expands its elements for each child.
 	 * 
@@ -298,5 +302,7 @@ private:
 
 	friend struct FGraphActionNodeImpl;
 	/** For sorting, when we don't alphabetically sort (so menu items don't jump around). */
-	int32 InsertOrder;	
+	int32 InsertOrder;
+	/** Root entry only, counts the total leaf entries in this tree */
+	int32 TotalLeafs;
 };
diff --git a/Engine/Source/Editor/GraphEditor/Public/SGraphActionMenu.h b/Engine/Source/Editor/GraphEditor/Public/SGraphActionMenu.h
--- a/Engine/Source/Editor/GraphEditor/Public/SGraphActionMenu.h
+++ b/Engine/Source/Editor/GraphEditor/Public/SGraphActionMenu.h
@@ -215,8 +215,8 @@ protected:
 	/** List of all actions we can browser */
 	TSharedPtr<FGraphActionListBuilderBase> AllActions;
 
-	/** Flattened list of all actions passing the filter */
-	TArray< TSharedPtr<FGraphActionNode> > FilteredActionNodes; 
+	/** Flattened list of all actions passing the filter - access via GetFilteredActionNodes to ensure it's up to date */
+	TArray< TSharedPtr<FGraphActionNode> > FilteredActionNodes;
 
 	/** Root of filtered actions tree */
 	TSharedPtr<FGraphActionNode> FilteredRootAction;
@@ -225,8 +225,9 @@ protected:
 	float SelectedSuggestionScore;
 	/** Stored index in the AllActions list builder that we think is the best fit */
 	int32 SelectedSuggestionSourceIndex;
-	/** Used to track selected action for keyboard interaction */
-	int32 SelectedSuggestion;
+	/** The actual item selected - possibly redundant to selection information in the treeview, but the tree is rebuilt requently */
+	TSharedPtr<FGraphActionNode> SelectedAction;
+
 	/** Allows us to set selection (via keyboard) without triggering action */
 	bool bIgnoreUIUpdate;
 	/** Should we auto-expand categories */
@@ -324,7 +325,7 @@ public:
 
 	/** Updates the displayed list starting from IdxStart, useful for async building the display list of actions */
 	void UpdateForNewActions(int32 IdxStart);
-	/** Regenerated filtered results (FilteredRootAction and FilteredActionNodes) based on filter text  */ 
+	/** Regenerated filtered results (FilteredRootAction) based on filter text  */ 
 	void GenerateFilteredItems(bool bPreserveExpansion);
 
 	/** The last typed action within the graph action menu */
@@ -385,14 +386,17 @@ protected:
 	/** Callback for expanding tree items recursively */
 	void OnSetExpansionRecursive(TSharedPtr<FGraphActionNode> InTreeNode, bool bInIsItemExpanded);
 	/** Helper function for adding and scoring actions from our builder */
-	struct FScoreResults
-	{
-		int32 BestMatchIndex;
-		float BestMatchScore;
-	};
-	FScoreResults ScoreAndAddActions(int32 StartingIndex = INDEX_NONE);
-	/** Helper function to update the active selection after updating the displayed tree */
-	void UpdateActiveSelection(FScoreResults ForResults);
+	void ScoreAndAddActions(int32 StartingIndex = INDEX_NONE);
+
+	/** Helper functions for keyboard interaction */
+	void SelectPreviousAction(int32 Num = 1);
+	void SelectNextAction(int32 Num = 1);
+	void SelectFirstAction();
+	void SelectLastAction();
+	/** General helper functions for accessing the filtered tree */
+	TSharedPtr<FGraphActionNode> GetFirstAction();
+	const TArray< TSharedPtr<FGraphActionNode> >& GetFilteredActionNodes(int32* OutSelectedIndex = nullptr);
+	int32 GetTotalLeafNodes() const;
 private:
 	/** The pins that have been dragged off of to prompt the creation of this action menu. */
 	TArray<UEdGraphPin*> DraggedFromPins;
