Made the focussed menu entry stable in the blueprint context menu, even while populating with a filter string

diff --git a/Engine/Source/Editor/GraphEditor/Private/GraphActionNode.cpp b/Engine/Source/Editor/GraphEditor/Private/GraphActionNode.cpp
--- a/Engine/Source/Editor/GraphEditor/Private/GraphActionNode.cpp
+++ b/Engine/Source/Editor/GraphEditor/Private/GraphActionNode.cpp
@@ -464,6 +464,13 @@ bool FGraphActionNode::IsRenameRequestPending() const
 	return bPendingRenameRequest;
 }
 
+//------------------------------------------------------------------------------
+int32 FGraphActionNode::GetLinearizedIndex(TSharedPtr<FGraphActionNode> Node) const
+{
+	int32 Counter = 0;
+	return GetLinearizedIndex(Node, Counter);
+}
+
 //------------------------------------------------------------------------------
 TSharedPtr<FGraphActionNode> FGraphActionNode::NewSectionHeadingNode(TWeakPtr<FGraphActionNode> Parent, int32 Grouping, int32 SectionID)
 {
@@ -722,3 +729,26 @@ void FGraphActionNode::InsertChildAlphabetical(TSharedPtr<FGraphActionNode> Node
 	}
 }
 
+int32 FGraphActionNode::GetLinearizedIndex(TSharedPtr<FGraphActionNode> Node, int32& Counter) const
+{
+	if (Node.Get() == this)
+	{
+		return Counter;
+	}
+
+	// we didn't match, count ourself:
+	++Counter;
+
+	// and check/count each child:
+	for (const TSharedPtr<FGraphActionNode>& Child : Children)
+	{
+		int32 Result = Child->GetLinearizedIndex(Node, Counter);
+		if (Result != INDEX_NONE)
+		{
+			return Result;
+		}
+	}
+
+	// no matches, return INDEX_NONE to indicate we found no valid index, Counter will continue counting
+	return INDEX_NONE;
+}
diff --git a/Engine/Source/Editor/GraphEditor/Private/SGraphActionMenu.cpp b/Engine/Source/Editor/GraphEditor/Private/SGraphActionMenu.cpp
--- a/Engine/Source/Editor/GraphEditor/Private/SGraphActionMenu.cpp
+++ b/Engine/Source/Editor/GraphEditor/Private/SGraphActionMenu.cpp
@@ -1217,7 +1217,8 @@ void SGraphActionMenu::MarkActiveSuggestion()
 	if (SelectedAction.IsValid())
 	{
 		TreeView->SetSelection(SelectedAction);
-		TreeView->RequestScrollIntoView(SelectedAction);
+		int32 Idx = FilteredRootAction->GetLinearizedIndex(SelectedAction);
+		TreeView->SetScrollOffset(FMath::Max(((float)Idx) - 2.f, 0.f));
 	}
 	else
 	{
diff --git a/Engine/Source/Editor/GraphEditor/Public/GraphActionNode.h b/Engine/Source/Editor/GraphEditor/Public/GraphActionNode.h
--- a/Engine/Source/Editor/GraphEditor/Public/GraphActionNode.h
+++ b/Engine/Source/Editor/GraphEditor/Public/GraphActionNode.h
@@ -198,6 +198,8 @@ public:
 	 */
 	bool IsRenameRequestPending() const;
 
+	/** Returns the 'linearized' index of the node, including category nodes, useful for getting displayed position */
+	int32 GetLinearizedIndex(TSharedPtr<FGraphActionNode> Node) const;
 private:
 	/**
 	 *
@@ -284,6 +286,9 @@ private:
 
 	void AddChildGrouping(TSharedPtr<FGraphActionNode> ActionNode, TWeakPtr<FGraphActionNode> Parent);
 	void InsertChildAlphabetical(TSharedPtr<FGraphActionNode> NodeToAdd);
+
+	/** Recursive implementation helper for GetLinearizedIndex */
+	int32 GetLinearizedIndex(TSharedPtr<FGraphActionNode> Node, int32& Iter) const;
 private:
 	/** The category or action name (depends on what type of node this is) */
 	FText DisplayText;
