Allow keyboard navigation in the action menu to disable scoring of added entries - this allows the user to keyboard navigate effectively while the menu is still building

diff --git a/Engine/Source/Editor/GraphEditor/Private/SGraphActionMenu.cpp b/Engine/Source/Editor/GraphEditor/Private/SGraphActionMenu.cpp
--- a/Engine/Source/Editor/GraphEditor/Private/SGraphActionMenu.cpp
+++ b/Engine/Source/Editor/GraphEditor/Private/SGraphActionMenu.cpp
@@ -56,6 +56,18 @@ struct FSlateBrush;
 
 //////////////////////////////////////////////////////////////////////////
 
+namespace UE::GraphEditor::Private
+{
+	// These constants control where we attempt to put the focused item
+	// this is the 'displayed index' of a menu entry, e.g. 2nd or 12th from
+	// the top of the list control. We're handling this ourself because
+	// this control has to amortize list building because the list is
+	// so large... a virtual list control that can handle >1m items 
+	// is probably the ideal.
+	const int32 PREFERRED_TOP_INDEX = 2;
+	const int32 PREFERRED_BOTTOM_INDEX = 12;
+}
+
 template<typename ItemType>
 class SCategoryHeaderTableRow : public STableRow<ItemType>
 {
@@ -310,9 +322,11 @@ void SGraphActionMenu::Construct( const FArguments& InArgs, bool bIsReadOnly/* =
 	this->SelectedAction = TSharedPtr<FGraphActionNode>();
 	this->bIgnoreUIUpdate = false;
 	this->bUseSectionStyling = InArgs._UseSectionStyling;
+	this->bIsKeyboardNavigating = false;
 	this->bAllowPreselectedItemActivation = InArgs._bAllowPreselectedItemActivation;
 	this->bAutomaticallySelectSingleAction = InArgs._bAutomaticallySelectSingleAction;
 	this->DefaultRowExpanderBaseIndentLevel = InArgs._DefaultRowExpanderBaseIndentLevel;
+	this->DisplayIndex = UE::GraphEditor::Private::PREFERRED_TOP_INDEX;
 
 	this->bAutoExpandActionMenu = InArgs._AutoExpandActionMenu;
 	this->bShowFilterTextBox = InArgs._ShowFilterTextBox;
@@ -1148,6 +1162,7 @@ bool SGraphActionMenu::OnMouseButtonDownEvent( TWeakPtr<FEdGraphSchemaAction> In
 FReply SGraphActionMenu::OnKeyDown( const FGeometry& MyGeometry, const FKeyEvent& KeyEvent )
 {
 	int32 SelectionDelta = 0;
+	bIsKeyboardNavigating = false;
 
 	// Escape dismisses the menu without placing a node
 	if (KeyEvent.GetKey() == EKeys::Escape)
@@ -1218,7 +1233,7 @@ void SGraphActionMenu::MarkActiveSuggestion()
 	{
 		TreeView->SetSelection(SelectedAction);
 		int32 Idx = FilteredRootAction->GetLinearizedIndex(SelectedAction);
-		TreeView->SetScrollOffset(FMath::Max(((float)Idx) - 2.f, 0.f));
+		TreeView->SetScrollOffset(FMath::Max(((float)Idx) - (float)DisplayIndex, 0.f));
 	}
 	else
 	{
@@ -1302,7 +1317,7 @@ void SGraphActionMenu::ScoreAndAddActions(int32 StartingIndex)
 	}
 	ensure(SanitizedFilterTerms.Num() == FilterTerms.Num());// Both of these should match !
 
-	const bool bRequiresFiltering = FilterTerms.Num() > 0;
+	const bool bRequiresFiltering = FilterTerms.Num() > 0 && !bIsKeyboardNavigating;
 	float BestMatchCount = SelectedSuggestionScore;
 	int32 BestMatchIndex = SelectedSuggestionSourceIndex;
 
@@ -1376,28 +1391,38 @@ void SGraphActionMenu::ScoreAndAddActions(int32 StartingIndex)
 void SGraphActionMenu::SelectPreviousAction(int32 Num)
 {
 	// search backwards Num entries for a previous action, stop if we reach the first action:
+	bIsKeyboardNavigating = true;
 	int32 SelectedIndex = INDEX_NONE;
 	const TArray< TSharedPtr<FGraphActionNode> >& CurrentFilteredActionNodes = GetFilteredActionNodes(&SelectedIndex);
 	SelectedIndex = FMath::Max(0, SelectedIndex - Num);
 	SelectedAction = CurrentFilteredActionNodes.Num() > 0 ? CurrentFilteredActionNodes[SelectedIndex] : TSharedPtr<FGraphActionNode>();;
+	int32 NextIndex = DisplayIndex - Num;
+	DisplayIndex = NextIndex < UE::GraphEditor::Private::PREFERRED_TOP_INDEX ? UE::GraphEditor::Private::PREFERRED_BOTTOM_INDEX : NextIndex; // if we bump below 2, loop over to 10, causing a scroll
 }
 
 void SGraphActionMenu::SelectNextAction(int32 Num)
 {
 	// search forwards Num entries for a next action, stop if we reach the first action:
+	bIsKeyboardNavigating = true;
 	int32 SelectedIndex = INDEX_NONE;
 	const TArray< TSharedPtr<FGraphActionNode> >& CurrentFilteredActionNodes = GetFilteredActionNodes(&SelectedIndex);
 	SelectedIndex = FMath::Min(CurrentFilteredActionNodes.Num() - 1, SelectedIndex + Num);
 	SelectedAction = CurrentFilteredActionNodes.Num() > 0 ? CurrentFilteredActionNodes[SelectedIndex] : TSharedPtr<FGraphActionNode>();;
+	int32 NextIndex = DisplayIndex + Num;
+	DisplayIndex = NextIndex > UE::GraphEditor::Private::PREFERRED_BOTTOM_INDEX ? UE::GraphEditor::Private::PREFERRED_TOP_INDEX : NextIndex; // if we bump below 2, loop over to 10, causing a scroll
 }
 
 void SGraphActionMenu::SelectFirstAction()
 {
+	bIsKeyboardNavigating = true;
+	DisplayIndex = UE::GraphEditor::Private::PREFERRED_TOP_INDEX;
 	SelectedAction = GetFirstAction();
 }
 
 void SGraphActionMenu::SelectLastAction()
 {
+	bIsKeyboardNavigating = true;
+	DisplayIndex = UE::GraphEditor::Private::PREFERRED_TOP_INDEX;
 	// find the last unfiltered action:
 	const TArray< TSharedPtr<FGraphActionNode> >& CurrentFilteredActionNodes = GetFilteredActionNodes();
 	SelectedAction = CurrentFilteredActionNodes.Num() > 0 ? CurrentFilteredActionNodes.Last() : TSharedPtr<FGraphActionNode>();
diff --git a/Engine/Source/Editor/GraphEditor/Public/SGraphActionMenu.h b/Engine/Source/Editor/GraphEditor/Public/SGraphActionMenu.h
--- a/Engine/Source/Editor/GraphEditor/Public/SGraphActionMenu.h
+++ b/Engine/Source/Editor/GraphEditor/Public/SGraphActionMenu.h
@@ -240,12 +240,16 @@ protected:
 	bool bSortItemsRecursively;
 	/** Should the rows and sections be styled like the details panel? */
 	bool bUseSectionStyling;
+	/** True if the user is using the keyboard to navigatethe list, halting scoring of newly added entries so that selection is stable */
+	bool bIsKeyboardNavigating;
 	/** Whether we allow pre-selected items to be activated with a left-click */
 	bool bAllowPreselectedItemActivation;
 	/** Whether to automatically proceed with an action if it's the only one in the list. */
 	bool bAutomaticallySelectSingleAction;
 	/** The BaseIndentLevel of the default-created row expander. Not used with OnCreateCustomRowExpander. */
 	int32 DefaultRowExpanderBaseIndentLevel;
+	/** The slot that the selected item should be displayed at */
+	int32 DisplayIndex;
 	
 	/** Delegate to call when action is selected */
 	FOnActionSelected OnActionSelected;
