#UE Add referencing asset to asset pickers spawned by SGraphPinObject so they can be properly passed into asset reference filters

diff --git a/Engine/Source/Editor/GraphEditor/Private/KismetPins/SGraphPinObject.cpp b/Engine/Source/Editor/GraphEditor/Private/KismetPins/SGraphPinObject.cpp
--- a/Engine/Source/Editor/GraphEditor/Private/KismetPins/SGraphPinObject.cpp
+++ b/Engine/Source/Editor/GraphEditor/Private/KismetPins/SGraphPinObject.cpp
@@ -245,6 +245,16 @@ TSharedRef<SWidget> SGraphPinObject::GenerateAssetPicker()
 	AssetPickerConfig.InitialAssetViewType = EAssetViewType::List;
 	AssetPickerConfig.bAllowDragging = false;
 
+	UObject* OwningAsset = GraphPinObj->GetOwningNode();
+	while (OwningAsset && !OwningAsset->IsAsset())
+	{
+		OwningAsset = OwningAsset->GetOuter();
+	}	
+	if (OwningAsset)
+	{
+		AssetPickerConfig.AdditionalReferencingAssets.Add(FAssetData(OwningAsset));
+	}
+
 	// Check with the node to see if there is any "AllowClasses" or "DisallowedClasses" metadata for the pin
 	FString AllowedClassesFilterString = GraphPinObj->GetOwningNode()->GetPinMetaData(GraphPinObj->PinName, FName(TEXT("AllowedClasses")));
 	if( !AllowedClassesFilterString.IsEmpty() )
