Selector drop down doesn't highlight current sound for quick playback

diff --git a/Engine/Source/Editor/GraphEditor/Private/KismetPins/SGraphPinObject.cpp b/Engine/Source/Editor/GraphEditor/Private/KismetPins/SGraphPinObject.cpp
--- a/Engine/Source/Editor/GraphEditor/Private/KismetPins/SGraphPinObject.cpp
+++ b/Engine/Source/Editor/GraphEditor/Private/KismetPins/SGraphPinObject.cpp
@@ -244,6 +244,7 @@ TSharedRef<SWidget> SGraphPinObject::GenerateAssetPicker()
 	AssetPickerConfig.OnAssetEnterPressed = FOnAssetEnterPressed::CreateSP(this, &SGraphPinObject::OnAssetEnterPressedInPicker);
 	AssetPickerConfig.InitialAssetViewType = EAssetViewType::List;
 	AssetPickerConfig.bAllowDragging = false;
+	AssetPickerConfig.InitialAssetSelection = CachedAssetData.GetAsset();
 
 	UObject* OwningAsset = GraphPinObj->GetOwningNode();
 	while (OwningAsset && !OwningAsset->IsAsset())
