Static analysis fix

diff --git a/Engine/Source/Editor/GraphEditor/Private/GraphActionNode.cpp b/Engine/Source/Editor/GraphEditor/Private/GraphActionNode.cpp
--- a/Engine/Source/Editor/GraphEditor/Private/GraphActionNode.cpp
+++ b/Engine/Source/Editor/GraphEditor/Private/GraphActionNode.cpp
@@ -540,15 +540,8 @@ TSharedPtr<FGraphActionNode> FGraphActionNode::NewCategoryNode(FString const& Ca
 //------------------------------------------------------------------------------
 TSharedPtr<FGraphActionNode> FGraphActionNode::NewActionNode(const TSharedPtr<FEdGraphSchemaAction>& Action)
 {
-	int32 Grouping  = FGraphActionNodeImpl::DEFAULT_GROUPING;
-	int32 SectionID = INVALID_SECTION_ID;
-
-	Grouping = FMath::Max(Grouping, Action->GetGrouping());
-	if (SectionID == INVALID_SECTION_ID)
-	{
-		// take the first non-zero section ID
-		SectionID = Action->GetSectionID();
-	}
+	const int32 Grouping = FMath::Max(FGraphActionNodeImpl::DEFAULT_GROUPING, Action->GetGrouping());
+	const int32 SectionID = Action->GetSectionID();
 
 	FGraphActionNode* ActionNode = new FGraphActionNode(Action, Grouping, SectionID);
 	TSharedPtr<FEdGraphSchemaAction> PrimeAction = ActionNode->GetPrimaryAction();
