ClangEditor CIS fix - we need to declare our (deleted) copy and move facilities as allowed to use our deprecated member variable
#rnx

diff --git a/Engine/Source/Editor/GraphEditor/Public/GraphActionNode.h b/Engine/Source/Editor/GraphEditor/Public/GraphActionNode.h
--- a/Engine/Source/Editor/GraphEditor/Public/GraphActionNode.h
+++ b/Engine/Source/Editor/GraphEditor/Public/GraphActionNode.h
@@ -19,6 +19,17 @@ template <typename ItemType> class STreeView;
 struct GRAPHEDITOR_API FGraphActionNode : TSharedFromThis<FGraphActionNode>
 {
 public:
+	// We need to declare our copy constructors so that we can disable
+	// deprecation warnings around them for ClangEditor - when all of
+	// the deprecated members are deleted we can remove these:
+	PRAGMA_DISABLE_DEPRECATION_WARNINGS
+	FGraphActionNode(const FGraphActionNode& Node) = delete;
+	FGraphActionNode& operator=(const FGraphActionNode& Node) = delete;
+	FGraphActionNode(FGraphActionNode&& Node) = delete;
+	FGraphActionNode& operator=(FGraphActionNode&& Node) = delete;
+	~FGraphActionNode() = default;
+	PRAGMA_ENABLE_DEPRECATION_WARNINGS
+
 	/** */
 	static const int32 INVALID_SECTION_ID = 0;
 
