Stop leaking graph nodes (until panel refreshes)

diff --git a/Engine/Source/Editor/GraphEditor/Private/SGraphPanel.cpp b/Engine/Source/Editor/GraphEditor/Private/SGraphPanel.cpp
--- a/Engine/Source/Editor/GraphEditor/Private/SGraphPanel.cpp
+++ b/Engine/Source/Editor/GraphEditor/Private/SGraphPanel.cpp
@@ -2178,6 +2178,8 @@ void SGraphPanel::RemoveNode(const UEdGraphNode* Node)
 			break;
 		}
 	}
+
+	NodeToWidgetLookup.Remove(Node);
 }
 
 TSharedPtr<SGraphNode> SGraphPanel::GetNodeWidgetFromGuid(FGuid Guid) const
