Fix editor ensure for graph action leaf node count mismatch caused by group divider nodes

diff --git a/Engine/Source/Editor/GraphEditor/Private/GraphActionNode.cpp b/Engine/Source/Editor/GraphEditor/Private/GraphActionNode.cpp
--- a/Engine/Source/Editor/GraphEditor/Private/GraphActionNode.cpp
+++ b/Engine/Source/Editor/GraphEditor/Private/GraphActionNode.cpp
@@ -367,7 +367,7 @@ void FGraphActionNode::GetLeafNodes(TArray< TSharedPtr<FGraphActionNode> >& OutL
 		{
 			ChildNode->GetLeafNodes(OutLeafArray);
 		}
-		else
+		else if (!ChildNode->IsGroupDividerNode())
 		{
 			// @TODO: sometimes, certain action nodes can have children as well
 			//        (for sub-graphs in the "MyBlueprint" tab)
