[GraphEditor] fix for contextual menu trying to access node without valid associated action

diff --git a/Engine/Source/Editor/GraphEditor/Private/SGraphActionMenu.cpp b/Engine/Source/Editor/GraphEditor/Private/SGraphActionMenu.cpp
--- a/Engine/Source/Editor/GraphEditor/Private/SGraphActionMenu.cpp
+++ b/Engine/Source/Editor/GraphEditor/Private/SGraphActionMenu.cpp
@@ -504,12 +504,12 @@ void SGraphActionMenu::GetSelectedActions(TArray< TSharedPtr<FEdGraphSchemaActio
 {
 	OutSelectedActions.Empty();
 
-	TArray< TSharedPtr<FGraphActionNode> > SelectedNodes = TreeView->GetSelectedItems();
-	if(SelectedNodes.Num() > 0)
+	TArray<TSharedPtr<FGraphActionNode>> SelectedNodes = TreeView->GetSelectedItems();
+	for (TSharedPtr<FGraphActionNode>& SelectedNode : SelectedNodes)
 	{
-		for ( int32 NodeIndex = 0; NodeIndex < SelectedNodes.Num(); NodeIndex++ )
+		if (SelectedNode.IsValid() && SelectedNode->IsActionNode())
 		{
-			OutSelectedActions.Add( SelectedNodes[NodeIndex]->Action );
+			OutSelectedActions.Add(SelectedNode->Action);
 		}
 	}
 }
