UE-201629 - Don't close metasound editor when converting a metasound from a preset

diff --git a/Engine/Plugins/Runtime/Metasound/Source/MetasoundEditor/Private/MetasoundEditor.cpp b/Engine/Plugins/Runtime/Metasound/Source/MetasoundEditor/Private/MetasoundEditor.cpp
--- a/Engine/Plugins/Runtime/Metasound/Source/MetasoundEditor/Private/MetasoundEditor.cpp
+++ b/Engine/Plugins/Runtime/Metasound/Source/MetasoundEditor/Private/MetasoundEditor.cpp
@@ -1998,11 +1998,11 @@ namespace Metasound
 				Builder->ConvertFromPreset(Result);
 				ensure(Result == EMetaSoundBuilderResult::Succeeded);
 
-				// Hack until toolbar is polished up & corner text properly dynamically updates
-				if (UAssetEditorSubsystem* AssetEditorSubsystem = GEditor->GetEditorSubsystem<UAssetEditorSubsystem>())
+				if (UToolMenus* ToolMenus = UToolMenus::Get())
 				{
-					AssetEditorSubsystem->CloseAllEditorsForAsset(Metasound);
+					ToolMenus->RefreshAllWidgets();
 				}
+				RefreshDetails();
 			}
 		}
 
diff --git a/Engine/Source/Editor/GraphEditor/Private/SGraphEditorImpl.cpp b/Engine/Source/Editor/GraphEditor/Private/SGraphEditorImpl.cpp
--- a/Engine/Source/Editor/GraphEditor/Private/SGraphEditorImpl.cpp
+++ b/Engine/Source/Editor/GraphEditor/Private/SGraphEditorImpl.cpp
@@ -676,6 +676,12 @@ void SGraphEditorImpl::Construct( const FArguments& InArgs )
 	
 	struct Local
 	{
+		static FText GetCornerText(TAttribute<FGraphAppearanceInfo> Appearance, FText DefaultText)
+		{
+			FText OverrideText = Appearance.Get().CornerText;
+			return !OverrideText.IsEmpty() ? OverrideText : DefaultText;
+		}
+
 		static FText GetPIENotifyText(TAttribute<FGraphAppearanceInfo> Appearance, FText DefaultText)
 		{
 			FText OverrideText = Appearance.Get().PIENotifyText;
@@ -694,6 +700,11 @@ void SGraphEditorImpl::Construct( const FArguments& InArgs )
 			return !OverrideText.IsEmpty() ? OverrideText : DefaultText;
 		}
 	};
+
+	FText DefaultCornerText;
+	TAttribute<FText> CornerText = Appearance.IsBound() ?
+		TAttribute<FText>::Create(TAttribute<FText>::FGetter::CreateStatic(&Local::GetCornerText, Appearance, DefaultCornerText)) :
+		TAttribute<FText>(DefaultCornerText);
 	
 	FText DefaultPIENotify(LOCTEXT("GraphSimulatingText", "SIMULATING"));
 	TAttribute<FText> PIENotifyText = Appearance.IsBound() ?
@@ -805,7 +816,7 @@ void SGraphEditorImpl::Construct( const FArguments& InArgs )
 			SNew(STextBlock)
 			.Visibility( EVisibility::HitTestInvisible )
 			.TextStyle( FAppStyle::Get(), "Graph.CornerText" )
-			.Text(Appearance.Get().CornerText)
+			.Text(CornerText)
 		]
 
 		// Top-right corner text indicating PIE is active
