Fix the "Class" pin 's class picker dropdown being a weird size and not using the class display name.

#pf skip
664ffa5b982a15c36f83f5e5

diff --git a/Engine/Source/Editor/GraphEditor/Private/KismetPins/SGraphPinClass.cpp b/Engine/Source/Editor/GraphEditor/Private/KismetPins/SGraphPinClass.cpp
--- a/Engine/Source/Editor/GraphEditor/Private/KismetPins/SGraphPinClass.cpp
+++ b/Engine/Source/Editor/GraphEditor/Private/KismetPins/SGraphPinClass.cpp
@@ -131,6 +131,7 @@ TSharedRef<SWidget> SGraphPinClass::GenerateAssetPicker()
 	FClassViewerInitializationOptions Options;
 	Options.Mode = EClassViewerMode::ClassPicker;
 	Options.bShowNoneOption = true;
+	Options.NameTypeToDisplay = EClassViewerNameTypeToDisplay::DisplayName;
 
 	// Get the min. spec for the classes allowed
 	const UClass* PinRequiredParentClass = Cast<const UClass>(GraphPinObj->PinType.PinSubCategoryObject.Get());
@@ -168,10 +169,8 @@ TSharedRef<SWidget> SGraphPinClass::GenerateAssetPicker()
 
 	return
 		SNew(SBox)
-		.WidthOverride(280)
 		[
 			SNew(SVerticalBox)
-
 			+ SVerticalBox::Slot()
 			.FillHeight(1.0f)
 			.MaxHeight(500)
