UE-215855 - Submix, Sound Class, and Sound Cue graphs are missing the graph name from the bottom right-corner of the graph

diff --git a/Engine/Source/Editor/GraphEditor/Private/SGraphEditorImpl.cpp b/Engine/Source/Editor/GraphEditor/Private/SGraphEditorImpl.cpp
--- a/Engine/Source/Editor/GraphEditor/Private/SGraphEditorImpl.cpp
+++ b/Engine/Source/Editor/GraphEditor/Private/SGraphEditorImpl.cpp
@@ -701,7 +701,7 @@ void SGraphEditorImpl::Construct( const FArguments& InArgs )
 		}
 	};
 
-	FText DefaultCornerText;
+	FText DefaultCornerText = Appearance.Get().CornerText;
 	TAttribute<FText> CornerText = Appearance.IsBound() ?
 		TAttribute<FText>::Create(TAttribute<FText>::FGetter::CreateStatic(&Local::GetCornerText, Appearance, DefaultCornerText)) :
 		TAttribute<FText>(DefaultCornerText);
