Fix links are not straight that come from blueprint event nodes

diff --git a/Engine/Source/Editor/GraphEditor/Private/KismetNodes/SGraphNodeK2Event.cpp b/Engine/Source/Editor/GraphEditor/Private/KismetNodes/SGraphNodeK2Event.cpp
--- a/Engine/Source/Editor/GraphEditor/Private/KismetNodes/SGraphNodeK2Event.cpp
+++ b/Engine/Source/Editor/GraphEditor/Private/KismetNodes/SGraphNodeK2Event.cpp
@@ -31,7 +31,7 @@ void SGraphNodeK2Event::AddPin( const TSharedRef<SGraphPin>& PinToAdd )
 		TitleAreaWidget->AddSlot()
 			.HAlign(HAlign_Right)
 			.VAlign(VAlign_Center)
-			.Padding(FMargin(4))
+			.Padding(4.0f, 0.0f)
 		[
 			PinToAdd
 		];
