Added support for UPARAM AllowedClasses and ShowDisplayNames metadata.

diff --git a/Engine/Source/Editor/BlueprintGraph/Classes/EdGraphSchema_K2.h b/Engine/Source/Editor/BlueprintGraph/Classes/EdGraphSchema_K2.h
--- a/Engine/Source/Editor/BlueprintGraph/Classes/EdGraphSchema_K2.h
+++ b/Engine/Source/Editor/BlueprintGraph/Classes/EdGraphSchema_K2.h
@@ -275,8 +275,13 @@ public:
 	/** Stub function used internally by animation blueprints */
 	static const FName MD_AnimBlueprintFunction;
 
+	/** Metadata that should be used with UPARAM to specify whether a TSubclassOf argument should use the DisplayName of classes */
+	static const FName MD_ShowDisplayNames;
+	
 	/** Metadata that should be used with UPARAM to specify whether a TSubclassOf argument allows abstract classes */
 	static const FName MD_AllowAbstractClasses;
+	/** Metadata that should be used with UPARAM to specify which classes are allowed in selection */
+	static const FName MD_AllowedClasses;
 
 	/** Metadata that should be used with UPARAM to specify a function name that generates a list of available values */
 	static const FName MD_GetOptions;
diff --git a/Engine/Source/Editor/BlueprintGraph/Private/EdGraphSchema_K2.cpp b/Engine/Source/Editor/BlueprintGraph/Private/EdGraphSchema_K2.cpp
--- a/Engine/Source/Editor/BlueprintGraph/Private/EdGraphSchema_K2.cpp
+++ b/Engine/Source/Editor/BlueprintGraph/Private/EdGraphSchema_K2.cpp
@@ -258,7 +258,10 @@ const FName FBlueprintMetadata::MD_UseEnumValuesAsMaskValuesInEditor(TEXT("UseEn
 
 const FName FBlueprintMetadata::MD_AnimBlueprintFunction(TEXT("AnimBlueprintFunction"));
 
+const FName FBlueprintMetadata::MD_ShowDisplayNames(TEXT("ShowDisplayNames"));
+
 const FName FBlueprintMetadata::MD_AllowAbstractClasses(TEXT("AllowAbstract"));
+const FName FBlueprintMetadata::MD_AllowedClasses(TEXT("AllowedClasses"));
 const FName FBlueprintMetadata::MD_GetOptions(TEXT("GetOptions"));
 
 const FName FBlueprintMetadata::MD_Namespace(TEXT("Namespace"));
diff --git a/Engine/Source/Editor/GraphEditor/Private/KismetPins/SGraphPinClass.cpp b/Engine/Source/Editor/GraphEditor/Private/KismetPins/SGraphPinClass.cpp
--- a/Engine/Source/Editor/GraphEditor/Private/KismetPins/SGraphPinClass.cpp
+++ b/Engine/Source/Editor/GraphEditor/Private/KismetPins/SGraphPinClass.cpp
@@ -26,6 +26,7 @@
 #include "Misc/CString.h"
 #include "Misc/PackageName.h"
 #include "Modules/ModuleManager.h"
+#include "PropertyCustomizationHelpers.h"
 #include "SGraphPin.h"
 #include "ScopedTransaction.h"
 #include "SlotBase.h"
@@ -141,21 +142,38 @@ TSharedRef<SWidget> SGraphPinClass::GenerateAssetPicker()
 		PinRequiredParentClass = UObject::StaticClass();
 	}
 
+	//Looks like this defaults to ClassName? Either way, allow UPARAM to specify this. 
+	const FString ShowDisplayNamesString = GraphPinObj->GetOwningNode()->GetPinMetaData(GraphPinObj->PinName, FBlueprintMetadata::MD_ShowDisplayNames);
+	if (!ShowDisplayNamesString.IsEmpty() && ShowDisplayNamesString.ToBool() == true)
+	{
+		Options.NameTypeToDisplay = EClassViewerNameTypeToDisplay::DisplayName;
+	}
+
 	TSharedPtr<FGraphPinFilter> Filter = MakeShareable(new FGraphPinFilter);
 	Filter->bAllowAbstractClasses = bAllowAbstractClasses;
 
+	FString AllowedClassesString = GraphPinObj->GetOwningNode()->GetPinMetaData(GraphPinObj->PinName, FBlueprintMetadata::MD_AllowedClasses);
+	if (!AllowedClassesString.IsEmpty())
+	{
+		Filter->AllowedChildrenOfClasses.Append(PropertyCustomizationHelpers::GetClassesFromMetadataString(AllowedClassesString));
+	}
+	
 	// Check with the node to see if there is any "AllowAbstract" metadata for the pin
 	FString AllowAbstractString = GraphPinObj->GetOwningNode()->GetPinMetaData(GraphPinObj->PinName, FBlueprintMetadata::MD_AllowAbstractClasses);
-
+	
 	// Override bAllowAbstractClasses is the AllowAbstract metadata was set
 	if (!AllowAbstractString.IsEmpty())
 	{
 		Filter->bAllowAbstractClasses = AllowAbstractString.ToBool();
 	}
-
+	
 	Options.ClassFilters.Add(Filter.ToSharedRef());
 
-	Filter->AllowedChildrenOfClasses.Add(PinRequiredParentClass);
+	if (Filter->AllowedChildrenOfClasses.Num() == 0)
+	{
+		Filter->AllowedChildrenOfClasses.Add(PinRequiredParentClass);
+	}
+	
 	Filter->GraphPinOutermostPackage = GraphPinObj->GetOuter()->GetOutermost();
 
 	if (UEdGraphNode* ParentNode = GraphPinObj->GetOwningNode())
