Fix 'Context Menu: Show Favorites Section' and 'Flatten Favorites Menus' which were broken by recent blueprint context menu optimizations
#rnx

diff --git a/Engine/Source/Editor/GraphEditor/Private/GraphActionNode.cpp b/Engine/Source/Editor/GraphEditor/Private/GraphActionNode.cpp
--- a/Engine/Source/Editor/GraphEditor/Private/GraphActionNode.cpp
+++ b/Engine/Source/Editor/GraphEditor/Private/GraphActionNode.cpp
@@ -239,7 +239,7 @@ TSharedPtr<FGraphActionNode> FGraphActionNode::AddChildAlphabetical(const TShare
 	}
 
 	// if a divider hasn't been created for the grouping, create one:
-	AddChildGrouping(ActionNode, this->AsShared());
+	AddChildGrouping(ActionNode, this->AsShared(), true);
 
 	// find or add categories iteratively, inserting as needed:
 	FGraphActionNode* OwningCategory = this;
@@ -273,7 +273,7 @@ TSharedPtr<FGraphActionNode> FGraphActionNode::AddChildAlphabetical(FGraphAction
 	}
 
 	// if a divider hasn't been created for the grouping, create one:
-	AddChildGrouping(ActionNode, this->AsShared());
+	AddChildGrouping(ActionNode, this->AsShared(), true);
 
 	// find or add categories iteratively, inserting as needed:
 	FGraphActionNode* OwningCategory = this;
@@ -697,7 +697,7 @@ void FGraphActionNode::InsertChild(TSharedPtr<FGraphActionNode> NodeToAdd)
 	// hardcode the order), but if this isn't in a section...
 	else
 	{
-		AddChildGrouping(NodeToAdd, NodeToAdd->ParentNode);
+		AddChildGrouping(NodeToAdd, NodeToAdd->ParentNode, false);
 	}
 
 	NodeToAdd->InsertOrder = Children.Num();
@@ -709,7 +709,7 @@ void FGraphActionNode::InsertChild(TSharedPtr<FGraphActionNode> NodeToAdd)
 }
 
 //------------------------------------------------------------------------------
-void FGraphActionNode::AddChildGrouping(TSharedPtr<FGraphActionNode> ActionNode, TWeakPtr<FGraphActionNode> Parent)
+void FGraphActionNode::AddChildGrouping(TSharedPtr<FGraphActionNode> ActionNode, TWeakPtr<FGraphActionNode> Parent, bool bInsertAlphabetically)
 {
 	if (ChildGroupings.Find(ActionNode->Grouping))
 	{
@@ -729,8 +729,15 @@ void FGraphActionNode::AddChildGrouping(TSharedPtr<FGraphActionNode> ActionNode,
 		// divider associated with it)
 		int32 DividerGrouping = FMath::Max(LowestGrouping, ActionNode->Grouping);
 
-		ChildGroupings.Add(ActionNode->Grouping); // to avoid recursion, add before we insert
-		InsertChild(NewGroupDividerNode(this->AsShared(), DividerGrouping));
+		ChildGroupings.Add(ActionNode->Grouping);  // to avoid recursion, add before we insert
+		if(bInsertAlphabetically)
+		{
+			InsertChildAlphabetical(NewGroupDividerNode(this->AsShared(), DividerGrouping));
+		}
+		else
+		{
+			InsertChild(NewGroupDividerNode(this->AsShared(), DividerGrouping));
+		}
 	}
 	else
 	{
diff --git a/Engine/Source/Editor/GraphEditor/Public/GraphActionNode.h b/Engine/Source/Editor/GraphEditor/Public/GraphActionNode.h
--- a/Engine/Source/Editor/GraphEditor/Public/GraphActionNode.h
+++ b/Engine/Source/Editor/GraphEditor/Public/GraphActionNode.h
@@ -305,7 +305,7 @@ private:
 	 */
 	void InsertChild(TSharedPtr<FGraphActionNode> NodeToAdd);
 
-	void AddChildGrouping(TSharedPtr<FGraphActionNode> ActionNode, TWeakPtr<FGraphActionNode> Parent);
+	void AddChildGrouping(TSharedPtr<FGraphActionNode> ActionNode, TWeakPtr<FGraphActionNode> Parent, bool bInsertAlphabetically);
 	void InsertChildAlphabetical(TSharedPtr<FGraphActionNode> NodeToAdd);
 
 	/** Recursive implementation helper for GetLinearizedIndex */
diff --git a/Engine/Source/Runtime/Engine/Private/EdGraph/EdGraphSchema.cpp b/Engine/Source/Runtime/Engine/Private/EdGraph/EdGraphSchema.cpp
--- a/Engine/Source/Runtime/Engine/Private/EdGraph/EdGraphSchema.cpp
+++ b/Engine/Source/Runtime/Engine/Private/EdGraph/EdGraphSchema.cpp
@@ -24,6 +24,9 @@ void FEdGraphSchemaAction::CosmeticUpdateCategory(FText NewCategory)
 	Category.BuildSourceString().ParseIntoArray(FullSearchCategoryArray, TEXT(" "), true);
 
 	// Glob search text together, we use the SearchText string for basic filtering:
+#if WITH_EDITORONLY_DATA
+	CategoryChain.Reset();
+#endif // WITH_EDITORONLY_DATA
 	UpdateSearchText();
 }
 
