[GraphActionMenu] fixed reversed CanRenameNode behavior and its use. Also test if the delegate is bound before executing it.

diff --git a/Engine/Source/Editor/GraphEditor/Private/SGraphActionMenu.cpp b/Engine/Source/Editor/GraphEditor/Private/SGraphActionMenu.cpp
--- a/Engine/Source/Editor/GraphEditor/Private/SGraphActionMenu.cpp
+++ b/Engine/Source/Editor/GraphEditor/Private/SGraphActionMenu.cpp
@@ -819,7 +819,11 @@ bool SGraphActionMenu::ShouldExpandNodes() const
 
 bool SGraphActionMenu::CanRenameNode(TWeakPtr<FGraphActionNode> InNode) const
 {
-	return !OnCanRenameSelectedAction.Execute(InNode);
+	if (OnCanRenameSelectedAction.IsBound())
+	{
+		return OnCanRenameSelectedAction.Execute(InNode);
+	}
+	return false;
 }
 
 void SGraphActionMenu::OnFilterTextChanged( const FText& InFilterText )
@@ -953,7 +957,15 @@ TSharedRef<ITableRow> SGraphActionMenu::MakeWidget( TSharedPtr<FGraphActionNode>
 		}
 		else
 		{
-			ReadOnlyArgument.IsReadOnly(this, &SGraphActionMenu::CanRenameNode, WeakItem);
+			ReadOnlyArgument.IsReadOnly_Lambda([WeakThis = this->AsWeak(), WeakItem]
+			{
+				const TSharedPtr AsSharedWidget = WeakThis.Pin();
+				if (const SGraphActionMenu* Menu = static_cast<SGraphActionMenu*>(AsSharedWidget.Get()))
+				{
+					return !Menu->CanRenameNode(WeakItem);
+				}
+				return true;
+			});
 		}
 
 		TSharedRef<SGraphActionCategoryWidget> CategoryWidget =
