GraphEditor: Guard against remove + add of a node happening outside of the widget's tick

#rnx

diff --git a/Engine/Source/Editor/GraphEditor/Private/SGraphPanel.cpp b/Engine/Source/Editor/GraphEditor/Private/SGraphPanel.cpp
--- a/Engine/Source/Editor/GraphEditor/Private/SGraphPanel.cpp
+++ b/Engine/Source/Editor/GraphEditor/Private/SGraphPanel.cpp
@@ -135,6 +135,9 @@ void SGraphPanel::Construct( const SGraphPanel::FArguments& InArgs )
 	PreviousFrameSavedMousePosForSplineOverlap = FVector2D::ZeroVector;
 
 	TimeLeftToInvalidatePerTick = 0.0f;
+	
+	bHasCustomPrepass = true;
+	bCheckNodeGraphObjValidity = false;
 }
 
 SGraphPanel::~SGraphPanel()
@@ -751,6 +754,16 @@ void SGraphPanel::OnArrangeChildren( const FGeometry& AllottedGeometry, FArrange
 	ArrangedChildren.Append(MyArrangedChildren);
 }
 
+bool SGraphPanel::CustomPrepass(float LayoutScaleMultiplier)
+{
+	if(bCheckNodeGraphObjValidity)
+	{
+		bCheckNodeGraphObjValidity = false;
+		RemoveAllNodesWithInvalidPointers();
+	}
+	return true; // still run prepass on everything
+}
+
 void SGraphPanel::UpdateSelectedNodesPositions(FVector2D PositionIncrement)
 {
 	FScopedTransaction Transaction(NSLOCTEXT("GraphEditor", "NudgeNodeAction", "Nudge Node"));
@@ -2182,6 +2195,36 @@ void SGraphPanel::RemoveNode(const UEdGraphNode* Node)
 	NodeToWidgetLookup.Remove(Node);
 }
 
+void SGraphPanel::RemoveAllNodesWithInvalidPointers()
+{
+	if (GraphObj == nullptr)
+	{
+		return;
+	}
+	
+	TArray<TSharedRef<SGraphNode>> NodesWithInvalidPointers;
+
+	for (int32 Iter = 0; Iter != Children.Num(); ++Iter)
+	{
+		TSharedRef<SGraphNode> Child = GetChild(Iter);
+		if (const UEdGraphNode* NodeObj = Child->GetNodeObj())
+		{
+			if (!GraphObj->Nodes.Contains(NodeObj))
+			{
+				NodesWithInvalidPointers.Add(Child);
+			}
+		}
+	}
+
+	for (const TSharedRef<SGraphNode>& NodeWithInvalidPoint : NodesWithInvalidPointers)
+	{
+		if (const UEdGraphNode* NodeObj = NodeWithInvalidPoint->GetNodeObj())
+		{
+			RemoveNode(NodeObj);
+		}
+	}
+}
+
 TSharedPtr<SGraphNode> SGraphPanel::GetNodeWidgetFromGuid(FGuid Guid) const
 {
 	return NodeGuidMap.FindRef(Guid).Pin();
@@ -2457,6 +2500,7 @@ void SGraphPanel::OnGraphChanged(const FEdGraphEditAction& EditAction)
 			}
 
 			RegisterActiveTimer(0.f, FWidgetActiveTimerDelegate::CreateLambda(RemoveNodesDelegateWrapper, StaticCastWeakPtr<SGraphPanel>(AsWeak()), NodePtrSet));
+			bCheckNodeGraphObjValidity = true;
 		}
 		if (bWasAddAction)
 		{
diff --git a/Engine/Source/Editor/GraphEditor/Public/SGraphPanel.h b/Engine/Source/Editor/GraphEditor/Public/SGraphPanel.h
--- a/Engine/Source/Editor/GraphEditor/Public/SGraphPanel.h
+++ b/Engine/Source/Editor/GraphEditor/Public/SGraphPanel.h
@@ -138,6 +138,7 @@ public:
 	virtual FReply OnKeyDown( const FGeometry& MyGeometry, const FKeyEvent& InKeyEvent ) override;
 	virtual bool SupportsKeyboardFocus() const override;
 	virtual void OnArrangeChildren( const FGeometry& AllottedGeometry, FArrangedChildren& ArrangedChildren ) const override;
+	virtual bool CustomPrepass(float LayoutScaleMultiplier) override;
 	// End of SWidget interface
 
 	// SNodePanel interface
@@ -243,6 +244,9 @@ protected:
 
 	/** Helper method to remove a node from the panel */
 	void RemoveNode(const UEdGraphNode* Node);
+
+	/** Helper method to remove all nodes from the panel holding garbage / invalid pointers */
+	void RemoveAllNodesWithInvalidPointers();
 public:
 	/** Pin marked via shift-clicking */
 	TWeakPtr<SGraphPin> MarkedPin;
@@ -399,4 +403,7 @@ private:
 
 	/** Weak pointer to the last summoned context menu, for dismissing it when requested. */
 	TWeakPtr<IMenu> ContextMenu;
+
+	/** A flag indicating that we need to check if we have any nodes which use invalid pointers */
+	bool bCheckNodeGraphObjValidity;
 };
