Ensure that the Blueprint "Watch" bubbles always display on top of the node and the blueprint execution arrow.

diff --git a/Engine/Source/Editor/GraphEditor/Private/SNodePanel.cpp b/Engine/Source/Editor/GraphEditor/Private/SNodePanel.cpp
--- a/Engine/Source/Editor/GraphEditor/Private/SNodePanel.cpp
+++ b/Engine/Source/Editor/GraphEditor/Private/SNodePanel.cpp
@@ -1511,10 +1511,16 @@ void SNodePanel::PaintComment(const FString& CommentText, const FGeometry& Allot
 	const FVector2D CommentBubbleOffset = FVector2D(0, -(CommentTextSize.Y + CommentCalloutArrow->ImageSize.Y) - PositionBias);
 	const FVector2D CommentBubbleArrowOffset = FVector2D( CommentCalloutArrow->ImageSize.X, -CommentCalloutArrow->ImageSize.Y - PositionBias);
 
+	// We want the watch window comment bubbles to draw on top of the node so that they are not obscured
+	// and on top of the big blueprint debugging arrow
+	static constexpr int32 CommentBubbleLayer = 100;
+
+	const int32 DrawingLayer = DrawLayerId + CommentBubbleLayer;
+
 	// Draw a comment bubble
 	FSlateDrawElement::MakeBox(
 		OutDrawElements,
-		DrawLayerId-1,
+		DrawingLayer - 1,	// Put the background 1 layer behind, so the text is on top
 		AllottedGeometry.ToPaintGeometry(CommentTextSize, FSlateLayoutTransform(CommentBubbleOffset)),
 		CommentCalloutBubble,
 		ESlateDrawEffect::None,
@@ -1523,7 +1529,7 @@ void SNodePanel::PaintComment(const FString& CommentText, const FGeometry& Allot
 
 	FSlateDrawElement::MakeBox(
 		OutDrawElements,
-		DrawLayerId-1,
+		DrawingLayer - 1,	// Put the background 1 layer behind, so the text is on top
 		AllottedGeometry.ToPaintGeometry( CommentCalloutArrow->ImageSize, FSlateLayoutTransform(CommentBubbleArrowOffset) ),
 		CommentCalloutArrow,
 		ESlateDrawEffect::None,
@@ -1533,7 +1539,7 @@ void SNodePanel::PaintComment(const FString& CommentText, const FGeometry& Allot
 	// Draw the comment text itself
 	FSlateDrawElement::MakeText(
 		OutDrawElements,
-		DrawLayerId,
+		DrawingLayer,
 		AllottedGeometry.ToPaintGeometry( CommentTextSize, FSlateLayoutTransform(CommentBubbleOffset + CommentBubblePadding) ),
 		CommentText,
 		CommentFont,
