Reroute nodes can now display compiler notes, warnings, and errors

[RN] Reroute nodes can now display compiler notes, warnings, and errors

diff --git a/Engine/Source/Editor/GraphEditor/Private/SGraphNodeKnot.cpp b/Engine/Source/Editor/GraphEditor/Private/SGraphNodeKnot.cpp
--- a/Engine/Source/Editor/GraphEditor/Private/SGraphNodeKnot.cpp
+++ b/Engine/Source/Editor/GraphEditor/Private/SGraphNodeKnot.cpp
@@ -350,6 +350,8 @@ void SGraphNodeKnot::UpdateGraphNode()
 
 	this->ContentScale.Bind( this, &SGraphNode::GetContentScale );
 
+	SetupErrorReporting();
+
 	this->GetOrAddSlot( ENodeZone::Center )
 		.HAlign(HAlign_Center)
 		.VAlign(VAlign_Center)
@@ -390,6 +392,13 @@ void SGraphNodeKnot::UpdateGraphNode()
 						]
 					]
 				]
+				+SVerticalBox::Slot()
+				.VAlign(VAlign_Bottom)
+				.HAlign(HAlign_Center)
+				.AutoHeight()
+				[
+					ErrorReporting->AsWidget()
+				]
 			]
 		];
 	// Create comment bubble
