Allow comment nodes to be placed in graphs that are inside external packages

The outer-chain traversal ended up crashing as the asset's package is not its direct outer when externally packaged.

diff --git a/Engine/Source/Editor/GraphEditor/Private/SGraphNodeComment.cpp b/Engine/Source/Editor/GraphEditor/Private/SGraphNodeComment.cpp
--- a/Engine/Source/Editor/GraphEditor/Private/SGraphNodeComment.cpp
+++ b/Engine/Source/Editor/GraphEditor/Private/SGraphNodeComment.cpp
@@ -158,18 +158,6 @@ void SGraphNodeComment::UpdateGraphNode()
 	UEdGraphNode_Comment* CommentNode = CastChecked<UEdGraphNode_Comment>(GraphNode);
 	bCachedBubbleVisibility = CommentNode->bCommentBubbleVisible_InDetailsPanel;
 
-	// Setup a tag for this node
-	FString TagName;
-
-	// We want the name of the blueprint as our name - we can find the node from the GUID
-	UObject* Package = GraphNode->GetOutermost();
-	UObject* LastOuter = GraphNode->GetOuter();
-	while (LastOuter->GetOuter() != Package)
-	{
-		LastOuter = LastOuter->GetOuter();
-	}
-	TagName = FString::Printf(TEXT("GraphNode,%s,%s"), *LastOuter->GetFullName(), *GraphNode->NodeGuid.ToString());
-
 	SetupErrorReporting();
 
 	// Setup a meta tag for this node
@@ -547,12 +535,7 @@ void SGraphNodeComment::PopulateMetaTag(FGraphNodeMetaData* TagMeta) const
 	if (GraphNode != nullptr)
 	{
 		// We want the name of the blueprint as our name - we can find the node from the GUID
-		UObject* Package = GraphNode->GetOutermost();
-		UObject* LastOuter = GraphNode->GetOuter();
-		while (LastOuter->GetOuter() != Package)
-		{
-			LastOuter = LastOuter->GetOuter();
-		}
+		UObject* LastOuter = GraphNode->GetOutermostObject();
 		TagMeta->Tag = FName(*FString::Printf(TEXT("GraphNode_%s_%s"), *LastOuter->GetFullName(), *GraphNode->NodeGuid.ToString()));
 		TagMeta->OuterName = LastOuter->GetFullName();
 		TagMeta->GUID = GraphNode->NodeGuid;
