Improvements for UScriptType function params
- Added support for MetaStruct meta data
- Added support to show DisplayNames
- Struct marked as Hidden won't be displayed

#rnx

diff --git a/Engine/Source/Editor/GraphEditor/Private/KismetPins/SGraphPinStruct.cpp b/Engine/Source/Editor/GraphEditor/Private/KismetPins/SGraphPinStruct.cpp
--- a/Engine/Source/Editor/GraphEditor/Private/KismetPins/SGraphPinStruct.cpp
+++ b/Engine/Source/Editor/GraphEditor/Private/KismetPins/SGraphPinStruct.cpp
@@ -74,6 +74,13 @@ public:
 			return MetaStruct == nullptr;
 		}
 
+		// Don't show any hidden structs
+		static const FName NAME_HiddenMetaTag = "Hidden";
+		if (InStruct->HasMetaData(NAME_HiddenMetaTag))
+		{
+			return false;
+		}
+
 		// Query the native struct to see if it has the correct parent type (if any)
 		return !MetaStruct || InStruct->IsChildOf(MetaStruct);
 	}
@@ -92,10 +99,20 @@ TSharedRef<SWidget> SGraphPinStruct::GenerateAssetPicker()
 	// Fill in options
 	FStructViewerInitializationOptions Options;
 	Options.Mode = EStructViewerMode::StructPicker;
+	Options.NameTypeToDisplay = EStructViewerNameTypeToDisplay::DisplayName;
 	Options.bShowNoneOption = true;
 
-	// TODO: We would need our own PC_ type to be able to get the meta-struct here
 	const UScriptStruct* MetaStruct = nullptr;
+	static const FName NAME_MetaStructMetaTag = "MetaStruct";
+	const FString MetaStructName = GraphPinObj->GetOwningNode()->GetPinMetaData(GraphPinObj->PinName, NAME_MetaStructMetaTag);
+	if (!MetaStructName.IsEmpty())
+	{
+		MetaStruct = UClass::TryFindTypeSlow<UScriptStruct>(MetaStructName, EFindFirstObjectOptions::EnsureIfAmbiguous);
+		if (!MetaStruct)
+		{
+			MetaStruct = LoadObject<UScriptStruct>(nullptr, *MetaStructName);
+		}
+	}
 
 	TSharedRef<FGraphPinStructFilter> StructFilter = MakeShared<FGraphPinStructFilter>();
 	Options.StructFilter = StructFilter;
