Honor action even when our SelectedNode is out of date

[RN] Fixed bug that would cause clicking an entry in the blueprint context menu to spawn the incorrect node

diff --git a/Engine/Source/Editor/GraphEditor/Private/SGraphActionMenu.cpp b/Engine/Source/Editor/GraphEditor/Private/SGraphActionMenu.cpp
--- a/Engine/Source/Editor/GraphEditor/Private/SGraphActionMenu.cpp
+++ b/Engine/Source/Editor/GraphEditor/Private/SGraphActionMenu.cpp
@@ -1144,23 +1144,8 @@ bool SGraphActionMenu::OnMouseButtonDownEvent( TWeakPtr<FEdGraphSchemaAction> In
 	bool bResult = false;
 	if( (!bIgnoreUIUpdate) && InAction.IsValid() )
 	{
-		TArray< TSharedPtr<FGraphActionNode> > SelectionList = TreeView->GetSelectedItems();
-		TSharedPtr<FGraphActionNode> SelectedNode;
-		if (SelectionList.Num() == 1)
-		{	
-			SelectedNode = SelectionList[0];
-		}
-		else if (GetTotalLeafNodes() == 1)
-		{
-			SelectedNode = GetFirstAction();
-		}
-		if (SelectedNode.IsValid() && SelectedNode->HasValidAction())
-		{
-			if( SelectedNode->GetPrimaryAction().Get() == InAction.Pin().Get() )
-			{				
-				bResult = HandleSelection( SelectedNode, ESelectInfo::OnMouseClick );
-			}
-		}
+		OnActionSelected.Execute({InAction.Pin()}, ESelectInfo::OnMouseClick);
+		bResult = true;
 	}
 	return bResult;
 }
