Continue to avoid running logic for SGraphActionMenus with unbound OnActionSelected delegates

#rnx

diff --git a/Engine/Source/Editor/GraphEditor/Private/SGraphActionMenu.cpp b/Engine/Source/Editor/GraphEditor/Private/SGraphActionMenu.cpp
--- a/Engine/Source/Editor/GraphEditor/Private/SGraphActionMenu.cpp
+++ b/Engine/Source/Editor/GraphEditor/Private/SGraphActionMenu.cpp
@@ -1142,7 +1142,7 @@ FReply SGraphActionMenu::OnItemDragDetected( const FGeometry& MyGeometry, const
 bool SGraphActionMenu::OnMouseButtonDownEvent( TWeakPtr<FEdGraphSchemaAction> InAction )
 {
 	bool bResult = false;
-	if( (!bIgnoreUIUpdate) && InAction.IsValid() )
+	if( (!bIgnoreUIUpdate) && InAction.IsValid() && OnActionSelected.IsBound())
 	{
 		OnActionSelected.Execute({InAction.Pin()}, ESelectInfo::OnMouseClick);
 		bResult = true;
