UE-199276 - Able to add a Custom Event to a function graph using Create Event node

When adding new custom events via the Create Event node, we need to ensure that we add them to the Blueprint's event graph, and not the node's current graph.

Additionally, we should only allow this option if the Blueprint actually supports event graphs (eg: Blueprint libraries do _not_ have these).

diff --git a/Engine/Source/Editor/GraphEditor/Private/KismetNodes/SGraphNodeK2CreateDelegate.cpp b/Engine/Source/Editor/GraphEditor/Private/KismetNodes/SGraphNodeK2CreateDelegate.cpp
--- a/Engine/Source/Editor/GraphEditor/Private/KismetNodes/SGraphNodeK2CreateDelegate.cpp
+++ b/Engine/Source/Editor/GraphEditor/Private/KismetNodes/SGraphNodeK2CreateDelegate.cpp
@@ -39,6 +39,15 @@
 
 class SWidget;
 
+namespace UE::SGraphNodeK2::Private
+{
+	bool MatchEventGraphPredicate(const TObjectPtr<UEdGraph>& Entry)
+	{
+		check(Entry);
+		return Entry->GetFName() == UEdGraphSchema_K2::GN_EventGraph;
+	};
+}
+
 FText SGraphNodeK2CreateDelegate::FunctionDescription(const UFunction* Function, const bool bOnlyDescribeSignature /*= false*/, const int32 CharacterLimit /*= 32*/)
 {
 	if (!Function || !Function->GetOuter())
@@ -171,7 +180,7 @@ FText SGraphNodeK2CreateDelegate::GetCurrentFunctionDescription() const
 
 void SGraphNodeK2CreateDelegate::OnFunctionSelected(TSharedPtr<FString> FunctionItemData, ESelectInfo::Type SelectInfo)
 {
-	const FScopedTransaction Transaction(NSLOCTEXT("GraphNodeK2Create", "CreateMatchingSigniture", "Create matching signiture"));
+	const FScopedTransaction Transaction(NSLOCTEXT("GraphNodeK2Create", "CreateMatchingSignature", "Create matching signature"));
 
 	if (FunctionItemData.IsValid())
 	{
@@ -203,29 +212,26 @@ void SGraphNodeK2CreateDelegate::OnFunctionSelected(TSharedPtr<FString> Function
 			}
 			else if (FunctionItemData == CreateMatchingEventData)
 			{
-				// Get a valid name for the function graph
-				FName NewEventName = FBlueprintEditorUtils::FindUniqueCustomEventName(NodeBP);
+				TObjectPtr<UEdGraph>* EventGraphPtr = NodeBP->UbergraphPages.FindByPredicate(UE::SGraphNodeK2::Private::MatchEventGraphPredicate);
+				check(EventGraphPtr);
+				UEdGraph* EventGraph = *EventGraphPtr;
+				check(EventGraph);
 
+				FName NewEventName = FBlueprintEditorUtils::FindUniqueCustomEventName(NodeBP);
 				UBlueprintEventNodeSpawner* Spawner = UBlueprintEventNodeSpawner::Create(UK2Node_CustomEvent::StaticClass(), NewEventName);
-				
-				// Get a good spawn location for the new event
-				UEdGraph* NodeGraph = Node->GetGraph();
-				check(NodeGraph);
-				FVector2D SpawnPos = NodeGraph->GetGoodPlaceForNewNode();
+				FVector2D SpawnPos = EventGraph->GetGoodPlaceForNewNode();
 
-				UEdGraphNode* NewNode = Spawner->Invoke(NodeGraph, IBlueprintNodeBinder::FBindingSet(), SpawnPos);
+				UEdGraphNode* NewNode = Spawner->Invoke(EventGraph, IBlueprintNodeBinder::FBindingSet(), SpawnPos);
+
+				UK2Node_CustomEvent* NewEventNode = CastChecked<UK2Node_CustomEvent>(NewNode);
+				NewEventNode->SetDelegateSignature(Node->GetDelegateSignature());
+				// Reconstruct to get the new parameters to show in the editor
+				NewEventNode->ReconstructNode();
+				NewEventNode->bIsEditable = true;
+				FKismetEditorUtilities::BringKismetToFocusAttentionOnObject(NewEventNode);
 
-				if (UK2Node_CustomEvent* NewEventNode = Cast<UK2Node_CustomEvent>(NewNode))
-				{
-					NewEventNode->SetDelegateSignature(Node->GetDelegateSignature());
-					// Reconstruct to get the new parameters to show in the editor
-					NewEventNode->ReconstructNode();
-					NewEventNode->bIsEditable = true;
-					FKismetEditorUtilities::BringKismetToFocusAttentionOnObject(NewEventNode);
-				}
-				
 				Node->SetFunction(NewEventName);
-				FBlueprintEditorUtils::MarkBlueprintAsStructurallyModified(NodeBP);				
+				FBlueprintEditorUtils::MarkBlueprintAsStructurallyModified(NodeBP);
 			}
 			else
 			{
@@ -255,6 +261,9 @@ void SGraphNodeK2CreateDelegate::CreateBelowPinControls(TSharedPtr<SVerticalBox>
 {
 	if (UK2Node_CreateDelegate* Node = Cast<UK2Node_CreateDelegate>(GraphNode))
 	{
+		const UBlueprint* NodeBP = Node->GetBlueprint();
+		check(NodeBP);
+
 		UFunction* FunctionSignature = Node->GetDelegateSignature();
 		UClass* ScopeClass = Node->GetScopeClass();
 
@@ -293,7 +302,13 @@ void SGraphNodeK2CreateDelegate::CreateBelowPinControls(TSharedPtr<SVerticalBox>
 			CreateMatchingFunctionData = AddDefaultFunctionDataOption(NSLOCTEXT("GraphNodeK2Create", "CreateMatchingFunctionOption", "[Create a matching function]"));
 
 			// Only signatures with no output parameters can be events
-			if (!UEdGraphSchema_K2::HasFunctionAnyOutputParameter(FunctionSignature))
+			// There must also be an event graph to add the event to
+			const bool bCanCreateEvent =
+				!UEdGraphSchema_K2::HasFunctionAnyOutputParameter(FunctionSignature) &&
+				NodeBP->UbergraphPages.ContainsByPredicate(UE::SGraphNodeK2::Private::MatchEventGraphPredicate)
+			;
+
+			if (bCanCreateEvent)
 			{
 				CreateMatchingEventData = AddDefaultFunctionDataOption(NSLOCTEXT("GraphNodeK2Create", "CreateMatchingEventOption", "[Create a matching event]"));
 			}
