Substrate - account for new eye BSDF feature on SubstrateEye node as ewll as on Substrate details tab.

#rnx
#tests PC dx12

diff --git a/Engine/Source/Editor/GraphEditor/Private/MaterialNodes/SGraphSubstrateMaterial.cpp b/Engine/Source/Editor/GraphEditor/Private/MaterialNodes/SGraphSubstrateMaterial.cpp
--- a/Engine/Source/Editor/GraphEditor/Private/MaterialNodes/SGraphSubstrateMaterial.cpp
+++ b/Engine/Source/Editor/GraphEditor/Private/MaterialNodes/SGraphSubstrateMaterial.cpp
@@ -165,9 +165,30 @@ static const TSharedRef<SWidget> InternalProcessOperator(
 		case SUBSTRATE_OPERATOR_BSDF_LEGACY:	// legacy BSDF should have been converted to BSDF already.
 		case SUBSTRATE_OPERATOR_BSDF:
 		{
+			auto GetEyeDesc = [&Op]()
+			{
+				FString EyeDesc = TEXT("Eye");
+				const bool bIrisN = Op.Has(ESubstrateBsdfFeature_EyeIrisNormalPluggedIn);
+				const bool bIrisT = Op.Has(ESubstrateBsdfFeature_EyeIrisTangentPluggedIn);
+				if (bIrisN || bIrisT)
+				{
+					EyeDesc += TEXT("Iris");
+				}
+				if (bIrisN)
+				{
+					EyeDesc += TEXT("N");
+				}
+				if (bIrisT)
+				{
+					EyeDesc += TEXT("T");
+				}
+				EyeDesc += TEXT(" ");
+				return EyeDesc;
+			};
+
 			FString BSDFDesc = OutputType == ESubstrateWidgetOutputType::Node ? 
 											 FString(TEXT("BSDF")) : 
-											 FString::Printf(TEXT("BSDF (%s%s%s%s%s%s%s%s)")
+											 FString::Printf(TEXT("BSDF (%s%s%s%s%s%s%s%s%s%s)")
 											, Op.Has(ESubstrateBsdfFeature_EdgeColor) ? TEXT("F90 ") : TEXT("")
 											, Op.Has(ESubstrateBsdfFeature_SSS) ? TEXT("SSS ") : TEXT("")
 											, Op.Has(ESubstrateBsdfFeature_MFPPluggedIn) ? TEXT("MFP ") : TEXT("")
@@ -176,7 +197,7 @@ static const TSharedRef<SWidget> InternalProcessOperator(
 											, Op.Has(ESubstrateBsdfFeature_Fuzz) ? TEXT("Fuz ") : TEXT("")
 											, Op.Has(ESubstrateBsdfFeature_Glint) ? TEXT("Gli ") : TEXT("")
 											, Op.Has(ESubstrateBsdfFeature_SpecularProfile) ? TEXT("Spc ") : TEXT("")
-											, Op.Has(ESubstrateBsdfFeature_Eye) ? TEXT("Eye ") : TEXT("")
+											, Op.Has(ESubstrateBsdfFeature_Eye) ? *GetEyeDesc() : TEXT("")
 											, Op.Has(ESubstrateBsdfFeature_Hair) ? TEXT("Hai ") : TEXT("")
 											);
 
diff --git a/Engine/Source/Runtime/Engine/Private/Materials/MaterialExpressionSubstrate.cpp b/Engine/Source/Runtime/Engine/Private/Materials/MaterialExpressionSubstrate.cpp
--- a/Engine/Source/Runtime/Engine/Private/Materials/MaterialExpressionSubstrate.cpp
+++ b/Engine/Source/Runtime/Engine/Private/Materials/MaterialExpressionSubstrate.cpp
@@ -2202,6 +2202,15 @@ FSubstrateOperator* UMaterialExpressionSubstrateEyeBSDF::SubstrateGenerateMateri
 	FSubstrateOperator& SubstrateOperator = Compiler->SubstrateCompilationRegisterOperator(SUBSTRATE_OPERATOR_BSDF, Compiler->SubstrateTreeStackGetPathUniqueId(), this, Parent, Compiler->SubstrateTreeStackGetParentPathUniqueId());
 	SubstrateOperator.BSDFType = SUBSTRATE_BSDF_TYPE_EYE;
 	SubstrateOperator.BSDFFeatures = ESubstrateBsdfFeature_Eye;
+	if (IrisNormal.IsConnected())
+	{
+		SubstrateOperator.BSDFFeatures |= ESubstrateBsdfFeature_EyeIrisNormalPluggedIn;
+	}
+	if (IrisPlaneNormal.IsConnected())
+	{
+		SubstrateOperator.BSDFFeatures |= ESubstrateBsdfFeature_EyeIrisTangentPluggedIn;
+	}
+
 	SubstrateOperator.ThicknessIndex = Compiler->SubstrateThicknessStackGetThicknessIndex();
 	SubstrateOperator.bBSDFWritesEmissive = EmissiveColor.IsConnected();
 	return &SubstrateOperator;
