Substrate - fix for material root node FrontMaterial pin Substrate node suggestions.

#tests PC editor
#rnx

diff --git a/Engine/Source/Editor/GraphEditor/Private/MaterialNodes/SGraphSubstrateMaterial.cpp b/Engine/Source/Editor/GraphEditor/Private/MaterialNodes/SGraphSubstrateMaterial.cpp
--- a/Engine/Source/Editor/GraphEditor/Private/MaterialNodes/SGraphSubstrateMaterial.cpp
+++ b/Engine/Source/Editor/GraphEditor/Private/MaterialNodes/SGraphSubstrateMaterial.cpp
@@ -32,11 +32,9 @@ bool FSubstrateWidget::HasInputSubstrateType(const UEdGraphPin* InPin)
 
 	if (UMaterialGraphNode_Root* RootPinNode = Cast<UMaterialGraphNode_Root>(InPin->GetOwningNode()))
 	{
-		EMaterialProperty propertyId = (EMaterialProperty)FCString::Atoi(*InPin->PinType.PinSubCategory.ToString());
-		switch (propertyId)
+		if(RootPinNode->GetPinMaterialType(InPin) == MCT_Substrate)
 		{
-			case MP_FrontMaterial:
-				return true;
+			return true;
 		}
 	}
 	if (UMaterialGraphNode* PinNode = Cast<UMaterialGraphNode>(InPin->GetOwningNode()))
