Allow GetOptions to work for nodes other than 'call function'

No longer relies on the call function node to provide a class context, as the underlying PropertyEditorUtils::GetPropertyOptions implementation no longer requires it (such as when using a full path to a BP function library call)

diff --git a/Engine/Source/Editor/GraphEditor/Private/KismetPins/SGraphPinString.cpp b/Engine/Source/Editor/GraphEditor/Private/KismetPins/SGraphPinString.cpp
--- a/Engine/Source/Editor/GraphEditor/Private/KismetPins/SGraphPinString.cpp
+++ b/Engine/Source/Editor/GraphEditor/Private/KismetPins/SGraphPinString.cpp
@@ -142,37 +142,43 @@ TSharedPtr<SWidget> SGraphPinString::TryBuildComboBoxWidget()
 		return nullptr;
 	}
 
-	const UK2Node_CallFunction* CallFunctionNode = Cast<UK2Node_CallFunction>(GraphPinObj->GetOwningNode());
-	if (!CallFunctionNode)
-	{
-		return nullptr;
-	}
+	ComboBoxOptions.Empty();
 
-	const UClass* FunctionParentClass = CallFunctionNode->FunctionReference.GetMemberParentClass();
-	if (!FunctionParentClass)
-	{
-		// If we don't have a parent class, then we have to be in self context
-		if (!CallFunctionNode->FunctionReference.IsSelfContext())
-		{
-			return nullptr;
-		}
+	TArray<UObject*> PropertyContainers;
 
-		// If we're in self context, let's get the Blueprint we're residing in
-		const UBlueprint* Blueprint = GraphPinObj->GetOwningNode()->GetTypedOuter<UBlueprint>();
-		if (!Blueprint || !Blueprint->GeneratedClass)
+	const UK2Node_CallFunction* CallFunctionNode = Cast<UK2Node_CallFunction>(GraphPinObj->GetOwningNode());
+	if (CallFunctionNode)
+	{
+		const UClass* FunctionParentClass = CallFunctionNode->FunctionReference.GetMemberParentClass();
+		if (!FunctionParentClass)
 		{
-			return nullptr;
+			// If we don't have a parent class, then we have to be in self context
+			if (!CallFunctionNode->FunctionReference.IsSelfContext())
+			{
+				return nullptr;
+			}
+
+			// If we're in self context, let's get the Blueprint we're residing in
+			const UBlueprint* Blueprint = GraphPinObj->GetOwningNode()->GetTypedOuter<UBlueprint>();
+			if (!Blueprint || !Blueprint->GeneratedClass)
+			{
+				return nullptr;
+			}
+
+			// Work with the blueprint class the node is in
+			FunctionParentClass = Blueprint->GeneratedClass;
 		}
 
-		// Work with the blueprint class the node is in
-		FunctionParentClass = Blueprint->GeneratedClass;
+		PropertyContainers.Add(FunctionParentClass->GetDefaultObject());
 	}
 
-	ComboBoxOptions.Empty();
-
-	TArray<UObject*> PropertyContainers = { FunctionParentClass->GetDefaultObject() };
 	PropertyEditorUtils::GetPropertyOptions(PropertyContainers, GetOptionsFunctionName, ComboBoxOptions);
 
+	if(ComboBoxOptions.Num() == 0)
+	{
+		return nullptr;
+	}
+
 	const FString DefaultString = GraphPinObj->GetDefaultAsString();
 	const TSharedPtr<FString>* DefaultValue = ComboBoxOptions.FindByPredicate([&DefaultString](
 		const TSharedPtr<FString>& Value)
