Allow dragging to disconnect wires

[RN] Ctrl clicking on wires with multiple connections now behaves as expected, disconnecting from the closest pin rather than the pin with the fewest connections

diff --git a/Engine/Source/Editor/GraphEditor/Private/ConnectionDrawingPolicy.cpp b/Engine/Source/Editor/GraphEditor/Private/ConnectionDrawingPolicy.cpp
--- a/Engine/Source/Editor/GraphEditor/Private/ConnectionDrawingPolicy.cpp
+++ b/Engine/Source/Editor/GraphEditor/Private/ConnectionDrawingPolicy.cpp
@@ -567,23 +567,8 @@ void FGraphSplineOverlapResult::ComputeBestPin()
 	}
 	else
 	{
-		// Both are valid, now see if one of the pins has one connection while the other has more than one
-		const int32 LinksTo1 = Pin1->LinkedTo.Num();
-		const int32 LinksTo2 = Pin2->LinkedTo.Num();
-
-		if ((LinksTo1 > 1) && (LinksTo2 == 1))
-		{
-			BestPin = Pin2;
-		}
-		else if ((LinksTo1 == 1) && (LinksTo2 > 1))
-		{
-			BestPin = Pin1;
-		}
-		else
-		{
-			// Both pins have multiple links, or both pins have one link, choose based on distance to the pins
-			BestPin = (DistanceSquaredToPin1 < DistanceSquaredToPin2) ? Pin1 : Pin2;
-		}
+		// choose based on distance to the pins
+		BestPin = (DistanceSquaredToPin1 < DistanceSquaredToPin2) ? Pin1 : Pin2;
 	}
 
 	BestPinHandle = FGraphPinHandle(BestPin);
diff --git a/Engine/Source/Editor/GraphEditor/Private/SGraphPanel.cpp b/Engine/Source/Editor/GraphEditor/Private/SGraphPanel.cpp
--- a/Engine/Source/Editor/GraphEditor/Private/SGraphPanel.cpp
+++ b/Engine/Source/Editor/GraphEditor/Private/SGraphPanel.cpp
@@ -960,18 +960,39 @@ FReply SGraphPanel::OnKeyDown( const FGeometry& MyGeometry, const FKeyEvent& InK
 
 FReply SGraphPanel::OnMouseButtonDown(const FGeometry& MyGeometry, const FPointerEvent& MouseEvent)
 {
-	if ((MouseEvent.GetEffectingButton() == EKeys::LeftMouseButton) && (MouseEvent.IsAltDown() || MouseEvent.IsControlDown()))
+	const bool bIsControlDown = MouseEvent.IsControlDown();
+	const bool bIsAltDown = MouseEvent.IsAltDown();
+	if ((MouseEvent.GetEffectingButton() == EKeys::LeftMouseButton) && (bIsAltDown || bIsControlDown))
 	{
 		// Intercept alt-left clicking on the hovered spline for targeted break link
 		UEdGraphPin* Pin1;
 		UEdGraphPin* Pin2;
-		if (MouseEvent.IsAltDown() && PreviousFrameSplineOverlap.GetPins(*this, Pin1, Pin2))
+		if (bIsAltDown && PreviousFrameSplineOverlap.GetPins(*this, Pin1, Pin2))
 		{
 			const UEdGraphSchema* Schema = GraphObj->GetSchema();
 			Schema->BreakSinglePinLink(Pin1, Pin2);
 		}
 		else if (SGraphPin* BestPinFromHoveredSpline = GetBestPinFromHoveredSpline())
 		{
+			if(bIsControlDown)
+			{
+				// try to disconnect just the wire (spline) that is hovered):
+				TSharedPtr<SGraphPin> BestPin = PreviousFrameSplineOverlap.GetPin1Handle().FindInGraphPanel(*this);
+				TSharedPtr<SGraphPin> OtherPin = PreviousFrameSplineOverlap.GetPin2Handle().FindInGraphPanel(*this);
+				if(BestPin && OtherPin)
+				{
+					if (BestPin.Get() != BestPinFromHoveredSpline)
+					{
+						std::swap(BestPin, OtherPin);
+					}
+					
+					return OnMouseDownDisconnectClosestPinFromWire(
+						BestPin.ToSharedRef(), 
+						OtherPin.ToSharedRef());
+				}
+			}
+
+			// fall back legacy behavior, just allow the spline to handle mouse down:
 			return BestPinFromHoveredSpline->OnPinMouseDown(MyGeometry, MouseEvent);
 		}
 	}
@@ -1609,6 +1630,17 @@ void SGraphPanel::MoveNodesToAnchorPinAtGraphPosition(TArrayView<UEdGraphNode* c
 	}
 }
 
+FReply SGraphPanel::OnMouseDownDisconnectClosestPinFromWire(TSharedRef<SGraphPin> ClosestPin, TSharedRef<SGraphPin> ConnectedPin)
+{
+	TArray<TSharedRef<SGraphPin>> PinArray;
+	PinArray.Add(ConnectedPin);
+	TSharedPtr<FDragDropOperation> DragEvent = ClosestPin->SpawnPinDragEvent(SharedThis(this), PinArray);
+	const UEdGraphSchema* Schema = GraphObj->GetSchema();
+	Schema->BreakSinglePinLink(ClosestPin->GetPinObj(), ConnectedPin->GetPinObj());
+	ClosestPin->bIsMovingLinks = true;
+	return FReply::Handled().BeginDragDrop(DragEvent.ToSharedRef());
+}
+
 void SGraphPanel::OnBeginMakingConnection(UEdGraphPin* InOriginatingPin)
 {
 	OnBeginMakingConnection(FGraphPinHandle(InOriginatingPin));
diff --git a/Engine/Source/Editor/GraphEditor/Public/SGraphPanel.h b/Engine/Source/Editor/GraphEditor/Public/SGraphPanel.h
--- a/Engine/Source/Editor/GraphEditor/Public/SGraphPanel.h
+++ b/Engine/Source/Editor/GraphEditor/Public/SGraphPanel.h
@@ -392,6 +392,9 @@ private:
 	/** Will move a group of nodes by the amount needed for an anchor pin to be at a certain position */
 	void MoveNodesToAnchorPinAtGraphPosition(TArrayView<UEdGraphNode* const> NodesToMove, FGraphPinHandle PinToAnchor, FVector2D DesiredPinGraphPosition);
 
+	/** Begins a drag operation that disconnects the pin at ClosestPin, but only from the ConnectedPin, rather than all connections */
+	FReply OnMouseDownDisconnectClosestPinFromWire(TSharedRef<class SGraphPin> ClosestPin, TSharedRef<class SGraphPin> ConnectedPin);
+
 	/** Handle to timer callback that allows the UI to refresh it's arrangement each tick, allows animations to occur within the UI */
 	TWeakPtr<FActiveTimerHandle> ActiveTimerHandleInvalidatePerTick;
 
diff --git a/Engine/Source/Editor/GraphEditor/Public/SGraphPin.h b/Engine/Source/Editor/GraphEditor/Public/SGraphPin.h
--- a/Engine/Source/Editor/GraphEditor/Public/SGraphPin.h
+++ b/Engine/Source/Editor/GraphEditor/Public/SGraphPin.h
@@ -428,4 +428,5 @@ protected:
 
 private:
 	TSharedPtr<SLevelOfDetailBranchNode> PinNameLODBranchNode;
+	friend class SGraphPanel;
 };
