Improve styling of SGraphPinEnum values in the graph editor so that items have more padding and the length of the pop-up menu matches the length of the button that opens it.

diff --git a/Engine/Source/Editor/GraphEditor/Private/KismetPins/SGraphPinEnum.cpp b/Engine/Source/Editor/GraphEditor/Private/KismetPins/SGraphPinEnum.cpp
--- a/Engine/Source/Editor/GraphEditor/Private/KismetPins/SGraphPinEnum.cpp
+++ b/Engine/Source/Editor/GraphEditor/Private/KismetPins/SGraphPinEnum.cpp
@@ -67,17 +67,17 @@ void SPinComboBox::Construct( const FArguments& InArgs )
 		]
 		.MenuContent()
 		[
-			SNew(SVerticalBox)
-			+SVerticalBox::Slot()
-			.MaxHeight(450.0f)
+			SNew(SBorder)
+			.BorderImage(FAppStyle::Get().GetBrush("Brushes.Recessed"))
+			.Padding(0) // No padding here so that the inner box controls matching width with padding
 			[
-				SNew(SBorder)
-				.BorderImage(FAppStyle::GetBrush("Menu.Background"))
-				.Padding( 0 )
+				SNew(SBox)
+				.MaxDesiredHeight(450.0f)
+				.MinDesiredWidth(this, &SPinComboBox::GetMenuDesiredWidth)
+				.Padding(3.0f, 6.0f)
 				[
 					SAssignNew( ComboList, SComboList )
 					.ListItemsSource( &ComboItemList )
-//					.ItemHeight( 22 )
 					.OnGenerateRow( this, &SPinComboBox::OnGenerateComboWidget )
 					.OnSelectionChanged( this, &SPinComboBox::OnSelectionChangedInternal )
 				]
@@ -86,6 +86,15 @@ void SPinComboBox::Construct( const FArguments& InArgs )
 	];
 }
 
+FOptionalSize SPinComboBox::GetMenuDesiredWidth() const
+{
+	if (ComboButton)
+	{
+		return ComboButton->GetDesiredSize().X;
+	}
+	return FOptionalSize(150.0f);
+}
+
 void SPinComboBox::RemoveItemByIndex(int32 InIndexToRemove)
 {
 	for(int32 Index=0;Index<ComboItemList.Num();Index++)
@@ -121,6 +130,7 @@ TSharedRef<ITableRow> SPinComboBox::OnGenerateComboWidget( TSharedPtr<int32> InC
 		[
 			SNew(SBox)
 			.MinDesiredWidth(150.0f)
+			.Padding(8.0f, 3.0f)
 			[
 				SNew(STextBlock)
 				.Text( this, &SPinComboBox::GetRowString, RowIndex )
diff --git a/Engine/Source/Editor/GraphEditor/Public/SGraphPinComboBox.h b/Engine/Source/Editor/GraphEditor/Public/SGraphPinComboBox.h
--- a/Engine/Source/Editor/GraphEditor/Public/SGraphPinComboBox.h
+++ b/Engine/Source/Editor/GraphEditor/Public/SGraphPinComboBox.h
@@ -55,7 +55,10 @@ private:
 
 	// Function to create each row of the combo widget
 	TSharedRef<ITableRow> OnGenerateComboWidget( TSharedPtr<int32> InComboIndex, const TSharedRef<STableViewBase>& OwnerTable );
-
+	
+	/** Returns the target width of the enum selection menu. */
+	FOptionalSize GetMenuDesiredWidth() const;
+	
 	FText GetRowString(int32 RowIndex) const
 	{
 		return OnGetDisplayName.Execute(RowIndex);
