Update the size assumption of the SGraphNode title spacer to be correct with the current styling.

diff --git a/Engine/Source/Editor/GraphEditor/Private/SGraphNode.cpp b/Engine/Source/Editor/GraphEditor/Private/SGraphNode.cpp
--- a/Engine/Source/Editor/GraphEditor/Private/SGraphNode.cpp
+++ b/Engine/Source/Editor/GraphEditor/Private/SGraphNode.cpp
@@ -164,7 +164,7 @@ void SNodeTitle::RebuildWidget()
 	{
 		// Note: This code a little fragile, and will need to be updated if the font or padding of titles
 		// changes in the future, but the failure mode is just a slight misalignment.
-		const int32 EstimatedExtraHeight = (Lines.Num() - 1) * 13;
+		const int32 EstimatedExtraHeight = (Lines.Num() - 1) * 14;
 
 		const int32 SnapSize = (int32)SNodePanel::GetSnapGridSize();
 		const int32 PadSize = SnapSize - (EstimatedExtraHeight % SnapSize);
