Allow SDefaultGraphActionWidget to fall back to the incoming FCreateWidgetForActionData highlight if no explicit highlight text binding is defined.

diff --git a/Engine/Source/Editor/GraphEditor/Private/SGraphActionMenu.cpp b/Engine/Source/Editor/GraphEditor/Private/SGraphActionMenu.cpp
--- a/Engine/Source/Editor/GraphEditor/Private/SGraphActionMenu.cpp
+++ b/Engine/Source/Editor/GraphEditor/Private/SGraphActionMenu.cpp
@@ -204,7 +204,7 @@ void SDefaultGraphActionWidget::Construct(const FArguments& InArgs, const FCreat
 			SNew(STextBlock)
 			.Font(FCoreStyle::GetDefaultFontStyle("Regular", 9))
 			.Text(InCreateData->Action->GetMenuDescription())
-			.HighlightText(InArgs._HighlightText)
+			.HighlightText(InArgs._HighlightText.IsSet() ? InArgs._HighlightText : InCreateData->HighlightText)
 		]
 	];
 }
