Modify SGraphNodeResizeable so that the resizing hotspots grow along as the node panel zooms out. This makes resizing & moving derived nodes such as Blueprint Comments much easier.

diff --git a/Engine/Source/Editor/GraphEditor/Private/SGraphNodeResizable.cpp b/Engine/Source/Editor/GraphEditor/Private/SGraphNodeResizable.cpp
--- a/Engine/Source/Editor/GraphEditor/Private/SGraphNodeResizable.cpp
+++ b/Engine/Source/Editor/GraphEditor/Private/SGraphNodeResizable.cpp
@@ -14,6 +14,7 @@
 #include "Misc/Attribute.h"
 #include "SNodePanel.h"
 #include "ScopedTransaction.h"
+#include "SGraphPanel.h"
 #include "Widgets/SWindow.h"
 
 namespace GraphNodeResizableDefs
@@ -267,19 +268,45 @@ FVector2D SGraphNodeResizable::GetCorrectedNodePosition() const
 SGraphNodeResizable::EResizableWindowZone SGraphNodeResizable::FindMouseZone(const FVector2D& LocalMouseCoordinates) const
 {
 	EResizableWindowZone OutMouseZone = CRWZ_NotInWindow;
-	const FSlateRect HitResultBorderSize = GetHitTestingBorder();
+	FSlateRect HitResultBorderSize = GetHitTestingBorder();
 	const FVector2D NodeSize = GetDesiredSize();
+	const float TitleBarHeight = GetTitleBarHeight();
 
+	const float Zoom = GetOwnerPanel()->GetZoomAmount();
+
+	if (Zoom > 0.0f && Zoom < 1.0f)
+	{
+		// Scale up the hit zones when zoomed out so that they remain accessible.
+		const float Scale = 1.0f / Zoom;
+		HitResultBorderSize = FSlateRect(HitResultBorderSize.Left * Scale, HitResultBorderSize.Top * Scale, HitResultBorderSize.Right * Scale, HitResultBorderSize.Bottom * Scale);
+	}
+	
 	// Test for hit in location of 'grab' zone
 	if (LocalMouseCoordinates.Y > ( NodeSize.Y - HitResultBorderSize.Bottom))
 	{
 		OutMouseZone = CRWZ_BottomBorder;
 	}
+	else if (TitleBarHeight <= HitResultBorderSize.Top * 1.5f)
+	{
+		// Use the largest of the two hit areas.
+		const float LargestZone = FMath::Max(TitleBarHeight, HitResultBorderSize.Top);
+		
+		// Allow half of the title bar to be treated as the top edge
+		// so that the title bar can still be grabbed
+		if (LocalMouseCoordinates.Y < LargestZone * 0.4)
+		{
+			OutMouseZone = CRWZ_TopBorder;
+		}
+		else if (LocalMouseCoordinates.Y < LargestZone)
+		{
+			OutMouseZone = CRWZ_TitleBar;
+		}
+	}
 	else if (LocalMouseCoordinates.Y <= (HitResultBorderSize.Top))
 	{
 		OutMouseZone = CRWZ_TopBorder;
 	}
-	else if (LocalMouseCoordinates.Y <= GetTitleBarHeight())
+	else if (LocalMouseCoordinates.Y <= TitleBarHeight)
 	{
 		OutMouseZone = CRWZ_TitleBar;
 	}
@@ -314,7 +341,7 @@ SGraphNodeResizable::EResizableWindowZone SGraphNodeResizable::FindMouseZone(con
 			OutMouseZone = CRWZ_LeftBorder;
 		}
 	}
-	
+
 	// Test for hit on rest of frame
 	if (OutMouseZone == CRWZ_NotInWindow)
 	{
