* Fix pin color with SubstrateType in material function (input/output/call)

diff --git a/Engine/Source/Editor/GraphEditor/Private/MaterialNodes/SGraphSubstrateMaterial.cpp b/Engine/Source/Editor/GraphEditor/Private/MaterialNodes/SGraphSubstrateMaterial.cpp
--- a/Engine/Source/Editor/GraphEditor/Private/MaterialNodes/SGraphSubstrateMaterial.cpp
+++ b/Engine/Source/Editor/GraphEditor/Private/MaterialNodes/SGraphSubstrateMaterial.cpp
@@ -9,6 +9,10 @@
 #include "MaterialGraph/MaterialGraphSchema.h"
 #include "Materials/MaterialExpression.h"
 #include "Materials/MaterialExpressionSubstrate.h"
+#include "Materials/MaterialExpressionFunctionInput.h"
+#include "Materials/MaterialExpressionFunctionOutput.h"
+#include "Materials/MaterialExpressionMaterialFunctionCall.h"
+
 #include "SGraphPin.h"
 #include "Math/Color.h"
 #include "Styling/StyleColors.h"
@@ -309,6 +313,17 @@ void FSubstrateWidget::GetPinColor(TSharedPtr<SGraphPin>& Out, const UMaterialGr
 			ColorModifier = Color1;
 		}
 	}
+	else if (InNode->MaterialExpression->IsA(UMaterialExpressionMaterialFunctionCall::StaticClass()) ||
+			 InNode->MaterialExpression->IsA(UMaterialExpressionFunctionInput::StaticClass()) ||
+			 InNode->MaterialExpression->IsA(UMaterialExpressionFunctionOutput::StaticClass()))
+	{
+		const uint32 PinType = UMaterialGraphSchema::GetMaterialValueType(Pin);
+		if ((Out->GetDirection() == EGPD_Input || Out->GetDirection() == EGPD_Output) && (PinType & MCT_Substrate) != 0)
+		{
+			bHasColorModifier = true;
+			ColorModifier = FSubstrateWidget::GetConnectionColor();
+		}
+	}
 
 	if (InNode->MaterialExpression->IsA(UMaterialExpressionSubstrateBSDF::StaticClass()) && Out->GetDirection() == EGPD_Output && UMaterialGraphSchema::GetMaterialValueType(Pin) == MCT_Substrate)
 	{
