Made graph zoom limit break configurable

diff --git a/Engine/Source/Editor/GraphEditor/Private/SNodePanel.cpp b/Engine/Source/Editor/GraphEditor/Private/SNodePanel.cpp
--- a/Engine/Source/Editor/GraphEditor/Private/SNodePanel.cpp
+++ b/Engine/Source/Editor/GraphEditor/Private/SNodePanel.cpp
@@ -1755,10 +1755,13 @@ void SNodePanel::ChangeZoomLevel(int32 ZoomLevelDelta, const FVector2D& WidgetSp
 
 	const int32 DefaultZoomLevel = ZoomLevels->GetDefaultZoomLevel();
 	const int32 NumZoomLevels = ZoomLevels->GetNumZoomLevels();
+	EGraphZoomLimitHandling ZoomLimitHandling = ZoomLevels->GetZoomLimitHandling();
 
+	bool bAllowLimitBreak = (bOverrideZoomLimiting && ZoomLimitHandling == EGraphZoomLimitHandling::Default) || ZoomLimitHandling == EGraphZoomLimitHandling::AllowLimitBreak;
 	const bool bAllowFullZoomRange =
 		// To zoom in past 1:1 the user must press control
-		(ZoomLevel == DefaultZoomLevel && ZoomLevelDelta > 0 && bOverrideZoomLimiting) ||
+		(ZoomLevel == DefaultZoomLevel && ZoomLevelDelta > 0 && bAllowLimitBreak)
+			 ||
 		// If they are already zoomed in past 1:1, user may zoom freely
 		(ZoomLevel > DefaultZoomLevel);
 
diff --git a/Engine/Source/Editor/GraphEditor/Public/SNodePanel.h b/Engine/Source/Editor/GraphEditor/Public/SNodePanel.h
--- a/Engine/Source/Editor/GraphEditor/Public/SNodePanel.h
+++ b/Engine/Source/Editor/GraphEditor/Public/SNodePanel.h
@@ -46,6 +46,7 @@
 #include "Widgets/DeclarativeSyntaxSupport.h"
 #include "Widgets/SPanel.h"
 #include "Widgets/SWidget.h"
+#include "SNodePanel.generated.h"
 
 class FActiveTimerHandle;
 class FReferenceCollector;
@@ -87,6 +88,17 @@ namespace EGraphRenderingLOD
 	};
 }
 
+UENUM()
+enum class EGraphZoomLimitHandling : uint8
+{
+	/** By default, hold down Ctrl to break past the zoom limit. */
+	Default,
+	/** Always allow breaking through the zoom limit. */
+	AllowLimitBreak,
+	/** Never allow breaking through the zoom limit. */
+	DisallowLimitBreak
+};
+
 // Context passed in when getting popup info
 struct FNodeInfoContext
 {
@@ -212,6 +224,8 @@ struct FZoomLevelsContainer
 	 */
 	virtual EGraphRenderingLOD::Type	GetLOD(int32 InZoomLevel) const = 0;
 
+	virtual EGraphZoomLimitHandling GetZoomLimitHandling() const { return EGraphZoomLimitHandling::Default; }
+	
 	// Necessary for Mac OS X to compile 'delete <pointer_to_this_object>;'
 	virtual ~FZoomLevelsContainer( void ) {};
 };
