Add base class / widget that can be used to faciliate custom search support
- Add search support to RigVM deformer graphs as API example

#tests editor

diff --git a/Engine/Plugins/Animation/DeformerGraph/Source/OptimusEditor/Private/FindInDeformer.cpp b/Engine/Plugins/Animation/DeformerGraph/Source/OptimusEditor/Private/FindInDeformer.cpp
new file mode 100644
--- /dev/null
+++ b/Engine/Plugins/Animation/DeformerGraph/Source/OptimusEditor/Private/FindInDeformer.cpp
@@ -0,0 +1,59 @@
+// Copyright Epic Games, Inc. All Rights Reserved.
+
+#include "FindInDeformer.h"
+#include "OptimusEditor.h"
+#include "GraphEditor.h"
+
+#define LOCTEXT_NAMESPACE "FindInDeformer"
+
+//////////////////////////////////////////////////////////////////////////
+// FFindInDeformerResult
+
+FFindInDeformerResult::FFindInDeformerResult(const FFindInGraphResult::FCreateParams& InCreateParams)
+	: FFindInGraphResult(InCreateParams)
+{
+}
+
+void FFindInDeformerResult::JumpToNode(TWeakPtr<FAssetEditorToolkit> AssetEditorToolkit, const UEdGraphNode* InNode) const
+{
+	if (TSharedPtr<FOptimusEditor> OptimusEditor = StaticCastSharedPtr<FOptimusEditor>(AssetEditorToolkit.Pin()))
+	{
+		if (TSharedPtr<SGraphEditor> GraphEditor = OptimusEditor->GetGraphEditorWidget())
+		{
+			GraphEditor->JumpToNode(InNode, false);
+		}
+	}
+}
+
+
+//////////////////////////////////////////////////////////////////////////
+// SFindInDeformer
+
+TSharedPtr<FFindInGraphResult> SFindInDeformer::MakeSearchResult(const FFindInGraphResult::FCreateParams& InParams)
+{
+	return MakeShared<FFindInDeformerResult>(InParams);
+}
+
+const UEdGraph* SFindInDeformer::GetGraph()
+{
+	if (TSharedPtr<FOptimusEditor> OptimusEditor = StaticCastSharedPtr<FOptimusEditor>(AssetEditorToolkitPtr.Pin()))
+	{
+		if (TSharedPtr<SGraphEditor> GraphEditor = OptimusEditor->GetGraphEditorWidget())
+		{
+			return GraphEditor->GetCurrentGraph();
+		}
+	}
+
+	return nullptr;
+}
+
+bool SFindInDeformer::MatchTokensInNode(const UEdGraphNode* Node, const TArray<FString>& Tokens)
+{
+	// For now no type-specific node matching
+	return false;
+}
+
+
+/////////////////////////////////////////////////////
+
+#undef LOCTEXT_NAMESPACE
diff --git a/Engine/Plugins/Animation/DeformerGraph/Source/OptimusEditor/Private/OptimusEditor.cpp b/Engine/Plugins/Animation/DeformerGraph/Source/OptimusEditor/Private/OptimusEditor.cpp
--- a/Engine/Plugins/Animation/DeformerGraph/Source/OptimusEditor/Private/OptimusEditor.cpp
+++ b/Engine/Plugins/Animation/DeformerGraph/Source/OptimusEditor/Private/OptimusEditor.cpp
@@ -24,6 +24,7 @@
 #include "Animation/DebugSkelMeshComponent.h"
 #include "AnimationEditorPreviewActor.h"
 #include "Engine/StaticMeshActor.h"
+#include "FindInDeformer.h"
 #include "Framework/Commands/GenericCommands.h"
 #include "GraphEditor.h"
 #include "IAssetFamily.h"
@@ -806,6 +807,12 @@ void FOptimusEditor::OnCreateComment()
 	NewCommentNodeAction.PerformAction(EditorGraph, nullptr, CurrentMouseGraphPosition, true);
 }
 
+void FOptimusEditor::OnFindInDeformer()
+{
+	TabManager->TryInvokeTab(FOptimusEditorFindTabSummoner::TabId);
+	FindResults->FocusForUse();
+}
+
 void FOptimusEditor::SubscribeToGraphNotifies(UOptimusNodeGraph* Graph, const FOnGraphNotified& OnGraphNotifiedDelegate)
 {
 	if (IsValid(Graph))
@@ -1121,6 +1128,9 @@ void FOptimusEditor::CreateWidgets()
 	// -- Graph editor
 	GraphEditorWidget = CreateGraphEditorWidget();
 	GraphEditorWidget->SetViewLocation(FVector2D::ZeroVector, 1);
+
+	FindResults =
+		SNew(SFindInDeformer, SharedThis(this));
 }
 
 
@@ -1201,6 +1211,8 @@ TSharedRef<SGraphEditor> FOptimusEditor::CreateGraphEditorWidget()
 			FExecuteAction::CreateSP(this, &FOptimusEditor::OnCreateComment)
 		);
 
+		GraphEditorCommands->MapAction(FOptimusEditorGraphCommands::Get().FindInDeformer,
+			FExecuteAction::CreateSP(this, &FOptimusEditor::OnFindInDeformer));
 #if 0
 
 		// Alignment Commands
@@ -1258,8 +1270,6 @@ TSharedRef<SGraphEditor> FOptimusEditor::CreateGraphEditorWidget()
 		.OptimusEditor(SharedThis(this))
 		.OnGraphCrumbClickedEvent_Lambda([this](UOptimusNodeGraph* InNodeGraph) { SetEditGraph(InNodeGraph);});
 
-	
-
 	return SNew(SGraphEditor)
 		.AdditionalCommands(GraphEditorCommands)
 		.IsEditable(this, &FOptimusEditor::IsGraphEditable)
diff --git a/Engine/Plugins/Animation/DeformerGraph/Source/OptimusEditor/Private/OptimusEditor.h b/Engine/Plugins/Animation/DeformerGraph/Source/OptimusEditor/Private/OptimusEditor.h
--- a/Engine/Plugins/Animation/DeformerGraph/Source/OptimusEditor/Private/OptimusEditor.h
+++ b/Engine/Plugins/Animation/DeformerGraph/Source/OptimusEditor/Private/OptimusEditor.h
@@ -16,6 +16,7 @@ class IOptimusPathResolver;
 class IPersonaPreviewScene;
 class IPersonaToolkit;
 class IPersonaViewport;
+class SFindInDeformer;
 class SGraphEditor;
 class SOptimusEditorViewport;
 class SOptimusGraphTitleBar;
@@ -170,7 +171,9 @@ private:
 	void ConvertToSubGraph();
 	bool CanConvertToSubGraph() const;
 
-	void OnCreateComment();	
+	void OnCreateComment();
+
+	void OnFindInDeformer();
 	
 	// ----------------------------------------------------------------------------------------
 	// Graph event listeners
@@ -215,6 +218,9 @@ public:
 	// KILL ME
 	TSharedPtr<SGraphEditor> GetGraphEditorWidget() const { return GraphEditorWidget; }
 
+	// @TODO: Remove below if above gets removed
+	TSharedPtr<SFindInDeformer> GetFindResultsWidget() const { return FindResults; }
+
 	bool GetBoundsForSelectedNodes(FSlateRect& Rect);
 	
 private:
@@ -251,6 +257,7 @@ private:
 	
 	// -- Widgets
 	TSharedPtr<IPersonaViewport> ViewportWidget;
+	TSharedPtr<SFindInDeformer> FindResults;
 	TSharedPtr<SGraphEditor> GraphEditorWidget;
 	TSharedPtr<IDetailsView> PropertyDetailsWidget;
 	TSharedPtr<IDetailsView> PreviewDetailsWidget;
diff --git a/Engine/Plugins/Animation/DeformerGraph/Source/OptimusEditor/Private/OptimusEditorGraphCommands.cpp b/Engine/Plugins/Animation/DeformerGraph/Source/OptimusEditor/Private/OptimusEditorGraphCommands.cpp
--- a/Engine/Plugins/Animation/DeformerGraph/Source/OptimusEditor/Private/OptimusEditorGraphCommands.cpp
+++ b/Engine/Plugins/Animation/DeformerGraph/Source/OptimusEditor/Private/OptimusEditorGraphCommands.cpp
@@ -30,6 +30,7 @@ void FOptimusEditorGraphCommands::RegisterCommands()
 	
 	UI_COMMAND(ConvertToFunction, "Convert to a Function", "Convert the selected SubGraph to a reusable Function.", EUserInterfaceActionType::Button, FInputChord());
 	UI_COMMAND(ConvertToSubGraph, "Convert to a SubGraph", "Convert the selected Function to a non-reusable SubGraph.", EUserInterfaceActionType::Button, FInputChord());
+	UI_COMMAND(FindInDeformer, "FindInDeformer", "Search the current deformer graph.", EUserInterfaceActionType::Button, FInputChord(EModifierKey::Control, EKeys::F));
 }
 
 #undef LOCTEXT_NAMESPACE
diff --git a/Engine/Plugins/Animation/DeformerGraph/Source/OptimusEditor/Private/OptimusEditorGraphCommands.h b/Engine/Plugins/Animation/DeformerGraph/Source/OptimusEditor/Private/OptimusEditorGraphCommands.h
--- a/Engine/Plugins/Animation/DeformerGraph/Source/OptimusEditor/Private/OptimusEditorGraphCommands.h
+++ b/Engine/Plugins/Animation/DeformerGraph/Source/OptimusEditor/Private/OptimusEditorGraphCommands.h
@@ -22,4 +22,6 @@ public:
 	
 	TSharedPtr<FUICommandInfo> ConvertToFunction;
 	TSharedPtr<FUICommandInfo> ConvertToSubGraph;
+
+	TSharedPtr<FUICommandInfo> FindInDeformer;
 };
diff --git a/Engine/Plugins/Animation/DeformerGraph/Source/OptimusEditor/Private/OptimusEditorMode.cpp b/Engine/Plugins/Animation/DeformerGraph/Source/OptimusEditor/Private/OptimusEditorMode.cpp
--- a/Engine/Plugins/Animation/DeformerGraph/Source/OptimusEditor/Private/OptimusEditorMode.cpp
+++ b/Engine/Plugins/Animation/DeformerGraph/Source/OptimusEditor/Private/OptimusEditorMode.cpp
@@ -20,6 +20,7 @@ FOptimusEditorMode::FOptimusEditorMode(TSharedRef<FOptimusEditor> InEditorApp) :
 	TabFactories.RegisterFactory(MakeShared<FOptimusEditorNodePaletteTabSummoner>(InEditorApp));
 	TabFactories.RegisterFactory(MakeShared<FOptimusEditorExplorerTabSummoner>(InEditorApp));
 	TabFactories.RegisterFactory(MakeShared<FOptimusEditorGraphTabSummoner>(InEditorApp));
+	TabFactories.RegisterFactory(MakeShared<FOptimusEditorFindTabSummoner>(InEditorApp));
 	TabFactories.RegisterFactory(MakeShared<FOptimusEditorCompilerOutputTabSummoner>(InEditorApp));
 	TabFactories.RegisterFactory(MakeShared<FOptimusEditorShaderTextEditorTabSummoner>(InEditorApp));
 	
@@ -65,7 +66,7 @@ TSharedRef<FTabManager::FLayout> FOptimusEditorMode::CreatePaneLayout()
 	// Deets = Details panel
 	// Text = Shader Text Editor
 
-	TSharedRef<FTabManager::FLayout> Layout = FTabManager::NewLayout("Standalone_OptimusEditor_Layout_v06")
+	TSharedRef<FTabManager::FLayout> Layout = FTabManager::NewLayout("Standalone_OptimusEditor_Layout_v07")
 		->AddArea(
 			FTabManager::NewPrimaryArea()->SetOrientation(Orient_Vertical)
 			->Split(							// - Main work area
@@ -94,6 +95,8 @@ TSharedRef<FTabManager::FLayout> FOptimusEditorMode::CreatePaneLayout()
 					->Split(					// --- Output
 						FTabManager::NewStack()->SetSizeCoefficient(0.2f)
 						->AddTab(FOptimusEditorCompilerOutputTabSummoner::TabId, ETabState::OpenedTab)
+						->AddTab(FOptimusEditorFindTabSummoner::TabId, ETabState::OpenedTab)
+						->SetForegroundTab(FOptimusEditorCompilerOutputTabSummoner::TabId)
 					)
 				)
 				->Split(						// -- Details + Shader Text Editor
diff --git a/Engine/Plugins/Animation/DeformerGraph/Source/OptimusEditor/Private/OptimusEditorTabSummoners.cpp b/Engine/Plugins/Animation/DeformerGraph/Source/OptimusEditor/Private/OptimusEditorTabSummoners.cpp
--- a/Engine/Plugins/Animation/DeformerGraph/Source/OptimusEditor/Private/OptimusEditorTabSummoners.cpp
+++ b/Engine/Plugins/Animation/DeformerGraph/Source/OptimusEditor/Private/OptimusEditorTabSummoners.cpp
@@ -4,6 +4,7 @@
 
 #include "Widgets/Docking/SDockTab.h"
 
+#include "FindInDeformer.h"
 #include "GraphEditor.h"
 #include "MessageLogModule.h"
 #include "Modules/ModuleManager.h"
@@ -16,7 +17,8 @@
 
 const FName FOptimusEditorNodePaletteTabSummoner::TabId("OptimusEditor_Palette");
 const FName FOptimusEditorExplorerTabSummoner::TabId("OptimusEditor_Explorer");
-const FName FOptimusEditorGraphTabSummoner::TabId("OptimusEditor_Graph"); 
+const FName FOptimusEditorGraphTabSummoner::TabId("OptimusEditor_Graph");
+const FName FOptimusEditorFindTabSummoner::TabId("OptimusEditor_Find");
 const FName FOptimusEditorCompilerOutputTabSummoner::TabId("OptimusEditor_Output"); 
 const FName FOptimusEditorShaderTextEditorTabSummoner::TabId("OptimusEditor_ShaderTextEditor");
 
@@ -88,11 +90,36 @@ TSharedRef<SWidget> FOptimusEditorGraphTabSummoner::CreateTabBody(
 	const FWorkflowTabSpawnInfo& Info
 	) const
 {
+	TSharedRef<SWidget> Test = EditorPtr.Pin()->GetGraphEditorWidget().ToSharedRef();
 	// FIXME: Move to own widget.
 	return EditorPtr.Pin()->GetGraphEditorWidget().ToSharedRef();
 }
 
 
+FOptimusEditorFindTabSummoner::FOptimusEditorFindTabSummoner(
+	TSharedRef<FOptimusEditor> InEditorApp
+) :
+	FWorkflowTabFactory(TabId, InEditorApp),
+	EditorPtr(InEditorApp)
+{
+	TabLabel = LOCTEXT("FindTab_TabLabel", "Find");
+	TabIcon = FSlateIcon(FAppStyle::GetAppStyleSetName(), "Kismet.Tabs.FindResults");
+
+	bIsSingleton = false;
+
+	ViewMenuDescription = LOCTEXT("FindTab_MenuLabel", "Find");
+	ViewMenuTooltip = LOCTEXT("FindTab_MenuLabel_Tooltip", "Search contents of the Graph");
+}
+
+
+TSharedRef<SWidget> FOptimusEditorFindTabSummoner::CreateTabBody(
+	const FWorkflowTabSpawnInfo& Info
+) const
+{
+	return EditorPtr.Pin()->GetFindResultsWidget().ToSharedRef();
+}
+
+
 FOptimusEditorCompilerOutputTabSummoner::FOptimusEditorCompilerOutputTabSummoner(
 	TSharedRef<FOptimusEditor> InEditorApp
 	) :
diff --git a/Engine/Plugins/Animation/DeformerGraph/Source/OptimusEditor/Private/OptimusEditorTabSummoners.h b/Engine/Plugins/Animation/DeformerGraph/Source/OptimusEditor/Private/OptimusEditorTabSummoners.h
--- a/Engine/Plugins/Animation/DeformerGraph/Source/OptimusEditor/Private/OptimusEditorTabSummoners.h
+++ b/Engine/Plugins/Animation/DeformerGraph/Source/OptimusEditor/Private/OptimusEditorTabSummoners.h
@@ -48,6 +48,20 @@ protected:
 };
 
 
+struct FOptimusEditorFindTabSummoner :
+	public FWorkflowTabFactory
+{
+	static const FName TabId;
+
+	explicit FOptimusEditorFindTabSummoner(TSharedRef<FOptimusEditor> InEditorApp);
+
+	TSharedRef<SWidget> CreateTabBody(const FWorkflowTabSpawnInfo& Info) const override;
+
+protected:
+	TWeakPtr<FOptimusEditor> EditorPtr;
+};
+
+
 
 struct FOptimusEditorCompilerOutputTabSummoner :
 	public FWorkflowTabFactory
diff --git a/Engine/Plugins/Animation/DeformerGraph/Source/OptimusEditor/Public/FindInDeformer.h b/Engine/Plugins/Animation/DeformerGraph/Source/OptimusEditor/Public/FindInDeformer.h
new file mode 100644
--- /dev/null
+++ b/Engine/Plugins/Animation/DeformerGraph/Source/OptimusEditor/Public/FindInDeformer.h
@@ -0,0 +1,35 @@
+// Copyright Epic Games, Inc. All Rights Reserved.
+
+#pragma once
+
+#include "CoreMinimal.h"
+#include "FindInGraph.h"
+
+class FAssetEditorToolkit;
+class SSearchBox;
+
+/** Item that matched the search results */
+class FFindInDeformerResult : public FFindInGraphResult
+{
+public:
+
+	FFindInDeformerResult(const FFindInGraphResult::FCreateParams& InCreateParams);
+
+	//~ Begin FFindInGraphResult Interface
+	virtual void JumpToNode(TWeakPtr<FAssetEditorToolkit> AssetEditorToolkit, const UEdGraphNode* InNode) const override;
+	//~ End FFindInGraphResult Interface
+};
+
+
+/** Widget for searching for items that are part of a Deformer Graph */
+class SFindInDeformer : public SFindInGraph
+{
+
+protected:
+
+	//~ Begin SFindInGraph Interface
+	virtual TSharedPtr<FFindInGraphResult> MakeSearchResult(const FFindInGraphResult::FCreateParams& InParams) override;
+	virtual const UEdGraph* GetGraph() override;
+	virtual bool MatchTokensInNode(const UEdGraphNode* Node, const TArray<FString>& Tokens) override;
+	//~ End SFindInGraph Interface
+};
diff --git a/Engine/Source/Editor/GraphEditor/Private/FindInGraph.cpp b/Engine/Source/Editor/GraphEditor/Private/FindInGraph.cpp
new file mode 100644
--- /dev/null
+++ b/Engine/Source/Editor/GraphEditor/Private/FindInGraph.cpp
@@ -0,0 +1,435 @@
+// Copyright Epic Games, Inc. All Rights Reserved.
+
+#include "FindInGraph.h"
+#include "Layout/WidgetPath.h"
+#include "Framework/Application/SlateApplication.h"
+#include "Widgets/Images/SImage.h"
+#include "Styling/AppStyle.h"
+#include "EdGraphNode_Comment.h"
+#include "EdGraph/EdGraphNode.h"
+#include "Widgets/Input/SSearchBox.h"
+
+#define LOCTEXT_NAMESPACE "FindInGraph"
+
+//////////////////////////////////////////////////////////////////////////
+// FFindInGraphResult
+
+FFindInGraphResult::FFindInGraphResult(const FCreateParams& InCreateParams)
+	: Parent(InCreateParams.Parent)
+	, Value(InCreateParams.Value)
+	, DuplicationIndex(InCreateParams.DuplicationIndex)
+	, Class(InCreateParams.Class)
+	, Pin(InCreateParams.Pin)
+	, GraphNode(InCreateParams.GraphNode)
+{
+	if (!Class && InCreateParams.GraphNode)
+	{
+		Class = InCreateParams.GraphNode->GetClass();
+	}
+
+	if (GraphNode.IsValid())
+	{
+		CommentText = GraphNode->NodeComment;
+	}
+}
+
+void FFindInGraphResult::JumpToNode(TWeakPtr<FAssetEditorToolkit> AssetEditorToolkit, const UEdGraphNode* InNode) const
+{
+}
+
+FReply FFindInGraphResult::OnClick(TWeakPtr<FAssetEditorToolkit> AssetEditorToolkit)
+{
+	if (GraphNode.IsValid())
+	{
+		JumpToNode(AssetEditorToolkit, GraphNode.Get());
+	}
+	else if (UEdGraphPin* ResolvedPin = Pin.Get())
+	{
+		JumpToNode(AssetEditorToolkit, ResolvedPin->GetOwningNode());
+	}
+	return FReply::Handled();
+}
+
+FText FFindInGraphResult::GetCategory() const
+{
+	if (Class == nullptr && Pin.Get())
+	{
+		return LOCTEXT("PinCategory", "Pin");
+	}
+
+	return LOCTEXT("NodeCategory", "Node");
+}
+
+TSharedRef<SWidget> FFindInGraphResult::CreateIcon() const
+{
+	FSlateColor IconColor = FSlateColor::UseForeground();
+	const FSlateBrush* Brush = NULL;
+	if (UEdGraphPin* ResolvedPin = Pin.Get())
+	{
+		if (ResolvedPin->PinType.IsArray())
+		{
+			Brush = FAppStyle::GetBrush(TEXT("GraphEditor.ArrayPinIcon"));
+		}
+		else if (ResolvedPin->PinType.bIsReference)
+		{
+			Brush = FAppStyle::GetBrush(TEXT("GraphEditor.RefPinIcon"));
+		}
+		else
+		{
+			Brush = FAppStyle::GetBrush(TEXT("GraphEditor.PinIcon"));
+		}
+		const UEdGraphSchema* Schema = ResolvedPin->GetSchema();
+			IconColor = Schema->GetPinTypeColor(ResolvedPin->PinType);
+		}
+	else if (GraphNode.IsValid())
+	{
+		Brush = FAppStyle::GetBrush(TEXT("GraphEditor.NodeGlyph"));
+	}
+
+	return 	SNew(SImage)
+		.Image(Brush)
+		.ColorAndOpacity(IconColor)
+		.ToolTipText(GetCategory());
+}
+
+FString FFindInGraphResult::GetCommentText() const
+{
+	return CommentText;
+}
+
+FString FFindInGraphResult::GetValueText() const
+{
+	FString DisplayVal;
+
+	if (UEdGraphPin* ResolvedPin = Pin.Get())
+	{
+		if ((!ResolvedPin->DefaultValue.IsEmpty() || !ResolvedPin->AutogeneratedDefaultValue.IsEmpty() || ResolvedPin->DefaultObject || !ResolvedPin->DefaultTextValue.IsEmpty())
+			&& (ResolvedPin->Direction == EGPD_Input)
+			&& (ResolvedPin->LinkedTo.Num() == 0))
+		{
+			if (ResolvedPin->DefaultObject)
+			{
+				const AActor* DefaultActor = Cast<AActor>(ResolvedPin->DefaultObject);
+				DisplayVal = FString::Printf(TEXT("(%s)"), DefaultActor ? *DefaultActor->GetActorLabel() : *ResolvedPin->DefaultObject->GetName());
+			}
+			else if (!ResolvedPin->DefaultValue.IsEmpty())
+			{
+				DisplayVal = FString::Printf(TEXT("(%s)"), *ResolvedPin->DefaultValue);
+			}
+			else if (!ResolvedPin->AutogeneratedDefaultValue.IsEmpty())
+			{
+				DisplayVal = FString::Printf(TEXT("(%s)"), *ResolvedPin->AutogeneratedDefaultValue);
+			}
+			else if (!ResolvedPin->DefaultTextValue.IsEmpty())
+			{
+				DisplayVal = FString::Printf(TEXT("(%s)"), *ResolvedPin->DefaultTextValue.ToString());
+			}
+		}
+	}
+
+	return DisplayVal;
+}
+
+//////////////////////////////////////////////////////////////////////////
+// SFindInGraph
+
+void SFindInGraph::Construct(const FArguments& InArgs, TSharedPtr<FAssetEditorToolkit> InAssetEditorToolkit)
+{
+	AssetEditorToolkitPtr = InAssetEditorToolkit;
+
+	this->ChildSlot
+	[
+		SNew(SVerticalBox)
+		+ SVerticalBox::Slot()
+		.AutoHeight()
+		[
+			SNew(SHorizontalBox)
+			+ SHorizontalBox::Slot()
+			.FillWidth(1)
+			[
+				SAssignNew(SearchTextField, SSearchBox)
+				.HintText(LOCTEXT("GraphSearchHint", "Search"))
+				.OnTextChanged(this, &SFindInGraph::OnSearchTextChanged)
+				.OnTextCommitted(this, &SFindInGraph::OnSearchTextCommitted)
+			]
+		]
+		+ SVerticalBox::Slot()
+		.FillHeight(1.0f)
+		.Padding(0.f, 4.f, 0.f, 0.f)
+		[
+			SNew(SBorder)
+			.BorderImage(FAppStyle::GetBrush("Menu.Background"))
+			[
+				SAssignNew(TreeView, STreeViewType)
+				.TreeItemsSource(&ItemsFound)
+				.OnGenerateRow(this, &SFindInGraph::OnGenerateRow)
+				.OnGetChildren(this, &SFindInGraph::OnGetChildren)
+				.OnSelectionChanged(this, &SFindInGraph::OnTreeSelectionChanged)
+				.OnMouseButtonDoubleClick(this, &SFindInGraph::OnTreeSelectionDoubleClick)
+				.SelectionMode(ESelectionMode::Multi)
+			]
+		]
+	];
+}
+
+void SFindInGraph::FocusForUse()
+{
+	// NOTE: Careful, GeneratePathToWidget can be reentrant in that it can call visibility delegates and such
+	FWidgetPath FilterTextBoxWidgetPath;
+	FSlateApplication::Get().GeneratePathToWidgetUnchecked(SearchTextField.ToSharedRef(), FilterTextBoxWidgetPath);
+
+	// Set keyboard focus directly
+	FSlateApplication::Get().SetKeyboardFocus(FilterTextBoxWidgetPath, EFocusCause::SetDirectly);
+}
+
+TSharedPtr<FFindInGraphResult> SFindInGraph::MakeSearchResult(const FFindInGraphResult::FCreateParams& InParams)
+{
+	return MakeShared<FFindInGraphResult>(InParams);
+}
+
+void SFindInGraph::OnSearchTextChanged(const FText& Text)
+{
+	SearchValue = Text.ToString();
+}
+
+void SFindInGraph::OnSearchTextCommitted(const FText& Text, ETextCommit::Type CommitType)
+{
+	if (CommitType == ETextCommit::OnEnter)
+	{
+		InitiateSearch();
+	}
+}
+
+void SFindInGraph::InitiateSearch()
+{
+	TArray<FString> Tokens;
+	if (SearchValue.Contains("\"") && SearchValue.ParseIntoArray(Tokens, TEXT("\""), true) > 0)
+	{
+		for (FString& Token : Tokens)
+		{
+			// We have the token, we don't need the quotes anymore, they'll just confused the comparison later on
+			Token = Token.TrimQuotes();
+
+			// We remove the spaces as all later comparison strings will also be de-spaced
+			Token = Token.Replace(TEXT(" "), TEXT(""));
+		}
+
+		// due to being able to handle multiple quoted blocks like ("Make Epic" "Game Now") we can end up with
+		// and empty string between (" ") blocks so this simply removes them
+		struct FRemoveMatchingStrings
+		{
+			bool operator()(const FString& RemovalCandidate) const
+			{
+				return RemovalCandidate.IsEmpty();
+			}
+		};
+		Tokens.RemoveAll(FRemoveMatchingStrings());
+	}
+	else
+	{
+		// unquoted search equivalent to a match-any-of search
+		SearchValue.ParseIntoArray(Tokens, TEXT(" "), true);
+	}
+
+	for (FSearchResult& ItemFound : ItemsFound)
+	{
+		TreeView->SetItemExpansion(ItemFound, false);
+	}
+	ItemsFound.Empty();
+
+	if (Tokens.Num() > 0)
+	{
+		HighlightText = FText::FromString(SearchValue);
+		MatchTokens(Tokens);
+	}
+
+	// Insert a fake result to inform user if none found
+	if (ItemsFound.Num() == 0)
+	{
+		ItemsFound.Add(MakeSearchResult({ LOCTEXT("GraphSearchNoResults", "No Results found").ToString() }));
+	}
+
+	TreeView->RequestTreeRefresh();
+
+	for (FSearchResult& ItemFound : ItemsFound)
+	{
+		TreeView->SetItemExpansion(ItemFound, true);
+	}
+}
+
+const UEdGraph* SFindInGraph::GetGraph()
+{
+	return nullptr;
+}
+
+void SFindInGraph::MatchTokens(const TArray<FString> &Tokens)
+{
+	RootSearchResult.Reset();
+
+	if (const UEdGraph* Graph = GetGraph())
+	{
+		MatchTokensInGraph(Graph, Tokens);
+	}
+}
+
+bool SFindInGraph::MatchTokensInNode(const UEdGraphNode* Node, const TArray<FString>& Tokens)
+{
+	return false;
+}
+
+void SFindInGraph::MatchTokensInGraph(const UEdGraph* Graph, const TArray<FString>& Tokens)
+{
+	if (!Graph)
+	{
+		return;
+	}
+
+	RootSearchResult = MakeSearchResult({ FString("GraphSearchTreeRoot") });
+
+	for (UEdGraphNode* Node : Graph->Nodes)
+	{
+		const FString NodeName = Node->GetNodeTitle(ENodeTitleType::FullTitle).ToString();
+		const FString NodeType = Node->GetNodeTitle(ENodeTitleType::ListView).ToString();
+
+		FSearchResult NodeResult = MakeSearchResult(
+		{ 
+			.Value = NodeName == NodeType ? NodeName : NodeName + " - " + NodeType, 
+			.Parent = RootSearchResult, 
+			.GraphNode = Node 
+		});
+
+		FString NodeSearchString = NodeName + NodeType + Node->NodeComment;
+		NodeSearchString = NodeSearchString.Replace(TEXT(" "), TEXT(""));
+
+		bool bNodeMatchesSearch = StringMatchesSearchTokens(Tokens, NodeSearchString);
+
+		// Allow for node-type specific searches to trigger positive matches if we didn't already find one
+		if (!bNodeMatchesSearch)
+		{
+			bNodeMatchesSearch = MatchTokensInNode(Node, Tokens);
+		}
+
+		for (UEdGraphPin* Pin : Node->Pins)
+		{
+			if (Pin && Pin->PinFriendlyName.CompareTo(FText::FromString(TEXT(" "))) != 0)
+			{
+				FText PinName = Pin->GetSchema()->GetPinDisplayName(Pin);
+				FString PinSearchString = Pin->PinName.ToString() 
+					+ Pin->PinFriendlyName.ToString() 
+					+ Pin->DefaultValue 
+					+ Pin->PinType.PinCategory.ToString() 
+					+ Pin->PinType.PinSubCategory.ToString() 
+					+ (Pin->PinType.PinSubCategoryObject.IsValid() ? Pin->PinType.PinSubCategoryObject.Get()->GetFullName() : TEXT(""));
+
+				PinSearchString = PinSearchString.Replace(TEXT(" "), TEXT(""));
+				if (StringMatchesSearchTokens(Tokens, PinSearchString))
+				{
+					FSearchResult PinResult = MakeSearchResult(
+					{
+						.Value = PinName.ToString(),
+						.Parent = NodeResult,
+						.Pin = Pin
+					});
+					NodeResult->Children.Add(PinResult);
+				}
+			}
+		}
+
+		if ((NodeResult->Children.Num() > 0) || bNodeMatchesSearch)
+		{
+			ItemsFound.Add(NodeResult);
+		}
+	}
+
+	for (const UEdGraph* Subgraph : Graph->SubGraphs)
+	{
+		MatchTokensInGraph(Subgraph, Tokens);
+	}
+}
+
+TSharedRef<ITableRow> SFindInGraph::OnGenerateRow(FSearchResult InItem, const TSharedRef<STableViewBase>& OwnerTable)
+{
+	FString CommentText = InItem->GetCommentText();
+
+	return SNew(STableRow< TSharedPtr<FFindInGraphResult> >, OwnerTable)
+		[
+			SNew(SHorizontalBox)
+			+ SHorizontalBox::Slot()
+			.AutoWidth()
+			.VAlign(VAlign_Center)
+			[
+				InItem->CreateIcon()
+			]
+			+ SHorizontalBox::Slot()
+				.AutoWidth()
+				.VAlign(VAlign_Center)
+				.Padding(2, 0)
+				[
+					SNew(STextBlock)
+					.Text(FText::FromString(InItem->Value))
+					.HighlightText(HighlightText)
+					.ToolTipText(FText::Format(LOCTEXT("GraphResultSearchToolTipFmt", "{0} : {1}"), InItem->GetCategory(), FText::FromString(InItem->Value)))
+				]
+			+ SHorizontalBox::Slot()
+				.FillWidth(1)
+				.HAlign(HAlign_Right)
+				.VAlign(VAlign_Center)
+				.Padding(2, 0)
+				[
+					SNew(STextBlock)
+					.Text(FText::FromString(InItem->GetValueText()))
+					.HighlightText(HighlightText)
+				]
+			+ SHorizontalBox::Slot()
+				.FillWidth(1)
+				.HAlign(HAlign_Right)
+				.VAlign(VAlign_Center)
+				.Padding(2, 0)
+				[
+					SNew(STextBlock)
+					.Text(CommentText.IsEmpty() ? FText::GetEmpty() : FText::Format(LOCTEXT("NodeCommentFmt", "Node Comment:[{0}]"), FText::FromString(CommentText)))
+					.ColorAndOpacity(FLinearColor::Yellow)
+					.HighlightText(HighlightText)
+				]
+		];
+}
+
+void SFindInGraph::OnGetChildren(FSearchResult InItem, TArray< FSearchResult >& OutChildren)
+{
+	OutChildren += InItem->Children;
+}
+
+void SFindInGraph::OnTreeSelectionChanged(FSearchResult Item, ESelectInfo::Type)
+{
+	if (Item.IsValid())
+	{
+		Item->OnClick(AssetEditorToolkitPtr);
+	}
+}
+
+void SFindInGraph::OnTreeSelectionDoubleClick(FSearchResult Item)
+{
+	if (Item.IsValid())
+	{
+		Item->OnClick(AssetEditorToolkitPtr);
+	}
+}
+
+bool SFindInGraph::StringMatchesSearchTokens(const TArray<FString>& Tokens, const FString& ComparisonString)
+{
+	bool bFoundAllTokens = true;
+	//search the entry for each token, it must have all of them to pass
+	for (const FString& Token : Tokens)
+	{
+		if (!ComparisonString.Contains(Token))
+		{
+			bFoundAllTokens = false;
+			break;
+		}
+	}
+	return bFoundAllTokens;
+}
+
+/////////////////////////////////////////////////////
+
+#undef LOCTEXT_NAMESPACE
diff --git a/Engine/Source/Editor/GraphEditor/Public/FindInGraph.h b/Engine/Source/Editor/GraphEditor/Public/FindInGraph.h
new file mode 100644
--- /dev/null
+++ b/Engine/Source/Editor/GraphEditor/Public/FindInGraph.h
@@ -0,0 +1,162 @@
+// Copyright Epic Games, Inc. All Rights Reserved.
+
+#pragma once
+
+#include "CoreMinimal.h"
+#include "EdGraph/EdGraph.h"
+#include "EdGraph/EdGraphPin.h"
+#include "Input/Reply.h"
+#include "Widgets/SWidget.h"
+#include "Widgets/DeclarativeSyntaxSupport.h"
+#include "Widgets/SCompoundWidget.h"
+#include "Toolkits/AssetEditorToolkit.h"
+#include "Widgets/Views/STableViewBase.h"
+#include "Widgets/Views/STableRow.h"
+#include "Widgets/Views/STreeView.h"
+
+class FAssetEditorToolkit;
+class SSearchBox;
+
+/** Item that matched the search results */
+class GRAPHEDITOR_API FFindInGraphResult
+{
+public:
+	virtual ~FFindInGraphResult() = default;
+
+	struct FCreateParams
+	{
+		const FString& Value;
+		TSharedPtr<FFindInGraphResult> Parent = nullptr;
+		int	DuplicationIndex = 0;
+		UClass* Class = nullptr;
+		UEdGraphPin* Pin = nullptr;
+		UEdGraphNode* GraphNode = nullptr;
+	};
+
+	/** Create a root (or only text) result */
+	FFindInGraphResult(const FCreateParams& InCreateParams);
+
+	/** By default NoOp, implement this to make your result jump to a node or node owning pin on click*/
+	virtual void JumpToNode(TWeakPtr<FAssetEditorToolkit> AssetEditorToolkit, const UEdGraphNode* InNode) const;
+
+	/** Called when user clicks on the search item */
+	virtual FReply OnClick(TWeakPtr<FAssetEditorToolkit> AssetEditorToolkit);
+
+	/* Get Category for this search result */
+	virtual FText GetCategory() const;
+
+	/** Create an icon to represent the result */
+	virtual TSharedRef<SWidget>	CreateIcon() const;
+
+	/** Gets the comment on this node if any */
+	virtual FString GetCommentText() const;
+
+	/** Gets the value of the pin if any */
+	virtual FString GetValueText() const;
+
+	/** Any children listed under this category */
+	TArray<TSharedPtr<FFindInGraphResult>> Children;
+
+	/** Search result Parent */
+	TWeakPtr<FFindInGraphResult> Parent;
+
+	/*The meta string that was stored in the asset registry for this item */
+	FString Value;
+
+	/*The graph may have multiple instances of whatever we are looking for, this tells us which instance # we refer to*/
+	int	DuplicationIndex;
+
+	/*The class this item refers to */
+	UClass* Class;
+
+	/** The pin that this search result refers to */
+	FEdGraphPinReference Pin;
+
+	/** The graph node that this search result refers to (if not by asset registry or UK2Node) */
+	TWeakObjectPtr<UEdGraphNode> GraphNode;
+
+	/** Display text for comment information */
+	FString CommentText;
+};
+
+
+/** Widget for searching for items that are part of a UEdGraph */
+class GRAPHEDITOR_API SFindInGraph : public SCompoundWidget
+{
+public:
+	SLATE_BEGIN_ARGS(SFindInGraph){}
+	SLATE_END_ARGS()
+
+	void Construct(const FArguments& InArgs, TSharedPtr<FAssetEditorToolkit> InAssetEditorToolkit);
+
+	/** Focuses this widget's search box */
+	virtual void FocusForUse();
+
+protected:
+	typedef TSharedPtr<FFindInGraphResult> FSearchResult;
+	typedef STreeView<FSearchResult> STreeViewType;
+
+	/** Override this to create the search result of your type */
+	virtual TSharedPtr<FFindInGraphResult> MakeSearchResult(const FFindInGraphResult::FCreateParams& InParams);
+
+	/** Called when user changes the text they are searching for */
+	virtual void OnSearchTextChanged(const FText& Text);
+
+	/** Called when user changes commits text to the search box */
+	virtual void OnSearchTextCommitted(const FText& Text, ETextCommit::Type CommitType);
+
+	/** Get the children of a row */
+	virtual void OnGetChildren(FSearchResult InItem, TArray<FSearchResult>& OutChildren);
+
+	/** Called when user clicks on a new result */
+	virtual void OnTreeSelectionChanged(FSearchResult Item, ESelectInfo::Type SelectInfo);
+
+	/** Called when user double clicks on a new result */
+	virtual void OnTreeSelectionDoubleClick(FSearchResult Item);
+
+	/** Called when a new row is being generated */
+	virtual TSharedRef<ITableRow> OnGenerateRow(FSearchResult InItem, const TSharedRef<STableViewBase>& OwnerTable);
+
+	/** Begins the search based on the SearchValue */
+	virtual void InitiateSearch();
+
+	/** Get a graph corresponding to this editor, implement if you want to want to use the default graph node search */
+	virtual const UEdGraph* GetGraph();
+
+	/** Find any results that contain all of the tokens */
+	virtual void MatchTokens(const TArray<FString>& Tokens);
+
+	/** 
+	 * Implement this if you have node-specific search behavior. IE: Casting into child types to compare tokens. 
+	 * @return true if all tokens are valid for this node, false otherwisee
+	 */
+	virtual bool MatchTokensInNode(const UEdGraphNode* Node, const TArray<FString>& Tokens);
+
+	/** Find any results that contain all of the tokens in provided graph and subgraphs */
+	virtual void MatchTokensInGraph(const UEdGraph* Graph, const TArray<FString>& Tokens);
+
+	/** Determines if a string matches the search tokens */
+	static bool StringMatchesSearchTokens(const TArray<FString>& Tokens, const FString& ComparisonString);
+
+protected:
+	/** Pointer back to the Material editor that owns us */
+	TWeakPtr<FAssetEditorToolkit> AssetEditorToolkitPtr;
+
+	/** The tree view displays the results */
+	TSharedPtr<STreeViewType> TreeView;
+
+	/** The search text box */
+	TSharedPtr<SSearchBox> SearchTextField;
+
+	/** This buffer stores the currently displayed results */
+	TArray<FSearchResult> ItemsFound;
+
+	/** we need to keep a handle on the root result, because it won't show up in the tree */
+	FSearchResult RootSearchResult;
+
+	/** The string to highlight in the results */
+	FText HighlightText;
+
+	/** The string to search for */
+	FString	SearchValue;
+};
