Use smaller icon for preview-default, used by material graph nodes

diff --git a/Engine/Source/Editor/GraphEditor/Private/MaterialNodes/SGraphNodeMaterialBase.cpp b/Engine/Source/Editor/GraphEditor/Private/MaterialNodes/SGraphNodeMaterialBase.cpp
--- a/Engine/Source/Editor/GraphEditor/Private/MaterialNodes/SGraphNodeMaterialBase.cpp
+++ b/Engine/Source/Editor/GraphEditor/Private/MaterialNodes/SGraphNodeMaterialBase.cpp
@@ -566,8 +566,14 @@ void SGraphNodeMaterialBase::SetDefaultTitleAreaWidget(TSharedRef<SOverlay> Defa
 				.OnClicked(this, &SGraphNodeMaterialBase::OnTogglePreviewClicked)
 				.Cursor(EMouseCursor::Default)
 				[
-					SNew(SImage)
-					.Image(FAppStyle::GetBrush(TEXT("Icons.Preview")))
+					SNew(SHorizontalBox)
+					+SHorizontalBox::Slot()
+					.VAlign(VAlign_Center)
+					.HAlign(HAlign_Center)
+					[
+						SNew(SImage)
+						.Image(FAppStyle::GetBrush(TEXT("Icons.Preview")))
+					]
 				]
 			]
 			+ SHorizontalBox::Slot()
diff --git a/Engine/Source/Runtime/SlateCore/Private/Styling/StarshipCoreStyle.cpp b/Engine/Source/Runtime/SlateCore/Private/Styling/StarshipCoreStyle.cpp
--- a/Engine/Source/Runtime/SlateCore/Private/Styling/StarshipCoreStyle.cpp
+++ b/Engine/Source/Runtime/SlateCore/Private/Styling/StarshipCoreStyle.cpp
@@ -430,7 +430,7 @@ TSharedRef<ISlateStyle> FStarshipCoreStyle::Create()
 
 		Style->Set("Icons.Calendar", new IMAGE_BRUSH_SVG("Starship/Common/Calendar", Icon16x16));
 
-		Style->Set("Icons.Preview", new IMAGE_BRUSH_SVG("Starship/Common/preview-default", Icon20x20));
+		Style->Set("Icons.Preview", new IMAGE_BRUSH_SVG("Starship/Common/preview-default", Icon10x10));
 	}
 
 	// Tool panels
