Material Editor: Niagara-esque swizzle functionality with new 'Convert' node

- Convert nodes allow you to specify any number of inputs/outputs and route their components between them.

- Added a virtual method to UMaterialExpression that allows expressions to specify a custom node representation without having to create a UMaterialGraphNode subclass just for that purpose

diff --git a/Engine/Source/Editor/GraphEditor/Private/MaterialNodes/SGraphNodeMaterialConvert.cpp b/Engine/Source/Editor/GraphEditor/Private/MaterialNodes/SGraphNodeMaterialConvert.cpp
new file mode 100644
--- /dev/null
+++ b/Engine/Source/Editor/GraphEditor/Private/MaterialNodes/SGraphNodeMaterialConvert.cpp
@@ -0,0 +1,910 @@
+// Copyright Epic Games, Inc. All Rights Reserved.
+
+#include "MaterialNodes/SGraphNodeMaterialConvert.h"
+
+#include "ConnectionDrawingPolicy.h"
+#include "Framework/MultiBox/MultiBoxBuilder.h"
+#include "MaterialGraph/MaterialGraphNode.h"
+#include "Materials/MaterialExpressionConvert.h"
+#include "Misc/Attribute.h"
+#include "ScopedTransaction.h"
+#include "Styling/AppStyle.h"
+#include "Widgets/Images/SImage.h"
+#include "Widgets/Input/SButton.h"
+#include "Widgets/Input/SNumericEntryBox.h"
+#include "Widgets/Layout/SSpacer.h"
+#include "Widgets/SBoxPanel.h"
+
+#define LOCTEXT_NAMESPACE "GraphNodeMaterialConvert"
+
+////////////////////////////////////////////////////////////////////////////////////////////////////////////////////////////////////////////
+// FConvertDragDropOp: Drag and Drop Operation used to form connections within the convert node
+////////////////////////////////////////////////////////////////////////////////////////////////////////////////////////////////////////////
+
+FConvertDragDropOp::FConvertDragDropOp(TSharedPtr<SConvertInnerPin> InSourcePin)
+: SourcePin(InSourcePin)
+{
+	Construct();
+}
+
+void FConvertDragDropOp::OnDrop(bool bDropWasHandled, const FPointerEvent& MouseEvent)
+{
+	FDragDropOperation::OnDrop(bDropWasHandled, MouseEvent);
+
+	if (SourcePin)
+	{
+		SourcePin->CancelDragDrop();
+	}
+}
+
+void FConvertDragDropOp::OnDragged(const class FDragDropEvent& DragDropEvent)
+{
+	FDragDropOperation::OnDragged(DragDropEvent);
+
+	ScreenPosition = DragDropEvent.GetScreenSpacePosition();
+}
+
+TSharedPtr<SConvertInnerPin> FConvertDragDropOp::GetSourcePin() const
+{
+	return SourcePin;
+}
+
+FVector2D FConvertDragDropOp::GetScreenPosition() const
+{ 
+	return ScreenPosition;
+}
+
+////////////////////////////////////////////////////////////////////////////////////////////////////////////////////////////////////////////
+// SMaterialExpressionConvertGraphPin: The Outer pins that form connections to other material graph nodes
+////////////////////////////////////////////////////////////////////////////////////////////////////////////////////////////////////////////
+
+void SMaterialExpressionConvertGraphPin::Construct(const FArguments& InArgs, UEdGraphPin* InPin)
+{
+	SGraphPin::Construct(SGraphPin::FArguments(), InPin);
+
+	// Setup custom hover logic to only treat this pin as hovered when you hover over the main pin image
+	this->SetHover(TAttribute<bool>::CreateSP(this, &SMaterialExpressionConvertGraphPin::IsHoveredOverPrimaryPin));
+}
+
+void SMaterialExpressionConvertGraphPin::CreateInnerPins(TSharedRef<SGraphNodeMaterialConvert> InOwningGraphNodeWidget)
+{
+	UMaterialGraphNode* MaterialNode = InOwningGraphNodeWidget->GetMaterialGraphNode();
+	if (!MaterialNode)
+	{
+		return;
+	}
+
+	UMaterialExpressionConvert* ConvertExpression = Cast<UMaterialExpressionConvert>(MaterialNode->MaterialExpression);
+	if (!ConvertExpression)
+	{
+		return;
+	}
+
+	UEdGraphPin* EdGraphPin = GetPinObj();
+	if (!EdGraphPin)
+	{
+		return;
+	}
+		
+	if (TSharedPtr<SHorizontalBox> PinnedRowWidget = GetFullPinHorizontalRowWidget().Pin())
+	{
+		const bool bIsInputPin = EdGraphPin->Direction == EGPD_Input;
+		const int32 PinIndex = EdGraphPin->SourceIndex;
+				
+		const EMaterialExpressionConvertType ConvertType 	= bIsInputPin
+			? ConvertExpression->ConvertInputs[PinIndex].Type
+			: ConvertExpression->ConvertOutputs[PinIndex].Type;
+		const int32 PinComponentCount = MaterialExpressionConvertType::GetComponentCount(ConvertType);
+				
+		InnerPins.Reset(PinComponentCount);
+		TSharedPtr<SVerticalBox> InnerPinsVerticalBox = SNew(SVerticalBox);
+		for (int32 ComponentIndex = 0; ComponentIndex < PinComponentCount; ++ComponentIndex)
+		{
+			TSharedRef<SConvertInnerPin> NewInnerPin = SNew(SConvertInnerPin, 
+				InOwningGraphNodeWidget,
+				SharedThis(this),
+				bIsInputPin, 
+				PinIndex, 
+				ComponentIndex
+			);
+			InnerPinsVerticalBox->AddSlot()
+			.HAlign(bIsInputPin ? HAlign_Right : HAlign_Left)
+			[
+				NewInnerPin
+			];
+			InnerPins.Add(NewInnerPin);
+		}
+
+		// Insert as last element (INDEX_NONE) for input pins, first element for output pins
+		PinnedRowWidget->InsertSlot(bIsInputPin ? INDEX_NONE : 0)
+		.HAlign(bIsInputPin ? HAlign_Right : HAlign_Left)
+		.VAlign(VAlign_Center)
+		.AutoWidth()
+		[
+			InnerPinsVerticalBox.ToSharedRef()
+		];
+	}
+}
+
+bool SMaterialExpressionConvertGraphPin::IsHoveredOverPrimaryPin() const
+{
+	TSharedPtr<SWidget> PinImageWidget = this->GetPinImageWidget();
+	return PinImageWidget ? PinImageWidget->IsHovered() : false;
+}
+
+////////////////////////////////////////////////////////////////////////////////////////////////////////////////////////////////////////////
+// SConvertInnerPin: The inner pins used to route values with the convert node
+////////////////////////////////////////////////////////////////////////////////////////////////////////////////////////////////////////////
+
+void SConvertInnerPin::Construct(
+	const FArguments& InArgs,
+	TSharedPtr<SGraphNodeMaterialConvert> InOwningNode,
+	TSharedPtr<SMaterialExpressionConvertGraphPin> InOwningPin,
+	bool bInIsInputPin, 
+	int32 InPinIndex, 
+	int32 InComponentIndex
+)
+{
+	WeakOwningNode = InOwningNode;
+	WeakOwningPin = InOwningPin;
+	bIsInputPin = bInIsInputPin;
+	PinIndex = InPinIndex;
+	ComponentIndex = InComponentIndex;
+
+	// Construct Horizontal Box and add the PinImage to it
+	TSharedPtr<SHorizontalBox> HorizontalBox = 
+	SNew(SHorizontalBox)
+	+SHorizontalBox::Slot()
+	.AutoWidth()
+	.HAlign(HAlign_Left)
+	.VAlign(VAlign_Center)
+	.Padding(2.5f, 0.0f)
+	[
+		SAssignNew(PinImage, SImage)
+		.Image(this, &SConvertInnerPin::GetPinBrush)
+	];
+
+	// Add default value entry to either the front or back of the hbox depending on if we're an input pin
+	HorizontalBox->InsertSlot(bIsInputPin ? 0 : INDEX_NONE)
+	.AutoWidth()
+	.HAlign(HAlign_Left)
+	.VAlign(VAlign_Center)
+	[
+		SNew(SHorizontalBox)
+		+SHorizontalBox::Slot()
+		.AutoWidth()
+		.HAlign(HAlign_Left)
+		.VAlign(VAlign_Center)
+		.Padding(2.5f, 0.0f)
+		[
+			SNew(STextBlock)
+			.TextStyle(FAppStyle::Get(), "MonospacedText")
+			.Visibility(this, &SConvertInnerPin::GetPinNameVisibility)
+			.Text(this, &SConvertInnerPin::GetPinName)
+		]
+		+SHorizontalBox::Slot()
+		.AutoWidth()
+		.HAlign(HAlign_Left)
+		.VAlign(VAlign_Center)
+		.Padding(2.5f, 0.0f)
+		[
+			SNew(SNumericEntryBox<float>)
+			.EditableTextBoxStyle(FAppStyle::Get(), "Graph.EditableTextBox")
+			.BorderForegroundColor(FSlateColor::UseForeground())
+			.Visibility(this, &SConvertInnerPin::GetDefaultValueVisibility)
+			.Value(this, &SConvertInnerPin::GetDefaultValue)
+			.OnValueCommitted(this, &SConvertInnerPin::SetDefaultValue)
+		]
+	];
+
+	// Add inner node connection space to either the front or back of the hbox depending on if we're an input pin
+	HorizontalBox->InsertSlot(bIsInputPin ? INDEX_NONE : 0)
+	.AutoWidth()
+	.HAlign(HAlign_Left)
+	.VAlign(VAlign_Center)
+	[
+		SNew(SSpacer)
+		.Size(FVector2D(30, 0))
+	];
+
+	ChildSlot
+	[
+		HorizontalBox.ToSharedRef()
+	];
+}
+
+int32 SConvertInnerPin::OnPaint(
+	const FPaintArgs& Args,
+	const FGeometry& AllottedGeometry,
+	const FSlateRect& MyCullingRect,
+	FSlateWindowElementList& OutDrawElements,
+	int32 LayerId,
+	const FWidgetStyle& InWidgetStyle,
+	bool bParentEnabled
+) const
+{
+	LayerId = SCompoundWidget::OnPaint(
+		Args, 
+		AllottedGeometry, 
+		MyCullingRect, 
+		OutDrawElements, 
+		LayerId, 
+		InWidgetStyle, 
+		bParentEnabled
+	);
+
+	// Update CenterAbsolute after we've painted our PinImage
+	if (PinImage.IsValid())
+	{
+		CenterAbsolute = FGeometryHelper::CenterOf(PinImage->GetPaintSpaceGeometry());
+	}
+
+	return LayerId;
+}
+
+FReply SConvertInnerPin::OnMouseButtonDown(const FGeometry& MyGeometry, const FPointerEvent& MouseEvent)
+{
+	if (MouseEvent.GetEffectingButton() == EKeys::LeftMouseButton)
+	{
+		if (MouseEvent.IsAltDown())
+		{
+			BreakConnections();
+			return FReply::Handled();
+		}
+
+		return FReply::Handled()
+		.DetectDrag(SharedThis(this), MouseEvent.GetEffectingButton())
+		.CaptureMouse(SharedThis(this));
+	}
+	
+	// We still want to detect right-clicks on mouse button up
+	return FReply::Handled();
+}
+
+FReply SConvertInnerPin::OnMouseButtonUp(const FGeometry& MyGeometry, const FPointerEvent& MouseEvent)
+{
+	if (MouseEvent.GetEffectingButton() == EKeys::RightMouseButton)
+	{
+		// Only show to context menu if we right click the actual InnerPin
+		if (PinImage && PinImage->GetCachedGeometry().IsUnderLocation(MouseEvent.GetScreenSpacePosition()))
+		{
+			FMenuBuilder ContextMenuBuilder(true, NULL);
+
+			const FText BreakConnectionText =
+				WeakConnectedPins.Num() == 1
+				? LOCTEXT("ContextMenu.BreakInnerConnection", "Break Connection")
+				: LOCTEXT("ContextMenu.BreakInnerConnections", "Break Connections");
+
+			ContextMenuBuilder.AddMenuEntry(
+				BreakConnectionText,
+				FText(),
+				FSlateIcon(),
+				FUIAction(FExecuteAction::CreateSP(this, &SConvertInnerPin::BreakConnections))
+			);
+
+			FSlateApplication::Get().PushMenu(
+				AsShared(),
+				FWidgetPath(),
+				ContextMenuBuilder.MakeWidget(),
+				FSlateApplication::Get().GetCursorPos(),
+				FPopupTransitionEffect(FPopupTransitionEffect::ContextMenu)
+			);
+
+			return FReply::Handled();
+		}
+	}
+		
+	return FReply::Unhandled();
+}
+
+FReply SConvertInnerPin::OnDragDetected(const FGeometry& MyGeometry, const FPointerEvent& MouseEvent)
+{
+	if (TSharedPtr<SGraphNodeMaterialConvert> OwningNode = WeakOwningNode.Pin())
+	{
+		TSharedRef<FConvertDragDropOp> DragDropOp = MakeShareable(new FConvertDragDropOp(SharedThis(this)));
+		OwningNode->SetCurrentDragDropOp(DragDropOp);
+		return FReply::Handled().BeginDragDrop(DragDropOp);
+	}
+
+	return FReply::Unhandled();
+}
+
+FReply SConvertInnerPin::OnDrop(const FGeometry& MyGeometry, const FDragDropEvent& DragDropEvent)
+{
+	if (TSharedPtr<FConvertDragDropOp> DragDropOp = DragDropEvent.GetOperationAs<FConvertDragDropOp>())
+	{
+		if (TSharedPtr<SConvertInnerPin> SourcePin = DragDropOp->GetSourcePin())
+		{
+			// Is this a valid connection
+			if (SourcePin.Get() != this && SourcePin->WeakOwningNode == this->WeakOwningNode && SourcePin->bIsInputPin != this->bIsInputPin)
+			{
+				if (TSharedPtr<SGraphNodeMaterialConvert> OwningNode = WeakOwningNode.Pin())
+				{
+					// Form connection between two pins
+					OwningNode->FormConnection(SourcePin, SharedThis(this));
+				
+					// Clear The Current Drag Drop Op
+					OwningNode->SetCurrentDragDropOp(nullptr);
+					return FReply::Handled();
+				}
+			}
+		}
+	}
+
+	return FReply::Unhandled();
+}
+
+const FSlateBrush* SConvertInnerPin::GetPinBrush() const
+{
+	const bool bHasConnections = WeakConnectedPins.Num() > 0;
+	const FSlateBrush* ConnectedBrush = FAppStyle::GetBrush("Graph.Pin.Connected");
+	const FSlateBrush* DisconnectedBrush = FAppStyle::GetBrush("Graph.Pin.Disconnected");
+	return bHasConnections ? ConnectedBrush : DisconnectedBrush;
+}
+
+void SConvertInnerPin::CancelDragDrop()
+{
+	if (TSharedPtr<SGraphNodeMaterialConvert> OwningNode = WeakOwningNode.Pin())
+	{
+		OwningNode->SetCurrentDragDropOp(nullptr);
+	}
+}
+
+FVector2D SConvertInnerPin::GetPinCenterAbsolute() const
+{ 
+	return CenterAbsolute; 
+}
+
+void SConvertInnerPin::AddConnection(TSharedPtr<SConvertInnerPin> InOtherPin)
+{ 
+	WeakConnectedPins.Add(InOtherPin); 
+}
+
+void SConvertInnerPin::RemoveConnection(TSharedPtr<SConvertInnerPin> InOtherPin)
+{ 
+	WeakConnectedPins.Remove(InOtherPin); 
+}
+
+void SConvertInnerPin::RemoveAllConnections() 
+{
+	for (TWeakPtr<SConvertInnerPin> WeakConnectedPin : WeakConnectedPins)
+	{
+		if (TSharedPtr<SConvertInnerPin> ConnectedPin = WeakConnectedPin.Pin())
+		{
+			ConnectedPin->RemoveConnection(SharedThis(this));
+		}
+	}
+
+	WeakConnectedPins.Reset();
+}
+
+void SConvertInnerPin::BreakConnections()
+{
+	if (TSharedPtr<SGraphNodeMaterialConvert> OwningNode = WeakOwningNode.Pin())
+	{
+		OwningNode->BreakConnections(SharedThis(this));
+	}
+}
+
+EVisibility SConvertInnerPin::GetDefaultValueVisibility() const
+{ 
+	bool bIsVisible = false;
+
+	if (bIsInputPin)
+	{
+		if (TSharedPtr<SMaterialExpressionConvertGraphPin> OwningPin = WeakOwningPin.Pin())
+		{
+			// Only show Input Default values if the owning SMaterialExpressionConvertGraphPin isn't connected 
+			bIsVisible = !OwningPin->IsConnected();
+		}
+	}
+	else
+	{
+		// Show Output pin default values if they don't yet have a connection to an input inner pin
+		bIsVisible = WeakConnectedPins.IsEmpty();
+	}
+
+	return bIsVisible ? EVisibility::Visible : EVisibility::Collapsed;
+}
+
+TOptional<float> SConvertInnerPin::GetDefaultValue() const
+{
+	TSharedPtr<SGraphNodeMaterialConvert> OwningNode = WeakOwningNode.Pin();
+	if (!OwningNode)
+	{
+		return TOptional<float>();
+	}
+	return OwningNode->GetDefaultValue(SharedThis(this));
+}
+
+void SConvertInnerPin::SetDefaultValue(float InDefaultValue, ETextCommit::Type CommitType)
+{
+	if (TSharedPtr<SGraphNodeMaterialConvert> OwningNode = WeakOwningNode.Pin())
+	{
+		OwningNode->SetDefaultValue(
+			SharedThis(this),
+			InDefaultValue
+		);
+	}
+}
+
+EVisibility SConvertInnerPin::GetPinNameVisibility() const
+{
+	TSharedPtr<SMaterialExpressionConvertGraphPin> OwningPin = WeakOwningPin.Pin();
+	if (OwningPin && OwningPin->GetInnerPins().Num() == 1)
+	{
+		return EVisibility::Collapsed;
+	}
+
+	return EVisibility::SelfHitTestInvisible;
+}
+
+FText SConvertInnerPin::GetPinName() const
+{
+	TSharedPtr<SMaterialExpressionConvertGraphPin> OwningPin = WeakOwningPin.Pin();
+	if (OwningPin && OwningPin->GetInnerPins().Num() == 1)
+	{
+		return INVTEXT("");
+	}
+
+	switch (ComponentIndex)
+	{
+		case 0: 	return INVTEXT("R");
+		case 1: 	return INVTEXT("G");
+		case 2: 	return INVTEXT("B");
+		case 3: 	return INVTEXT("A");
+		default: 	return INVTEXT("_");
+	}
+}
+
+////////////////////////////////////////////////////////////////////////////////////////////////////////////////////////////////////////////
+// SGraphNodeMaterialConvert: Custom Slate widget for UMaterialExpressionConvert
+////////////////////////////////////////////////////////////////////////////////////////////////////////////////////////////////////////////
+
+void SGraphNodeMaterialConvert::Construct(const FArguments& InArgs, class UMaterialGraphNode* InNode)
+{
+	this->GraphNode = InNode;
+	this->MaterialNode = InNode;
+	this->UpdateGraphNode();
+}
+
+void SGraphNodeMaterialConvert::CreatePinWidgets()
+{
+	SGraphNodeMaterialBase::CreatePinWidgets();
+
+	if (!MaterialNode)
+	{
+		return;
+	}
+
+	UMaterialExpressionConvert* ConvertExpression = Cast<UMaterialExpressionConvert>(MaterialNode->MaterialExpression);
+	if (!ConvertExpression)
+	{
+		return;
+	}
+
+	// After we've set up all of our pins, set up initial connections based on our ConvertMappings
+	for (const FMaterialExpressionConvertMapping& ConvertMapping : ConvertExpression->ConvertMappings)
+	{
+		TSharedPtr<SConvertInnerPin> InputInnerPin = GetInnerPin(true, ConvertMapping.InputIndex, ConvertMapping.InputComponentIndex);
+		TSharedPtr<SConvertInnerPin> OutputInnerPin = GetInnerPin(false, ConvertMapping.OutputIndex, ConvertMapping.OutputComponentIndex);
+		if (InputInnerPin && OutputInnerPin)
+		{
+			InputInnerPin->AddConnection(OutputInnerPin);
+			OutputInnerPin->AddConnection(InputInnerPin);
+		}
+	}
+}
+
+TSharedPtr<SGraphPin> SGraphNodeMaterialConvert::CreatePinWidget(UEdGraphPin* Pin) const
+{
+	return SNew(SMaterialExpressionConvertGraphPin, Pin);
+}
+
+void SGraphNodeMaterialConvert::AddPin(const TSharedRef<SGraphPin>& PinToAdd)
+{
+	SGraphNodeMaterialBase::AddPin(PinToAdd);
+
+	// Create our Inner Pins, which are used to route input values to output values
+	TSharedRef<SMaterialExpressionConvertGraphPin> ConvertOuterPin = StaticCastSharedRef<SMaterialExpressionConvertGraphPin>(PinToAdd);
+	ConvertOuterPin->CreateInnerPins(SharedThis(this));
+}
+
+void SGraphNodeMaterialConvert::CreateInputSideAddButton(TSharedPtr<SVerticalBox> InputBox)
+{
+	TSharedRef<SButton> AddPinButton = StaticCastSharedRef<SButton>(AddPinButtonContent(
+		LOCTEXT("GraphNodeMaterialConvert_AddInputPin", "Add Input"),
+		LOCTEXT("GraphNodeMaterialConvert_AddInputPin_Tooltip", "Add an input to this convert node."),
+		false /* bRightSide */
+	));
+
+	AddPinButton->SetOnClicked(FOnClicked::CreateSP(this, &SGraphNodeMaterialConvert::OnAddInputPinClicked));
+
+	InputBox->AddSlot()
+	.FillHeight(1.f)
+	.Padding(5.f)
+	.VAlign(VAlign_Bottom)
+	.HAlign(HAlign_Left)
+	[
+		AddPinButton
+	];
+}
+
+void SGraphNodeMaterialConvert::CreateOutputSideAddButton(TSharedPtr<SVerticalBox> OutputBox)
+{
+	TSharedRef<SButton> AddPinButton = StaticCastSharedRef<SButton>(AddPinButtonContent(
+		LOCTEXT("GraphNodeMaterialConvert_AddOutputPin", "Add Output"),
+		LOCTEXT("GraphNodeMaterialConvert_AddOutputPin_Tooltip", "Add an output to this convert node."),
+		true /* bRightSide */
+	));
+
+	AddPinButton->SetOnClicked(FOnClicked::CreateSP(this, &SGraphNodeMaterialConvert::OnAddOutputPinClicked));
+
+	OutputBox->AddSlot()
+	.FillHeight(1.f)
+	.Padding(5.f)
+	.VAlign(VAlign_Bottom)
+	.HAlign(HAlign_Right)
+	[
+		AddPinButton
+	];
+}
+
+EVisibility SGraphNodeMaterialConvert::IsAddPinButtonVisible() const
+{
+	return EVisibility::Visible;
+}
+
+FReply SGraphNodeMaterialConvert::OnAddInputPinClicked()
+{
+	FSlateApplication::Get().PushMenu(
+		AsShared(), 
+		FWidgetPath(), 
+		CreateAddPinContextMenu(/* bInputPin = */ true), 
+		FSlateApplication::Get().GetCursorPos(), 
+		FPopupTransitionEffect(FPopupTransitionEffect::ContextMenu)
+	);
+
+	return FReply::Handled();
+}
+
+FReply SGraphNodeMaterialConvert::OnAddOutputPinClicked()
+{
+	FSlateApplication::Get().PushMenu(
+		AsShared(), 
+		FWidgetPath(), 
+		CreateAddPinContextMenu(/* bInputPin = */ false),
+		FSlateApplication::Get().GetCursorPos(), 
+		FPopupTransitionEffect(FPopupTransitionEffect::ContextMenu)
+	);
+
+	return FReply::Handled();
+}
+
+TSharedRef<SWidget> SGraphNodeMaterialConvert::CreateAddPinContextMenu(bool bInputPin)
+{
+	FMenuBuilder ContextMenuBuilder(true, NULL);
+
+	static const TArray<EMaterialExpressionConvertType> ConvertTypes =
+	{ 
+		EMaterialExpressionConvertType::Float1, 
+		EMaterialExpressionConvertType::Float2, 
+		EMaterialExpressionConvertType::Float3, 
+		EMaterialExpressionConvertType::Float4 
+	};
+
+	for (const EMaterialExpressionConvertType ConvertType : ConvertTypes)
+	{
+		ContextMenuBuilder.AddMenuEntry(
+			MaterialExpressionConvertType::ToText(ConvertType),
+			FText(),
+			FSlateIcon(),
+			FUIAction(FExecuteAction::CreateSP(
+				this, 
+				&SGraphNodeMaterialConvert::AddNewPin, bInputPin, ConvertType
+			))
+		);
+	}
+	
+	return ContextMenuBuilder.MakeWidget();
+}
+
+void SGraphNodeMaterialConvert::AddNewPin(bool bInputPin, EMaterialExpressionConvertType ConvertType)
+{
+	UMaterialExpressionConvert* ConvertExpression = Cast<UMaterialExpressionConvert>(MaterialNode->MaterialExpression);
+	if (!ConvertExpression)
+	{
+		return;
+	}
+
+	const FScopedTransaction Transaction(LOCTEXT("AddPin", "Add Pin"));
+	ConvertExpression->Modify();
+
+	if (bInputPin)
+	{
+		FMaterialExpressionConvertInput NewInput;
+		NewInput.Type = ConvertType;
+		ConvertExpression->ConvertInputs.Add(NewInput);
+	}
+	else
+	{
+		FMaterialExpressionConvertOutput NewOutput;
+		NewOutput.Type = ConvertType;
+		ConvertExpression->ConvertOutputs.Add(NewOutput);
+	}
+
+	MaterialNode->ReconstructNode();
+}
+
+int32 SGraphNodeMaterialConvert::OnPaint(
+	const FPaintArgs& Args,
+	const FGeometry& AllottedGeometry,
+	const FSlateRect& MyCullingRect,
+	FSlateWindowElementList& OutDrawElements,
+	int32 LayerId,
+	const FWidgetStyle& InWidgetStyle,
+	bool bParentEnabled
+) const
+{
+	LayerId = SGraphNodeMaterialBase::OnPaint(
+		Args, 
+		AllottedGeometry, 
+		MyCullingRect, 
+		OutDrawElements, 
+		LayerId, 
+		InWidgetStyle, 
+		bParentEnabled
+	);
+
+	// Only Draw Connections when we're at a high-enough detail. Early-out otherwise
+	if (GetCurrentLOD() <= EGraphRenderingLOD::LowDetail)
+	{
+		return LayerId;
+	}
+
+	if (!MaterialNode)
+	{
+		return LayerId;
+	}
+
+	UMaterialExpressionConvert* ConvertExpression = Cast<UMaterialExpressionConvert>(MaterialNode->MaterialExpression);
+	if (!ConvertExpression)
+	{
+		return LayerId;
+	}
+
+	for (const FMaterialExpressionConvertMapping& ConvertMapping : ConvertExpression->ConvertMappings)
+	{
+		TSharedPtr<SConvertInnerPin> InputInnerPin = GetInnerPin(true, ConvertMapping.InputIndex, ConvertMapping.InputComponentIndex);
+		TSharedPtr<SConvertInnerPin> OutputInnerPin = GetInnerPin(false, ConvertMapping.OutputIndex, ConvertMapping.OutputComponentIndex);
+		if (InputInnerPin && OutputInnerPin)
+		{
+			const FVector2D CurveStart = InputInnerPin->GetPinCenterAbsolute();
+			const FVector2D CurveEnd = OutputInnerPin->GetPinCenterAbsolute();
+			MakeConnectionCurve(AllottedGeometry, OutDrawElements, LayerId, CurveStart, CurveEnd);
+		}
+	}
+
+	if (CurrentDragDropOp.IsValid())
+	{
+		if (TSharedPtr<SConvertInnerPin> SourcePin = CurrentDragDropOp->GetSourcePin())
+		{
+			const FVector2D CurveStart = SourcePin->GetPinCenterAbsolute();
+			const FVector2D CurveEnd = CurrentDragDropOp->GetScreenPosition() + FVector2D(Inverse(Args.GetWindowToDesktopTransform()));
+			MakeConnectionCurve(AllottedGeometry, OutDrawElements, LayerId, CurveStart, CurveEnd);
+		}
+	}
+
+
+	return LayerId;
+}
+
+void SGraphNodeMaterialConvert::FormConnection(TSharedPtr<SConvertInnerPin> InnerPinA, TSharedPtr<SConvertInnerPin> InnerPinB)
+{
+	check(InnerPinA->IsInputPin() != InnerPinB->IsInputPin());
+
+	if (!MaterialNode)
+	{
+		return;
+	}
+
+	UMaterialExpressionConvert* ConvertExpression = Cast<UMaterialExpressionConvert>(MaterialNode->MaterialExpression);
+	if (!ConvertExpression)
+	{
+		return;
+	}
+
+	const FScopedTransaction Transaction(LOCTEXT("AddConnection", "Convert Node: Add Connection"));
+	ConvertExpression->Modify();
+
+	// Determine Which Pin is the input and which is the output
+	TSharedPtr<SConvertInnerPin> InputPin 	= InnerPinA->IsInputPin() ? InnerPinA : InnerPinB;
+	TSharedPtr<SConvertInnerPin> OutputPin 	= InnerPinA->IsInputPin() ? InnerPinB : InnerPinA;
+
+	FMaterialExpressionConvertMapping NewConvertMapping;
+	NewConvertMapping.InputIndex = InputPin->GetPinIndex();
+	NewConvertMapping.InputComponentIndex = InputPin->GetComponentIndex();
+	NewConvertMapping.OutputIndex = OutputPin->GetPinIndex();
+	NewConvertMapping.OutputComponentIndex = OutputPin->GetComponentIndex();
+
+	// Remove any mappings with the same OutputIndex and OutputComponentIndex
+	ConvertExpression->ConvertMappings.RemoveAll(
+		[&NewConvertMapping](const FMaterialExpressionConvertMapping& ExistingConvertMapping)
+		{
+			return ExistingConvertMapping.OutputIndex == NewConvertMapping.OutputIndex
+				&& ExistingConvertMapping.OutputComponentIndex == NewConvertMapping.OutputComponentIndex;
+		}
+	);
+
+	// Add new convert mapping and refresh
+	ConvertExpression->ConvertMappings.Add(NewConvertMapping);
+	ConvertExpression->RefreshNode();
+
+	// Input Pins can have multiple connections
+	InputPin->AddConnection(OutputPin);
+
+	// Output Pins can only have one connection, so clean up any existing connections
+	OutputPin->RemoveAllConnections();
+	OutputPin->AddConnection(InputPin);
+}
+
+void SGraphNodeMaterialConvert::BreakConnections(TSharedPtr<SConvertInnerPin> InnerPin)
+{
+	check(InnerPin);
+
+	if (!MaterialNode)
+	{
+		return;
+	}
+
+	UMaterialExpressionConvert* ConvertExpression = Cast<UMaterialExpressionConvert>(MaterialNode->MaterialExpression);
+	if (!ConvertExpression)
+	{
+		return;
+	}
+
+	const FScopedTransaction Transaction(LOCTEXT("BreakConnections", "Convert Node: Break Connections"));
+	ConvertExpression->Modify();
+
+	const bool bIsInputPin = InnerPin->IsInputPin();
+	const int32 PinIndex = InnerPin->GetPinIndex();
+	const int32 ComponentIndex = InnerPin->GetComponentIndex();
+
+	ConvertExpression->ConvertMappings.RemoveAll(
+		[bIsInputPin, PinIndex, ComponentIndex](const FMaterialExpressionConvertMapping& ExistingConvertMapping)
+		{
+			if (bIsInputPin)
+			{
+				return ExistingConvertMapping.InputIndex == PinIndex && ExistingConvertMapping.InputComponentIndex == ComponentIndex;
+			}
+			else
+			{
+				return ExistingConvertMapping.OutputIndex == PinIndex && ExistingConvertMapping.OutputComponentIndex == ComponentIndex;
+			}
+		}
+	);
+
+	ConvertExpression->RefreshNode();
+
+	InnerPin->RemoveAllConnections();
+}
+
+TOptional<float> SGraphNodeMaterialConvert::GetDefaultValue(TSharedPtr<const SConvertInnerPin> InnerPin) const
+{
+	UMaterialExpressionConvert* ConvertExpression = Cast<UMaterialExpressionConvert>(MaterialNode->MaterialExpression);
+	if (!ConvertExpression)
+	{
+		return TOptional<float>();
+	}
+
+	const int32 PinIndex = InnerPin->GetPinIndex();
+	const int32 ComponentIndex = InnerPin->GetComponentIndex();
+
+	if (InnerPin->IsInputPin())
+	{
+		return ConvertExpression->ConvertInputs[PinIndex].DefaultValue.Component(ComponentIndex);
+	}
+	else
+	{
+		return ConvertExpression->ConvertOutputs[PinIndex].DefaultValue.Component(ComponentIndex);
+	}
+}
+
+void SGraphNodeMaterialConvert::SetDefaultValue(TSharedPtr<SConvertInnerPin> InnerPin, const float InDefaultValue)
+{
+	UMaterialExpressionConvert* ConvertExpression = Cast<UMaterialExpressionConvert>(MaterialNode->MaterialExpression);
+	if (!ConvertExpression)
+	{
+		return;
+	}
+
+	const FScopedTransaction Transaction( LOCTEXT( "SetPinDefaultValue", "Set Pin Default Value" ) );
+	ConvertExpression->Modify();
+
+	const int32 PinIndex = InnerPin->GetPinIndex();
+	const int32 ComponentIndex = InnerPin->GetComponentIndex();
+
+	if (InnerPin->IsInputPin())
+	{
+		ConvertExpression->ConvertInputs[PinIndex].DefaultValue.Component(ComponentIndex) = InDefaultValue;
+	}
+	else
+	{
+		ConvertExpression->ConvertOutputs[PinIndex].DefaultValue.Component(ComponentIndex) = InDefaultValue;
+	}
+
+	ConvertExpression->RefreshNode();
+}
+
+TSharedPtr<SConvertInnerPin> SGraphNodeMaterialConvert::GetInnerPin(bool bInputPin, const int32 InPinIndex, const int32 InComponentIndex) const
+{
+	// First Find relevant input/output graph pin
+	TSharedPtr<SMaterialExpressionConvertGraphPin> GraphPin = nullptr;
+	if (bInputPin)
+	{
+		if (InputPins.IsValidIndex(InPinIndex))
+		{
+			GraphPin = StaticCastSharedRef<SMaterialExpressionConvertGraphPin>(InputPins[InPinIndex]);
+		}
+	}
+	else
+	{
+		if (OutputPins.IsValidIndex(InPinIndex))
+		{
+			GraphPin = StaticCastSharedRef<SMaterialExpressionConvertGraphPin>(OutputPins[InPinIndex]);
+		}
+	}
+
+	// Then get that the relevant inner pin from that graph pin
+	if (GraphPin.IsValid() && GraphPin->GetInnerPins().IsValidIndex(InComponentIndex))
+	{
+		return GraphPin->GetInnerPins()[InComponentIndex];
+	}
+
+	// No valid inner pin found, return nullptr
+	return nullptr;
+}
+
+void SGraphNodeMaterialConvert::MakeConnectionCurve(
+	const FGeometry& InAllottedGeometry,
+	FSlateWindowElementList& OutDrawElements,
+	int32& InOutLayerId,
+	const FVector2D InCurveStart,
+	const FVector2D InCurveEnd
+) const
+{
+	FVector2D LocalStart = InAllottedGeometry.AbsoluteToLocal(InCurveStart);
+	FVector2D LocalEnd = InAllottedGeometry.AbsoluteToLocal(InCurveEnd);
+
+	//If LocalStart.X is > LocalEnd.X, swap values
+	if (LocalStart.X > LocalEnd.X)
+	{
+		std::swap(LocalStart, LocalEnd);
+	}
+
+	const bool bAllValuesValid = 
+		(LocalStart.X != -FLT_MAX) 
+		&& (LocalStart.Y != -FLT_MAX) 
+		&& (LocalEnd.X != -FLT_MAX) 
+		&& (LocalEnd.Y != -FLT_MAX);
+
+	if (bAllValuesValid)
+	{
+		/* Because we make sure the left-most position is always the start position,
+		 * curve direction is always pointing to the right */
+		static const FVector2D Direction = FVector2D(100, 0);
+
+		FSlateDrawElement::MakeSpline(
+			OutDrawElements,
+			++InOutLayerId,
+			InAllottedGeometry.ToPaintGeometry(),
+			LocalStart,
+			Direction,
+			LocalEnd,
+			Direction,
+			2
+		);
+	}
+}
+
+#undef LOCTEXT_NAMESPACE
diff --git a/Engine/Source/Editor/GraphEditor/Public/MaterialNodes/SGraphNodeMaterialConvert.h b/Engine/Source/Editor/GraphEditor/Public/MaterialNodes/SGraphNodeMaterialConvert.h
new file mode 100644
--- /dev/null
+++ b/Engine/Source/Editor/GraphEditor/Public/MaterialNodes/SGraphNodeMaterialConvert.h
@@ -0,0 +1,200 @@
+// Copyright Epic Games, Inc. All Rights Reserved.
+
+#pragma once
+
+#include "Input/DragAndDrop.h"
+#include "MaterialNodes/SGraphNodeMaterialBase.h"
+#include "SGraphPin.h"
+#include "Widgets/SCompoundWidget.h"
+
+class FConvertDragDropOp;
+class SConvertInnerPin;
+class SGraphNodeMaterialConvert;
+class SImage;
+enum class EMaterialExpressionConvertType : uint8;
+
+////////////////////////////////////////////////////////////////////////////////////////////////////////////////////////////////////////////
+// FConvertDragDropOp: Drag and Drop Operation used to form connections within the convert node
+////////////////////////////////////////////////////////////////////////////////////////////////////////////////////////////////////////////
+
+class FConvertDragDropOp : public FDragDropOperation
+{
+public:
+	DRAG_DROP_OPERATOR_TYPE(FConvertDragDropOp, FDragDropOperation)
+
+		FConvertDragDropOp(TSharedPtr<SConvertInnerPin> InSourcePin);
+
+	//~ Begin FDragDropOperation interface
+	virtual void OnDrop(bool bDropWasHandled, const FPointerEvent& MouseEvent) override;
+	virtual void OnDragged(const class FDragDropEvent& DragDropEvent) override;
+	//~ End FDragDropOperation interface
+
+	TSharedPtr<SConvertInnerPin> GetSourcePin() const;
+	FVector2D GetScreenPosition() const;
+
+protected:
+	TSharedPtr<SConvertInnerPin> SourcePin;
+	FVector2D ScreenPosition;
+};
+
+////////////////////////////////////////////////////////////////////////////////////////////////////////////////////////////////////////////
+// SMaterialExpressionConvertGraphPin: The Outer pins that form connections to other material graph nodes
+////////////////////////////////////////////////////////////////////////////////////////////////////////////////////////////////////////////
+
+class SMaterialExpressionConvertGraphPin : public SGraphPin
+{
+public:
+	SLATE_BEGIN_ARGS(SMaterialExpressionConvertGraphPin){}
+	SLATE_END_ARGS()
+
+	void Construct(const FArguments& InArgs, UEdGraphPin* InPin);
+
+	void CreateInnerPins(TSharedRef<SGraphNodeMaterialConvert> InOwningGraphNodeWidget);
+
+	const TArray<TSharedRef<SConvertInnerPin>>& GetInnerPins() const { return InnerPins; }
+
+	bool IsHoveredOverPrimaryPin() const;
+private:
+
+	// Inner Pins, one for each component
+	TArray<TSharedRef<SConvertInnerPin>> InnerPins;
+};
+
+////////////////////////////////////////////////////////////////////////////////////////////////////////////////////////////////////////////
+// SConvertInnerPin: The inner pins used to route values with the convert node
+////////////////////////////////////////////////////////////////////////////////////////////////////////////////////////////////////////////
+
+/** Class for the inner pins we'll use to route values in SGraphNodeMaterialConvert */
+class SConvertInnerPin : public SCompoundWidget
+{
+public:
+	SLATE_BEGIN_ARGS(SConvertInnerPin) {}
+	SLATE_END_ARGS()
+
+	void Construct(
+		const FArguments& InArgs,
+		TSharedPtr<SGraphNodeMaterialConvert> InOwningNode,
+		TSharedPtr<SMaterialExpressionConvertGraphPin> InOwningPin,
+		bool bInIsInputPin, 
+		int32 InPinIndex, 
+		int32 InComponentIndex
+	);
+
+	//~ Begin SWidget interface
+	virtual int32 OnPaint(
+		const FPaintArgs& Args, 
+		const FGeometry& AllottedGeometry, 
+		const FSlateRect& MyCullingRect, 
+		FSlateWindowElementList& OutDrawElements, 
+		int32 LayerId, 
+		const FWidgetStyle& InWidgetStyle, 
+		bool bParentEnabled
+	) const override;
+	virtual FReply OnMouseButtonDown(const FGeometry& MyGeometry, const FPointerEvent& MouseEvent) override;
+	virtual FReply OnMouseButtonUp(const FGeometry& MyGeometry, const FPointerEvent& MouseEvent) override;
+	virtual FReply OnDragDetected(const FGeometry& MyGeometry, const FPointerEvent& MouseEvent) override;
+	virtual FReply OnDrop(const FGeometry& MyGeometry, const FDragDropEvent& DragDropEvent) override;
+	//~ End SWidget interface
+
+	const FSlateBrush* GetPinBrush() const;
+	void CancelDragDrop();
+	FVector2D GetPinCenterAbsolute() const;
+
+	/** Adds a connection to WeakConnectedPins */
+	void AddConnection(TSharedPtr<SConvertInnerPin> InOtherPin);
+
+	/** Removes a connection from WeakConnectedPins */
+	void RemoveConnection(TSharedPtr<SConvertInnerPin> InOtherPin);
+
+	/** Clears WeakConnectedPins */
+	void RemoveAllConnections();
+
+	/** Calls into the owning node to break connections and refresh the material node */
+	void BreakConnections();
+
+	EVisibility GetDefaultValueVisibility() const;
+	TOptional<float> GetDefaultValue() const;
+	void SetDefaultValue(float InDefaultValue, ETextCommit::Type CommitType);
+
+	EVisibility GetPinNameVisibility() const;
+	FText GetPinName() const;
+
+	bool IsInputPin() const { return bIsInputPin; }
+	int32 GetPinIndex() const { return PinIndex; }
+	int32 GetComponentIndex() const { return ComponentIndex; }
+
+protected:
+	// Input Data
+	TWeakPtr<SGraphNodeMaterialConvert> WeakOwningNode;
+	TWeakPtr<SMaterialExpressionConvertGraphPin> WeakOwningPin;
+	bool bIsInputPin = false;
+	int32 PinIndex = INDEX_NONE;
+	int32 ComponentIndex = INDEX_NONE;
+
+	TSharedPtr<SImage> PinImage;
+	TSet<TWeakPtr<SConvertInnerPin>> WeakConnectedPins;
+	mutable FVector2D CenterAbsolute;
+};
+
+////////////////////////////////////////////////////////////////////////////////////////////////////////////////////////////////////////////
+// SGraphNodeMaterialConvert: Custom Slate widget for UMaterialExpressionConvert
+////////////////////////////////////////////////////////////////////////////////////////////////////////////////////////////////////////////
+
+class GRAPHEDITOR_API SGraphNodeMaterialConvert : public SGraphNodeMaterialBase
+{
+public:
+	SLATE_BEGIN_ARGS(SGraphNodeMaterialConvert){}
+	SLATE_END_ARGS()
+
+	void Construct(const FArguments& InArgs, class UMaterialGraphNode* InNode);
+
+	// Begin SGraphNode interface
+	virtual void CreatePinWidgets() override;
+	virtual TSharedPtr<SGraphPin> CreatePinWidget(UEdGraphPin* Pin) const override;
+	virtual void AddPin(const TSharedRef<SGraphPin>& PinToAdd) override;
+
+	virtual void CreateInputSideAddButton(TSharedPtr<SVerticalBox> InputBox) override;
+	virtual void CreateOutputSideAddButton(TSharedPtr<SVerticalBox> OutputBox) override;
+	virtual EVisibility IsAddPinButtonVisible() const override;
+	// End SGraphNode interface
+
+	FReply OnAddInputPinClicked();
+	FReply OnAddOutputPinClicked();
+
+	TSharedRef<SWidget> CreateAddPinContextMenu(bool bInputPin);
+	void AddNewPin(bool bInputPin, EMaterialExpressionConvertType ConvertType);
+
+	// Begin SWidget Interface
+	virtual int32 OnPaint(
+		const FPaintArgs& Args, 
+		const FGeometry& AllottedGeometry, 
+		const FSlateRect& MyCullingRect, 
+		FSlateWindowElementList& OutDrawElements, 
+		int32 LayerId, 
+		const FWidgetStyle& InWidgetStyle, 
+		bool bParentEnabled
+	) const override;
+	// End SWidget Interface
+
+	void FormConnection(TSharedPtr<SConvertInnerPin> InnerPinA, TSharedPtr<SConvertInnerPin> InnerPinB);
+	void BreakConnections(TSharedPtr<SConvertInnerPin> InnerPin);
+
+	TOptional<float> GetDefaultValue(TSharedPtr<const SConvertInnerPin> InnerPin) const;
+	void SetDefaultValue(TSharedPtr<SConvertInnerPin> InnerPin, const float InDefaultValue);
+
+	void SetCurrentDragDropOp(TSharedPtr<FConvertDragDropOp> InDragDropOp) { CurrentDragDropOp = InDragDropOp; }
+	TSharedPtr<FConvertDragDropOp> GetCurrentDragDropOp() const { return CurrentDragDropOp; }
+
+protected:
+	TSharedPtr<SConvertInnerPin> GetInnerPin(bool bInputPin, const int32 InPinIndex, const int32 InComponentIndex) const;
+	void MakeConnectionCurve(
+		const FGeometry& InAllottedGeometry,
+		FSlateWindowElementList& OutDrawElements,
+		int32& InOutLayerId, 
+		const FVector2D InCurveStart, 
+		const FVector2D InCurveEnd
+	) const;
+
+private:
+	TSharedPtr<FConvertDragDropOp> CurrentDragDropOp;
+};
\ No newline at end of file
diff --git a/Engine/Source/Editor/MaterialEditor/Private/MaterialEditor.cpp b/Engine/Source/Editor/MaterialEditor/Private/MaterialEditor.cpp
--- a/Engine/Source/Editor/MaterialEditor/Private/MaterialEditor.cpp
+++ b/Engine/Source/Editor/MaterialEditor/Private/MaterialEditor.cpp
@@ -2416,6 +2416,23 @@
 		);
 	}
 
+	// Delete Pin
 	{
+		FToolUIAction DeletePinAction;
+		DeletePinAction.ExecuteAction = FToolMenuExecuteAction::CreateSP(this, &FMaterialEditor::OnDeletePin);
+		DeletePinAction.IsActionVisibleDelegate = FToolMenuIsActionButtonVisible::CreateSP(this, &FMaterialEditor::OnCanDeletePin);
+
+		TSharedPtr<FUICommandInfo> DeletePinCommand = FMaterialEditorCommands::Get().DeletePin;
+		InSection.AddMenuEntry(
+			DeletePinCommand->GetCommandName(),
+			DeletePinCommand->GetLabel(),
+			DeletePinCommand->GetDescription(),
+			DeletePinCommand->GetIcon(),
+			DeletePinAction
+		);
+	}
+
+	// Substrate Pin Actions
+	{
 		auto AddSubstrateContextualMenu = [&](TSharedPtr<FUICommandInfo> CreateNodeCommand, ESubstrateNodeForPin SubstrateNodeForPin)
 		{
@@ -5364,6 +5381,49 @@
 	}
 }
 
+bool FMaterialEditor::OnCanResetToDefault(const FToolMenuContext& InMenuContext) const
+{
+	UGraphNodeContextMenuContext* NodeContext = InMenuContext.FindContext<UGraphNodeContextMenuContext>();
+	if (!NodeContext)
+	{
+		return false;
+	}
+	const UEdGraphPin* TargetPin = NodeContext->Pin;
+	const UMaterialGraphNode_Root* RootPinNode = Cast<UMaterialGraphNode_Root>(TargetPin->GetOwningNode());
+	if (RootPinNode != nullptr)
+	{
+		const UMaterialGraph* MaterialGraph = CastChecked<UMaterialGraph>(RootPinNode->GetGraph());
+		const FMaterialInputInfo& MaterialInput = MaterialGraph->MaterialInputs[TargetPin->SourceIndex];
+		const EMaterialProperty PropertyId = MaterialInput.GetProperty();
+		switch (PropertyId)
+		{
+			case MP_BaseColor:				
+			case MP_Opacity:				
+			case MP_Metallic:				
+			case MP_Specular:				
+			case MP_Roughness:				
+			case MP_Anisotropy:				
+			case MP_CustomData0:			
+			case MP_CustomData1:			
+			case MP_AmbientOcclusion:		
+			case MP_Refraction:				
+			case MP_OpacityMask:			
+			case MP_SurfaceThickness:		
+			case MP_Displacement:			
+			case MP_WorldPositionOffset:	
+			case MP_EmissiveColor:			
+			case MP_SubsurfaceColor:		
+			case MP_Normal:					
+			case MP_Tangent:				
+				return true;
+				
+			default:
+				return false;
+		}
+	}
+	return false;
+}
+
 void FMaterialEditor::OnResetToDefault(const FToolMenuContext& InMenuContext) const
 {
 	UGraphNodeContextMenuContext* NodeContext = InMenuContext.FindContext<UGraphNodeContextMenuContext>();
@@ -5480,48 +5540,38 @@
 	}
 }
 
-bool FMaterialEditor::OnCanResetToDefault(const FToolMenuContext& InMenuContext) const
+bool FMaterialEditor::OnCanDeletePin(const FToolMenuContext& InMenuContext) const
 {
-	UGraphNodeContextMenuContext* NodeContext = InMenuContext.FindContext<UGraphNodeContextMenuContext>();
-	if (!NodeContext)
+	if (UGraphNodeContextMenuContext* NodeContext = InMenuContext.FindContext<UGraphNodeContextMenuContext>())
 	{
-		return false;
-	}
-	const UEdGraphPin* TargetPin = NodeContext->Pin;
-	const UMaterialGraphNode_Root* RootPinNode = Cast<UMaterialGraphNode_Root>(TargetPin->GetOwningNode());
-	if (RootPinNode != nullptr)
-	{
-		const UMaterialGraph* MaterialGraph = CastChecked<UMaterialGraph>(RootPinNode->GetGraph());
-		const FMaterialInputInfo& MaterialInput = MaterialGraph->MaterialInputs[TargetPin->SourceIndex];
-		const EMaterialProperty PropertyId = MaterialInput.GetProperty();
-		UMaterialEditorOnlyData* EditorOnlyData = Material->GetEditorOnlyData();
-		switch (PropertyId)
+		if (UEdGraphPin* TargetPin = const_cast<UEdGraphPin*>(NodeContext->Pin))
 		{
-			case MP_BaseColor:				
-			case MP_Opacity:				
-			case MP_Metallic:				
-			case MP_Specular:				
-			case MP_Roughness:				
-			case MP_Anisotropy:				
-			case MP_CustomData0:			
-			case MP_CustomData1:			
-			case MP_AmbientOcclusion:		
-			case MP_Refraction:				
-			case MP_OpacityMask:			
-			case MP_SurfaceThickness:		
-			case MP_Displacement:			
-			case MP_WorldPositionOffset:	
-			case MP_EmissiveColor:			
-			case MP_SubsurfaceColor:		
-			case MP_Normal:					
-			case MP_Tangent:				
-				return true;
-				
-			default:
-				return false;
+			UMaterialGraphNode* MaterialGraphNode = Cast<UMaterialGraphNode>(TargetPin->GetOwningNode());
+			if (MaterialGraphNode && MaterialGraphNode->MaterialExpression)
+			{
+				return MaterialGraphNode->MaterialExpression->CanDeletePin(TargetPin->Direction, TargetPin->SourceIndex);
+			}
 		}
 	}
+
 	return false;
 }
 
+void FMaterialEditor::OnDeletePin(const FToolMenuContext& InMenuContext) const
+{
+	if (UGraphNodeContextMenuContext* NodeContext = InMenuContext.FindContext<UGraphNodeContextMenuContext>())
+	{
+		if (UEdGraphPin* TargetPin = const_cast<UEdGraphPin*>(NodeContext->Pin))
+		{
+			UMaterialGraphNode* MaterialGraphNode = Cast<UMaterialGraphNode>(TargetPin->GetOwningNode());
+			if (MaterialGraphNode && MaterialGraphNode->MaterialExpression)
+			{
+				const FScopedTransaction Transaction(LOCTEXT("DeletePin", "Delete Pin"));
+				MaterialGraphNode->MaterialExpression->Modify();
+				MaterialGraphNode->MaterialExpression->DeletePin(TargetPin->Direction, TargetPin->SourceIndex);
+			}
+		}
+	}
+}
+
 bool FMaterialEditor::OnCanPromoteToParameter(const FToolMenuContext& InMenuContext) const
diff --git a/Engine/Source/Editor/MaterialEditor/Private/MaterialEditor.h b/Engine/Source/Editor/MaterialEditor/Private/MaterialEditor.h
--- a/Engine/Source/Editor/MaterialEditor/Private/MaterialEditor.h
+++ b/Engine/Source/Editor/MaterialEditor/Private/MaterialEditor.h
@@ -850,6 +850,12 @@ private:
 
 	/** Will reset selected pin to it's default value */
 	void OnResetToDefault(const FToolMenuContext& InMenuContext) const;
+
+	/** Used to know if we can delete the selected pin */
+	bool OnCanDeletePin(const FToolMenuContext& InMenuContext) const;
+
+	/** Will delete the selected pin */
+	void OnDeletePin(const FToolMenuContext& InMenuContext) const;
 	
 	enum class ESubstrateNodeForPin : uint8
 	{
diff --git a/Engine/Source/Editor/MaterialEditor/Private/MaterialEditorActions.cpp b/Engine/Source/Editor/MaterialEditor/Private/MaterialEditorActions.cpp
--- a/Engine/Source/Editor/MaterialEditor/Private/MaterialEditorActions.cpp
+++ b/Engine/Source/Editor/MaterialEditor/Private/MaterialEditorActions.cpp
@@ -77,13 +77,18 @@ void FMaterialEditorCommands::RegisterCommands()
 	UI_COMMAND( ForceRefreshPreviews, "Force Refresh Previews", "Forces a refresh of all previews", EUserInterfaceActionType::Button, FInputChord(EKeys::SpaceBar) );
 	UI_COMMAND( CreateComponentMaskNode, "Create ComponentMask Node", "Creates a ComponentMask node at the current cursor position.", EUserInterfaceActionType::Button, FInputChord(EModifierKey::Shift, EKeys::C));
 	UI_COMMAND( FindInMaterial, "Search", "Finds expressions and comments in the current Material", EUserInterfaceActionType::Button, FInputChord(EModifierKey::Control, EKeys::F));
+	
+	// Pin Actions
 	UI_COMMAND( PromoteToParameter, "Promote to Parameter", "Promote selected Pin to parameter of pin type", EUserInterfaceActionType::Button, FInputChord());
+	UI_COMMAND( ResetToDefault, "Reset Pin to Default Value", "Reset pin value to default for that type", EUserInterfaceActionType::Button, FInputChord());
+	UI_COMMAND(DeletePin, "Delete Pin", "Delete Selected Pin", EUserInterfaceActionType::Button, FInputChord());
+
+	// Substrate Pin Actions
 	UI_COMMAND( CreateSlabNode, "Create Slab Node", "Create a Slab Node linked to the pin", EUserInterfaceActionType::Button, FInputChord());
 	UI_COMMAND( CreateHorizontalMixNode, "Create Horizontal Mix Node", "Create a Horizontal Mix Node linked to the pin", EUserInterfaceActionType::Button, FInputChord());
 	UI_COMMAND( CreateVerticalLayerNode, "Create Vertical Layer Node", "Create a Vertical Layer Node linked to the pin", EUserInterfaceActionType::Button, FInputChord());
 	UI_COMMAND( CreateWeightNode, "Create Weight Node", "Create a Weight Node linked to the pin", EUserInterfaceActionType::Button, FInputChord());
 	UI_COMMAND( CreateSelectNode, "Create Select Node", "Create a Select Node linked to the pin", EUserInterfaceActionType::Button, FInputChord());
-	UI_COMMAND( ResetToDefault, "Reset Pin to Default Value", "Reset pin value to default for that type", EUserInterfaceActionType::Button, FInputChord());
 
 	UI_COMMAND(QualityLevel_All, "All", "Sets node preview to show all quality levels.)", EUserInterfaceActionType::RadioButton, FInputChord());
 	UI_COMMAND(QualityLevel_Epic, "Epic", "Sets node preview to Epic quality.", EUserInterfaceActionType::RadioButton, FInputChord());
diff --git a/Engine/Source/Editor/MaterialEditor/Public/MaterialEditorActions.h b/Engine/Source/Editor/MaterialEditor/Public/MaterialEditorActions.h
--- a/Engine/Source/Editor/MaterialEditor/Public/MaterialEditorActions.h
+++ b/Engine/Source/Editor/MaterialEditor/Public/MaterialEditorActions.h
@@ -197,6 +197,9 @@ public:
 	/** Reset pin to default value */
 	TSharedPtr< FUICommandInfo > ResetToDefault;
 
+	/** Delete Pin */
+	TSharedPtr< FUICommandInfo > DeletePin;
+
 	/** Create slab node */
 	TSharedPtr< FUICommandInfo > CreateSlabNode;
 
diff --git a/Engine/Source/Editor/UnrealEd/Private/MaterialGraphNode.cpp b/Engine/Source/Editor/UnrealEd/Private/MaterialGraphNode.cpp
--- a/Engine/Source/Editor/UnrealEd/Private/MaterialGraphNode.cpp
+++ b/Engine/Source/Editor/UnrealEd/Private/MaterialGraphNode.cpp
@@ -655,7 +655,16 @@ void UMaterialGraphNode::PinDefaultValueChanged(UEdGraphPin* Pin)
 
 TSharedPtr<SGraphNode> UMaterialGraphNode::CreateVisualWidget()
 {
-	return SNew(SGraphNodeMaterialBase, this);
+	// First check to see if CreateCustomGraphNodeWidget returns a valid widget
+	TSharedPtr<SGraphNode> OutGraphNode = MaterialExpression->CreateCustomGraphNodeWidget();
+	
+	// If CreateCustomGraphNodeWidget returned nullptr, create default SGraphNodeMaterialBase
+	if (!OutGraphNode)
+	{
+		OutGraphNode = SNew(SGraphNodeMaterialBase, this);
+	}
+
+	return OutGraphNode;
 }
 
 bool UMaterialGraphNode::CanJumpToDefinition() const
diff --git a/Engine/Source/Runtime/Engine/Engine.Build.cs b/Engine/Source/Runtime/Engine/Engine.Build.cs
--- a/Engine/Source/Runtime/Engine/Engine.Build.cs
+++ b/Engine/Source/Runtime/Engine/Engine.Build.cs
@@ -322,6 +322,9 @@ public class Engine : ModuleRules
 			PrivateDependencyModuleNames.Add("LogVisualizer");
 			CircularlyReferencedDependentModules.Add("LogVisualizer");
 
+			PrivateIncludePathModuleNames.Add("GraphEditor");
+			PrivateDependencyModuleNames.Add("GraphEditor");
+
 			if (Target.Platform == UnrealTargetPlatform.Win64)
 			{
 				DynamicallyLoadedModuleNames.AddRange(
diff --git a/Engine/Source/Runtime/Engine/Private/Materials/MaterialExpressions.cpp b/Engine/Source/Runtime/Engine/Private/Materials/MaterialExpressions.cpp
--- a/Engine/Source/Runtime/Engine/Private/Materials/MaterialExpressions.cpp
+++ b/Engine/Source/Runtime/Engine/Private/Materials/MaterialExpressions.cpp
@@ -77,6 +77,7 @@
 #include "Materials/MaterialExpressionConstant3Vector.h"
 #include "Materials/MaterialExpressionConstant4Vector.h"
 #include "Materials/MaterialExpressionConstantBiasScale.h"
+#include "Materials/MaterialExpressionConvert.h"
 #include "Materials/MaterialExpressionCosine.h"
 #include "Materials/MaterialExpressionCrossProduct.h"
 #include "Materials/MaterialExpressionCustom.h"
@@ -300,6 +301,8 @@
 #if WITH_EDITOR
 #include "MaterialGraph/MaterialGraphNode_Comment.h"
 #include "MaterialGraph/MaterialGraphNode_Composite.h"
+#include "MaterialNodes/SGraphNodeMaterialBase.h"
+#include "MaterialNodes/SGraphNodeMaterialConvert.h"
 #include "Framework/Notifications/NotificationManager.h"
 #include "Widgets/Notifications/SNotificationList.h"
 #include "MaterialGraph/MaterialGraphSchema.h"
@@ -23224,4 +23227,282 @@ void UMaterialExpressionIsFirstPerson::GetCaption(TArray<FString>& OutCaptions)
 }
 #endif // WITH_EDITOR
 
+
+///////////////////////////////////////////////////////////////////////////////
+// Convert
+///////////////////////////////////////////////////////////////////////////////
+
+UMaterialExpressionConvert::UMaterialExpressionConvert(const FObjectInitializer& ObjectInitializer)
+: Super(ObjectInitializer)
+{
+#if WITH_EDITORONLY_DATA
+	bHidePreviewWindow = true;
+#endif
+}
+
+#if WITH_EDITOR
+void UMaterialExpressionConvert::PostEditChangeProperty(FPropertyChangedEvent& PropertyChangedEvent)
+{
+	const bool bNeedsNodeUpdate = 
+			PropertyChangedEvent.GetPropertyName() == GET_MEMBER_NAME_CHECKED(ThisClass, ConvertInputs)
+		|| 	PropertyChangedEvent.GetPropertyName() == GET_MEMBER_NAME_CHECKED(ThisClass, ConvertMappings)
+		|| 	PropertyChangedEvent.GetPropertyName() == GET_MEMBER_NAME_CHECKED(ThisClass, ConvertOutputs);
+
+	if (bNeedsNodeUpdate)
+	{
+		RecreateOutputs();
+		if (GraphNode)
+		{
+			GraphNode->ReconstructNode();
+		}
+	}
+
+	Super::PostEditChangeProperty(PropertyChangedEvent);
+}
+
+void UMaterialExpressionConvert::PostEditChangeChainProperty(struct FPropertyChangedChainEvent& PropertyChangedEvent)
+{
+	static FProperty* ConvertInputsProperty = UMaterialExpressionConvert::StaticClass()->FindPropertyByName(
+		GET_MEMBER_NAME_CHECKED(ThisClass, ConvertInputs)
+	);
+
+	static FProperty* ConvertOutputsProperty = UMaterialExpressionConvert::StaticClass()->FindPropertyByName(
+		GET_MEMBER_NAME_CHECKED(ThisClass, ConvertOutputs)
+	);
+
+	const bool bNeedsNodeUpdate = 
+			PropertyChangedEvent.PropertyChain.Contains(ConvertInputsProperty)
+		||	PropertyChangedEvent.PropertyChain.Contains(ConvertOutputsProperty);
+
+	if (bNeedsNodeUpdate)
+	{
+		RecreateOutputs();
+		if (GraphNode)
+		{
+			GraphNode->ReconstructNode();
+		}
+	}
+
+	Super::PostEditChangeChainProperty(PropertyChangedEvent);
+}
+
+int32 UMaterialExpressionConvert::Compile(class FMaterialCompiler* Compiler, int32 OutputIndex)
+{
+	if (!ConvertOutputs.IsValidIndex(OutputIndex))
+	{
+		return CompilerError(Compiler, TEXT("Invalid Output Index"));
+	}
+
+	const FMaterialExpressionConvertOutput& ConvertOutput = ConvertOutputs[OutputIndex];
+	const EMaterialExpressionConvertType OutputType = ConvertOutput.Type;
+	const int32 OutputComponentCount = MaterialExpressionConvertType::GetComponentCount(OutputType);
+
+	int32 Components[4] = { INDEX_NONE, INDEX_NONE, INDEX_NONE, INDEX_NONE };
+
+	TArray<int32> CachedCompiledInputs;
+	CachedCompiledInputs.Init(INDEX_NONE, ConvertInputs.Num());
+
+	// Determine value of any connected outputs
+	for (const FMaterialExpressionConvertMapping& ConvertMapping : ConvertMappings)
+	{
+		const int32 CurrentOutputIndex = ConvertMapping.OutputIndex;
+
+		// We only care about mapping relevant to this output
+		if (CurrentOutputIndex != OutputIndex)
+		{
+			continue;
+		}
+
+		const int32 OutputComponentIndex = ConvertMapping.OutputComponentIndex;
+		if (!IsValidComponentIndex(OutputComponentIndex, OutputType))
+		{
+			return CompilerError(Compiler, TEXT("Invalid Mapping Output Component Index"));
+		}
+
+		const int32 InputIndex = ConvertMapping.InputIndex;
+		if (!ConvertInputs.IsValidIndex(InputIndex))
+		{
+			return CompilerError(Compiler, TEXT("Invalid Mapping Input Index"));
+		}
+
+		FMaterialExpressionConvertInput& ConvertInput = ConvertInputs[InputIndex];
+		const int32 InputComponentIndex = ConvertMapping.InputComponentIndex;
+		if (!IsValidComponentIndex(InputComponentIndex, ConvertInput.Type))
+		{
+			return CompilerError(Compiler, TEXT("Invalid Mapping Input Component Index"));
+		}
+
+		// Search CachedCompiledInputs for previously compiled value. If it isn't found, compile that expression input now
+		int32& CompiledInput = CachedCompiledInputs[InputIndex];
+		if (CompiledInput == INDEX_NONE)
+		{
+			const FLinearColor& DefaultValue = ConvertInput.DefaultValue;
+			CompiledInput = ConvertInput.ExpressionInput.GetTracedInput().Expression 
+				? ConvertInput.ExpressionInput.Compile(Compiler)
+				: Compiler->Constant4(DefaultValue.R, DefaultValue.G, DefaultValue.B, DefaultValue.A);
+		}
+
+		// Mask out the specific value we want
+		const int32 MaskedCompiledInput = Compiler->ComponentMask(
+			CompiledInput, 
+			InputComponentIndex == 0,	// R
+			InputComponentIndex == 1,	// G
+			InputComponentIndex == 2,	// B
+			InputComponentIndex == 3	// A
+		);
+
+		Components[OutputComponentIndex] = MaskedCompiledInput;
+	}
+
+	// For any ComponentCompileResults that are still INDEX_NONE, set to ConvertOutput's default value
+	for (int32 OutputComponentIndex = 0; OutputComponentIndex < OutputComponentCount; ++OutputComponentIndex)
+	{
+		int32& ComponentCompileResult = Components[OutputComponentIndex];
+		// If we don't have a compile result here, default it to a that component's default value
+		if (ComponentCompileResult == INDEX_NONE)
+		{
+			ComponentCompileResult = Compiler->Constant(ConvertOutput.DefaultValue.Component(OutputComponentIndex));
+		}
+	}
+
+	// Init FinalCompileResult to first component result and then chain appends to form final vector
+	int32 FinalCompileResult = Components[0];
+	for (int32 AppendIndex = 1; AppendIndex < OutputComponentCount; ++AppendIndex)
+	{
+		FinalCompileResult = Compiler->AppendVector(
+			FinalCompileResult, 
+			Components[AppendIndex]
+		);
+	}
+	return FinalCompileResult;
+}
+
+void UMaterialExpressionConvert::GetCaption(TArray<FString>& OutCaptions) const 
+{ 
+	OutCaptions.Add(TEXT("Convert")); 
+}
+
+void UMaterialExpressionConvert::GetConnectorToolTip(int32 InputIndex, int32 OutputIndex, TArray<FString>& OutToolTip)
+{
+	if (ConvertInputs.IsValidIndex(InputIndex))
+	{
+		OutToolTip.Add(MaterialExpressionConvertType::ToText(ConvertInputs[InputIndex].Type).ToString());
+	}
+	else if (ConvertOutputs.IsValidIndex(OutputIndex))
+	{
+		OutToolTip.Add(MaterialExpressionConvertType::ToText(ConvertOutputs[OutputIndex].Type).ToString());
+	}
+}
+
+int32 UMaterialExpressionConvert::CountInputs() const
+{
+	return ConvertInputs.Num();
+}
+
+uint32 UMaterialExpressionConvert::GetInputType(int32 InputIndex)
+{
+	if (ConvertInputs.IsValidIndex(InputIndex))
+	{
+		return MaterialExpressionConvertType::ToMaterialValueType(ConvertInputs[InputIndex].Type);
+	}
+
+	checkNoEntry();
+	return MCT_Float;
+}
+
+FExpressionInput* UMaterialExpressionConvert::GetInput(int32 InputIndex)
+{
+	if (ConvertInputs.IsValidIndex(InputIndex))
+	{
+		return &ConvertInputs[InputIndex].ExpressionInput;
+	}
+
+	return nullptr;
+}
+
+uint32 UMaterialExpressionConvert::GetOutputType(int32 OutputIndex)
+{
+	if (ConvertOutputs.IsValidIndex(OutputIndex))
+	{
+		return MaterialExpressionConvertType::ToMaterialValueType(ConvertOutputs[OutputIndex].Type);
+	}
+
+	checkNoEntry();
+	return MCT_Float;
+}
+
+TArray<FExpressionOutput>& UMaterialExpressionConvert::GetOutputs()
+{
+	RecreateOutputs();
+	return Outputs;
+}
+
+bool UMaterialExpressionConvert::CanDeletePin(EEdGraphPinDirection PinDirection, int32 PinIndex) const
+{
+	return true;
+}
+
+void UMaterialExpressionConvert::DeletePin(EEdGraphPinDirection PinDirection, int32 PinIndex)
+{
+	switch (PinDirection)
+	{
+		case EEdGraphPinDirection::EGPD_Input:
+		{
+			if (ConvertInputs.IsValidIndex(PinIndex))
+			{
+				ConvertInputs.RemoveAt(PinIndex);
+			}
+			break;
+		}
+		case EEdGraphPinDirection::EGPD_Output:
+		{
+			if (ConvertOutputs.IsValidIndex(PinIndex))
+			{
+				ConvertOutputs.RemoveAt(PinIndex);
+				RecreateOutputs();
+			}
+			break;
+		}
+		default: break;
+	}
+
+	// Clean up any stale ConvertMappings
+	for (auto ConvertMappingIter = ConvertMappings.CreateIterator(); ConvertMappingIter; ++ConvertMappingIter)
+	{
+		if (!ConvertInputs.IsValidIndex(ConvertMappingIter->InputIndex) ||
+			!ConvertOutputs.IsValidIndex(ConvertMappingIter->OutputIndex))
+		{
+			ConvertMappingIter.RemoveCurrent();
+			continue;
+		}
+	}
+
+	if (GraphNode)
+	{
+		GraphNode->ReconstructNode();
+	}
+}
+
+TSharedPtr<class SGraphNodeMaterialBase> UMaterialExpressionConvert::CreateCustomGraphNodeWidget()
+{
+	if (UMaterialGraphNode* MaterialGraphNode = Cast<UMaterialGraphNode>(GraphNode))
+	{
+		return SNew(SGraphNodeMaterialConvert, MaterialGraphNode);
+	}
+
+	return nullptr;
+}
+
+void UMaterialExpressionConvert::RecreateOutputs()
+{
+	Outputs.Reset();
+	Outputs.AddDefaulted(ConvertOutputs.Num());
+}
+
+bool UMaterialExpressionConvert::IsValidComponentIndex(int32 InComponentIndex, EMaterialExpressionConvertType InType) const
+{
+	return InComponentIndex < MaterialExpressionConvertType::GetComponentCount(InType);
+}
+#endif // WITH_EDITOR
+
 #undef LOCTEXT_NAMESPACE
diff --git a/Engine/Source/Runtime/Engine/Private/Materials/MaterialExpressionsIR.cpp b/Engine/Source/Runtime/Engine/Private/Materials/MaterialExpressionsIR.cpp
--- a/Engine/Source/Runtime/Engine/Private/Materials/MaterialExpressionsIR.cpp
+++ b/Engine/Source/Runtime/Engine/Private/Materials/MaterialExpressionsIR.cpp
@@ -45,6 +45,7 @@
 #include "Materials/MaterialExpressionConstant3Vector.h"
 #include "Materials/MaterialExpressionConstant4Vector.h"
 #include "Materials/MaterialExpressionConstantBiasScale.h"
+#include "Materials/MaterialExpressionConvert.h"
 #include "Materials/MaterialExpressionCosine.h"
 #include "Materials/MaterialExpressionCrossProduct.h"
 #include "Materials/MaterialExpressionCurveAtlasRowParameter.h"
@@ -1005,4 +1006,100 @@ void UMaterialExpressionSmoothStep::Build(MIR::FEmitter& Emitter)
 	BuildTernaryArithmeticOperator(Emitter, MIR::TO_Smoothstep, &Min, ConstMin, &Max, ConstMax, &Value, ConstValue);
 }
 
+void UMaterialExpressionConvert::Build(MIR::FEmitter& Emitter)
+{
+	// Cache Compiled Inputs as they may be referred in multiple ConvertMappings/Outputs
+	TArray<MIR::FValue*> CachedCompiledInputs;
+	CachedCompiledInputs.Init(nullptr, ConvertInputs.Num());
+
+	for (int32 OutputIndex = 0; OutputIndex < ConvertOutputs.Num(); ++OutputIndex)
+	{
+		const FMaterialExpressionConvertOutput& ConvertOutput = ConvertOutputs[OutputIndex];
+		const EMaterialExpressionConvertType OutputType = ConvertOutput.Type;
+		const int32 OutputNumComponents = MaterialExpressionConvertType::GetComponentCount(OutputType);
+
+		MIR::FValue* OutComponents[4] = { nullptr, nullptr, nullptr, nullptr };
+
+		// Determine value of any connected outputs
+		for (const FMaterialExpressionConvertMapping& ConvertMapping : ConvertMappings)
+		{
+			const int32 CurrentOutputIndex = ConvertMapping.OutputIndex;
+
+			// We only care about mappings relevant to this output
+			if (CurrentOutputIndex != OutputIndex)
+			{
+				continue;
+			}
+
+			const int32 OutputComponentIndex = ConvertMapping.OutputComponentIndex;
+			if (!IsValidComponentIndex(OutputComponentIndex, OutputType))
+			{
+				Emitter.Errorf(TEXT("Convert mapping's output component (%d) is invalid."), OutputComponentIndex);
+				continue;
+			}
+
+			const int32 InputIndex = ConvertMapping.InputIndex;
+			if (!ConvertInputs.IsValidIndex(InputIndex))
+			{
+				Emitter.Errorf(TEXT("Convert mapping's input (%d) is invalid."), InputIndex);
+				continue;
+			}
+
+			FMaterialExpressionConvertInput& ConvertInput = ConvertInputs[InputIndex];
+			const int32 InputComponentIndex = ConvertMapping.InputComponentIndex;
+			if (!IsValidComponentIndex(InputComponentIndex, ConvertInput.Type))
+			{
+				Emitter.Errorf(TEXT("Convert mapping's input component (%d) is invalid."), InputComponentIndex);
+				continue;
+			}
+
+			// Search CachedCompiledInputs for previously compiled value. If it isn't found, compile that expression input now
+			MIR::FValue*& CompiledInput = CachedCompiledInputs[InputIndex];
+			if (CompiledInput == nullptr)
+			{
+				// Need to compile input if its connected or compile a default-value constant if not
+				CompiledInput = Emitter.TryInput(&ConvertInput.ExpressionInput)
+					? Emitter.Input(&ConvertInput.ExpressionInput)
+					: Emitter.EmitConstantFloat4(ConvertInput.DefaultValue);
+			}
+
+			// Choose the specific component/subscript we want
+			OutComponents[OutputComponentIndex] = Emitter.EmitSubscript(CompiledInput, InputComponentIndex);
+		}
+
+		// For any ComponentCompileResults that are still INDEX_NONE, set to ConvertOutput's default value
+		for (int32 OutputComponentIndex = 0; OutputComponentIndex < OutputNumComponents; ++OutputComponentIndex)
+		{
+			MIR::FValue*& ComponentCompileResult = OutComponents[OutputComponentIndex];
+			// If we don't have a compile result here, default it to a that component's default value
+			if (ComponentCompileResult == nullptr)
+			{
+				ComponentCompileResult = Emitter.EmitConstantFloatScalar(ConvertOutput.DefaultValue.Component(OutputComponentIndex));
+			}
+		}
+
+		MIR::FValue* OutValue = nullptr;
+		switch (OutputNumComponents)
+		{
+			case 1:
+				OutValue = OutComponents[0];
+				break;
+			case 2:
+				OutValue = Emitter.EmitVector2(OutComponents[0], OutComponents[1]);
+				break;
+			case 3:
+				OutValue = Emitter.EmitVector3(OutComponents[0], OutComponents[1], OutComponents[2]);
+				break;
+			case 4:
+				OutValue = Emitter.EmitVector4(OutComponents[0], OutComponents[1], OutComponents[2], OutComponents[3]);
+				break;
+			default:
+				Emitter.Errorf(TEXT("Convert node has an invalid component count of %d"), OutputNumComponents);
+				continue;
+		}
+
+		Emitter.Output(OutputIndex, OutValue);
+	}
+}
+
 #endif // WITH_EDITOR
diff --git a/Engine/Source/Runtime/Engine/Public/Materials/MaterialExpression.h b/Engine/Source/Runtime/Engine/Public/Materials/MaterialExpression.h
--- a/Engine/Source/Runtime/Engine/Public/Materials/MaterialExpression.h
+++ b/Engine/Source/Runtime/Engine/Public/Materials/MaterialExpression.h
@@ -25,6 +25,7 @@ class UMaterialExpression;
 class UMaterialExpressionComment;
 class FMaterialExpressionKey;
 
+enum EEdGraphPinDirection : int;
 enum class EMaterialParameterType : uint8;
 
 namespace MIR
@@ -281,6 +282,8 @@ class UMaterialExpression : public UObject
 	   that use certain material expressions, and can be used to limit shader generation to only materials that have
 	   use that material expression*/
 	ENGINE_API virtual void GetShaderTags(TArray<FName>& ShaderTagsOut) {}
+
+	ENGINE_API virtual TSharedPtr<class SGraphNodeMaterialBase> CreateCustomGraphNodeWidget() { return nullptr; }
 #endif // WITH_EDITOR
 
 	/**
@@ -548,6 +551,21 @@ class UMaterialExpression : public UObject
 	/** This overload accepts the set of visited expressions to avoid visiting them again when checking at once whether multiple expressions contain a loop. */
 	ENGINE_API bool ContainsInputLoop(TSet<UMaterialExpression*>& VisitedExpressions, const bool bStopOnFunctionCall = true);
 
+	/** 
+	 * Checks if a pin of specified PinDirection and PinIndex can be deleted 
+	 * @param PinDirection direction of the pin
+	 * @param PinIndex index of the pin 
+	 * @return returns true if the pin can be deleted
+	 */
+	virtual bool CanDeletePin(EEdGraphPinDirection PinDirection, int32 PinIndex) const { return false; }
+
+	/**
+	 * Deletes the pin of specified PinDirection and PinIndex. Should only be used if CanDeletePin returns true for that pin
+	 * @param PinDirection direction of the pin
+	 * @param PinIndex index of the pin
+	 */
+	virtual void DeletePin(EEdGraphPinDirection PinDirection, int32 PinIndex) {}
+
 protected:
 	/** Caches the list of fixed expression inputs this expression has. */
 	TArray<FExpressionInput*> CachedInputs;
diff --git a/Engine/Source/Runtime/Engine/Public/Materials/MaterialExpressionConvert.h b/Engine/Source/Runtime/Engine/Public/Materials/MaterialExpressionConvert.h
new file mode 100644
--- /dev/null
+++ b/Engine/Source/Runtime/Engine/Public/Materials/MaterialExpressionConvert.h
@@ -0,0 +1,166 @@
+// Copyright Epic Games, Inc. All Rights Reserved.
+
+#pragma once
+
+#include "Materials/MaterialExpression.h"
+
+#include "UObject/ObjectMacros.h"
+#include "MaterialExpressionIO.h"
+#include "MaterialValueType.h"
+
+#include "MaterialExpressionConvert.generated.h"
+
+class UEdGraphPin;
+
+#define LOCTEXT_NAMESPACE "MaterialExpressionConvert"
+
+UENUM()
+enum class EMaterialExpressionConvertType : uint8
+{
+	Float1,
+	Float2,
+	Float3,
+	Float4,
+};
+
+namespace MaterialExpressionConvertType
+{
+	static FORCEINLINE int32 GetComponentCount(const EMaterialExpressionConvertType InConvertType)
+	{
+		switch (InConvertType)
+		{
+			case EMaterialExpressionConvertType::Float1: return 1;
+			case EMaterialExpressionConvertType::Float2: return 2;
+			case EMaterialExpressionConvertType::Float3: return 3;
+			case EMaterialExpressionConvertType::Float4: return 4;
+			default: checkNoEntry(); return INDEX_NONE;
+		}
+	}
+
+	static FORCEINLINE EMaterialValueType ToMaterialValueType(const EMaterialExpressionConvertType InConvertType)
+	{
+		switch (InConvertType)
+		{
+			case EMaterialExpressionConvertType::Float1: return MCT_Float1;
+			case EMaterialExpressionConvertType::Float2: return MCT_Float2;
+			case EMaterialExpressionConvertType::Float3: return MCT_Float3;
+			case EMaterialExpressionConvertType::Float4: return MCT_Float4;
+			default: checkNoEntry(); return MCT_Float;
+		}
+	}
+
+	static FORCEINLINE FText ToText(const EMaterialExpressionConvertType InConvertType)
+	{
+		switch (InConvertType)
+		{
+			case EMaterialExpressionConvertType::Float1: return LOCTEXT("Float1", "Float1");
+			case EMaterialExpressionConvertType::Float2: return LOCTEXT("Float2", "Float2");
+			case EMaterialExpressionConvertType::Float3: return LOCTEXT("Float3", "Float3");
+			case EMaterialExpressionConvertType::Float4: return LOCTEXT("Float4", "Float4");
+			default: checkNoEntry(); return INVTEXT("");
+		}
+	}
+}
+
+USTRUCT()
+struct FMaterialExpressionConvertInput
+{
+	GENERATED_BODY()
+
+	UPROPERTY(EditAnywhere, Category = "MaterialExpressionConvert", meta = (RequiredInput = "false"))
+	FExpressionInput ExpressionInput;
+
+	UPROPERTY(EditAnywhere, Category = "MaterialExpressionConvert")
+	EMaterialExpressionConvertType Type = EMaterialExpressionConvertType::Float1;
+
+	/** Default Value used when this input has no incoming connection */
+	UPROPERTY(EditAnywhere, Category = "MaterialExpressionConvert")
+	FLinearColor DefaultValue = FLinearColor::Black;
+};
+
+USTRUCT()
+struct FMaterialExpressionConvertOutput
+{
+	GENERATED_BODY()
+
+	UPROPERTY(EditAnywhere, Category = "MaterialExpressionConvert")
+	EMaterialExpressionConvertType Type = EMaterialExpressionConvertType::Float1;
+
+	/** Default Value used when this output has no connection */
+	UPROPERTY(EditAnywhere, Category = "MaterialExpressionConvert")
+	FLinearColor DefaultValue = FLinearColor::Black;
+};
+
+USTRUCT()
+struct FMaterialExpressionConvertMapping
+{
+	GENERATED_BODY()
+
+	/** Which Input to map from */
+	UPROPERTY(EditAnywhere, Category = "MaterialExpressionConvert")
+	int32 InputIndex = 0;
+
+	/** Which Input Component to map from */
+	UPROPERTY(EditAnywhere, Category = "MaterialExpressionConvert")
+	int32 InputComponentIndex = 0;
+
+	/** Which Output to map to */
+	UPROPERTY(EditAnywhere, Category = "MaterialExpressionConvert")
+	int32 OutputIndex = 0;
+
+	/** Which Output Component to map to */
+	UPROPERTY(EditAnywhere, Category = "MaterialExpressionConvert")
+	int32 OutputComponentIndex = 0;
+};
+
+UCLASS(MinimalAPI)
+class UMaterialExpressionConvert : public UMaterialExpression
+{
+	GENERATED_BODY()
+public:
+	UPROPERTY(EditAnywhere, Category = "MaterialExpressionConvert")
+	TArray<FMaterialExpressionConvertInput> ConvertInputs;
+
+	UPROPERTY(EditAnywhere, Category = "MaterialExpressionConvert")
+	TArray<FMaterialExpressionConvertOutput> ConvertOutputs;
+
+	/** Describes how data flows from input components to output components */
+	UPROPERTY()
+	TArray<FMaterialExpressionConvertMapping> ConvertMappings;
+
+	UMaterialExpressionConvert(const FObjectInitializer& ObjectInitializer);
+
+	//~Begin UObject Interface
+	#if WITH_EDITOR
+	ENGINE_API virtual void PostEditChangeProperty(FPropertyChangedEvent& PropertyChangedEvent) override;
+	ENGINE_API virtual void PostEditChangeChainProperty(struct FPropertyChangedChainEvent& PropertyChangedEvent) override;
+	#endif // WITH_EDITOR
+	//~End UObject Interface
+
+	//~ Begin UMaterialExpression Interface
+	#if WITH_EDITOR
+	virtual void Build(MIR::FEmitter& Emitter) override;
+	virtual int32 Compile(class FMaterialCompiler* Compiler, int32 OutputIndex) override;
+	virtual void GetCaption(TArray<FString>& OutCaptions) const override;
+	virtual void GetConnectorToolTip(int32 InputIndex, int32 OutputIndex, TArray<FString>& OutToolTip) override;
+
+	virtual int32 CountInputs() const override;
+	virtual uint32 GetInputType(int32 InputIndex) override;
+	virtual FExpressionInput* GetInput(int32 InputIndex) override;
+
+	virtual uint32 GetOutputType(int32 OutputIndex) override;
+	virtual TArray<FExpressionOutput>& GetOutputs() override;
+
+	virtual bool CanDeletePin(EEdGraphPinDirection PinDirection, int32 PinIndex) const override;
+	virtual void DeletePin(EEdGraphPinDirection PinDirection, int32 PinIndex) override;
+	virtual TSharedPtr<class SGraphNodeMaterialBase> CreateCustomGraphNodeWidget() override;
+	#endif //WITH_EDITOR
+	//~ End UMaterialExpression Interface
+private:
+	#if WITH_EDITOR
+	void RecreateOutputs();
+	bool IsValidComponentIndex(int32 InComponentIndex, EMaterialExpressionConvertType InType) const;
+	#endif //WITH_EDITOR
+};
+
+#undef LOCTEXT_NAMESPACE
\ No newline at end of file
