Add visual cue on Substrate Slab node about SSS type and active features.

diff --git a/Engine/Source/Editor/GraphEditor/Private/MaterialNodes/SGraphSubstrateMaterial.cpp b/Engine/Source/Editor/GraphEditor/Private/MaterialNodes/SGraphSubstrateMaterial.cpp
--- a/Engine/Source/Editor/GraphEditor/Private/MaterialNodes/SGraphSubstrateMaterial.cpp
+++ b/Engine/Source/Editor/GraphEditor/Private/MaterialNodes/SGraphSubstrateMaterial.cpp
@@ -84,7 +84,8 @@ static const TSharedRef<SWidget> InternalProcessOperator(
 	const FSubstrateOperator& Op, 
 	ESubstrateWidgetOutputType OutputType,
 	const TArray<FGuid>& InGuid,
-	EStyleColor OverrideColor)
+	EStyleColor OverrideColor, 
+	TSharedPtr<SWidget>& OutFeatureWidget)
 {
 	const bool bIsCurrent = OutputType == ESubstrateWidgetOutputType::Node ? InGuid.Find(Op.MaterialExpressionGuid) != INDEX_NONE : false;
 	const EStyleColor Color0 = bIsCurrent ? GetSubstrateWidgetColor0() : OverrideColor;
@@ -94,7 +95,7 @@ static const TSharedRef<SWidget> InternalProcessOperator(
 		case SUBSTRATE_OPERATOR_WEIGHT:
 		{
 			const EStyleColor Color = bIsCurrent ? EStyleColor::AccentGreen : OverrideColor;
-			return InternalProcessOperator(CompilationOutput, CompilationOutput.Operators[Op.LeftIndex], OutputType, InGuid, Color);
+			return InternalProcessOperator(CompilationOutput, CompilationOutput.Operators[Op.LeftIndex], OutputType, InGuid, Color, OutFeatureWidget);
 		}
 		break;
 		case SUBSTRATE_OPERATOR_VERTICAL:
@@ -106,7 +107,7 @@ static const TSharedRef<SWidget> InternalProcessOperator(
 				.HAlign(HAlign_Fill)
 				.Padding(0.0f, 0.0f, 1.0f, 1.0f)
 				[
-					InternalProcessOperator(CompilationOutput, CompilationOutput.Operators[Op.LeftIndex], OutputType, InGuid, Color0)
+					InternalProcessOperator(CompilationOutput, CompilationOutput.Operators[Op.LeftIndex], OutputType, InGuid, Color0, OutFeatureWidget)
 				]
 				+ SVerticalBox::Slot()
 				.AutoHeight()
@@ -114,7 +115,7 @@ static const TSharedRef<SWidget> InternalProcessOperator(
 				.HAlign(HAlign_Fill)
 				.Padding(0.0f, 0.0f, 1.0f, 1.0f)
 				[
-					InternalProcessOperator(CompilationOutput, CompilationOutput.Operators[Op.RightIndex], OutputType, InGuid, Color1)
+					InternalProcessOperator(CompilationOutput, CompilationOutput.Operators[Op.RightIndex], OutputType, InGuid, Color1, OutFeatureWidget)
 				];
 			return VerticalOperator->AsShared();
 		}
@@ -129,7 +130,7 @@ static const TSharedRef<SWidget> InternalProcessOperator(
 				.HAlign(HAlign_Fill)
 				.Padding(0.0f, 0.0f, 1.0f, 1.0f)
 				[
-					InternalProcessOperator(CompilationOutput, CompilationOutput.Operators[Op.LeftIndex], OutputType, InGuid, Color0)
+					InternalProcessOperator(CompilationOutput, CompilationOutput.Operators[Op.LeftIndex], OutputType, InGuid, Color0, OutFeatureWidget)
 				]
 				+SHorizontalBox::Slot()
 				.AutoWidth()
@@ -137,7 +138,7 @@ static const TSharedRef<SWidget> InternalProcessOperator(
 				.HAlign(HAlign_Fill)
 				.Padding(0.0f, 0.0f, 1.0f, 1.0f)
 				[
-					InternalProcessOperator(CompilationOutput, CompilationOutput.Operators[Op.RightIndex], OutputType, InGuid, Color1)
+					InternalProcessOperator(CompilationOutput, CompilationOutput.Operators[Op.RightIndex], OutputType, InGuid, Color1, OutFeatureWidget)
 				];
 			return HorizontalOperator->AsShared();
 		}
@@ -151,7 +152,7 @@ static const TSharedRef<SWidget> InternalProcessOperator(
 				.HAlign(HAlign_Fill)
 				.Padding(0.0f, 0.0f, 1.0f, 1.0f)
 				[
-					InternalProcessOperator(CompilationOutput, CompilationOutput.Operators[Op.LeftIndex], OutputType, InGuid, Color0)
+					InternalProcessOperator(CompilationOutput, CompilationOutput.Operators[Op.LeftIndex], OutputType, InGuid, Color0, OutFeatureWidget)
 				]
 				+SHorizontalBox::Slot()
 				.AutoWidth()
@@ -159,7 +160,7 @@ static const TSharedRef<SWidget> InternalProcessOperator(
 				.HAlign(HAlign_Fill)
 				.Padding(0.0f, 0.0f, 1.0f, 1.0f)
 				[
-					InternalProcessOperator(CompilationOutput, CompilationOutput.Operators[Op.RightIndex], OutputType, InGuid, Color1)
+					InternalProcessOperator(CompilationOutput, CompilationOutput.Operators[Op.RightIndex], OutputType, InGuid, Color1, OutFeatureWidget)
 				];
 			return HorizontalOperator->AsShared();
 		}
@@ -218,6 +219,107 @@ static const TSharedRef<SWidget> InternalProcessOperator(
 				ToolTip += TEXT("Hai means the BSDF features hair is used.\n");
 			}
 
+			// Features and SubSurface type
+			if (bIsCurrent && OutputType == ESubstrateWidgetOutputType::Node)
+			{
+				EMaterialSubSurfaceType SubSurfaceType = EMaterialSubSurfaceType(Op.SubSurfaceType);
+				if (!Op.bIsBottom)
+				{
+					SubSurfaceType = EMaterialSubSurfaceType::MSS_SimpleVolume;
+				}
+				else if (!Op.Has(ESubstrateBsdfFeature_SSS))
+				{
+					SubSurfaceType = EMaterialSubSurfaceType::MSS_None;
+				}
+
+				const uint32 DisplayedFeatureMask = 
+					ESubstrateBsdfFeature::ESubstrateBsdfFeature_EdgeColor |
+					ESubstrateBsdfFeature::ESubstrateBsdfFeature_Fuzz |
+					ESubstrateBsdfFeature::ESubstrateBsdfFeature_SecondRoughnessOrSimpleClearCoat |
+					ESubstrateBsdfFeature::ESubstrateBsdfFeature_Anisotropy |
+					ESubstrateBsdfFeature::ESubstrateBsdfFeature_Glint |
+					ESubstrateBsdfFeature::ESubstrateBsdfFeature_SpecularProfile;
+
+				const bool bHasFeatures = (Op.BSDFFeatures & DisplayedFeatureMask) != 0;
+
+				if (bHasFeatures || SubSurfaceType != EMaterialSubSurfaceType::MSS_None)
+				{
+					auto Features = SNew(SVerticalBox);
+
+					if (SubSurfaceType != EMaterialSubSurfaceType::MSS_None)
+					{
+						const FString BehaviorHint = SubSurfaceType != EMaterialSubSurfaceType::MSS_SimpleVolume ? TEXT("Opaque") : TEXT("Translucent");
+						const FText SSSTypeName =StaticEnum<EMaterialSubSurfaceType>()->GetDisplayNameTextByValue(int64(SubSurfaceType));
+						const FText Message = FText::Format(FTextFormat::FromString(TEXT("SSS {0} ({1})")), SSSTypeName, FText::FromString(BehaviorHint));
+
+						Features->AddSlot()
+						.AutoHeight()
+						.VAlign(VAlign_Fill)
+						.HAlign(HAlign_Fill)
+						.Padding(0.0f, 0.0f, 1.0f, 1.0f)
+						[
+							SNew(SErrorText)
+							.ErrorText(Message)
+							.BackgroundColor(FSlateColor(FLinearColor(0.65,0.25,0.05)))
+						];
+					}
+
+					if (bHasFeatures)
+					{
+						auto HorizontalOperator = SNew(SHorizontalBox);
+						auto AddFeatureSlot = [&HorizontalOperator](FText Message)
+						{
+							HorizontalOperator->AddSlot()
+							.AutoWidth()
+							.VAlign(VAlign_Fill)
+							.HAlign(HAlign_Fill)
+							.Padding(0.0f, 0.0f, 1.0f, 1.0f)
+							[
+								SNew(SErrorText)
+								.ErrorText(Message)
+								.BackgroundColor(FSlateColor(FLinearColor(0.16, 0.015, 0.24)))
+							];
+						};
+
+						if (Op.Has(ESubstrateBsdfFeature::ESubstrateBsdfFeature_EdgeColor))
+						{
+							AddFeatureSlot(FText::FromString(TEXT("F90")));
+						}
+						if (Op.Has(ESubstrateBsdfFeature::ESubstrateBsdfFeature_Fuzz))
+						{
+							AddFeatureSlot(FText::FromString(TEXT("Fuzz")));
+						}
+						if (Op.Has(ESubstrateBsdfFeature::ESubstrateBsdfFeature_SecondRoughnessOrSimpleClearCoat))
+						{
+							AddFeatureSlot(FText::FromString(TEXT("Dual Spec.")));
+						}
+						if (Op.Has(ESubstrateBsdfFeature::ESubstrateBsdfFeature_Anisotropy))
+						{
+							AddFeatureSlot(FText::FromString(TEXT("Aniso")));
+						}
+						if (Op.Has(ESubstrateBsdfFeature::ESubstrateBsdfFeature_Glint))
+						{
+							AddFeatureSlot(FText::FromString(TEXT("Glints")));
+						}
+						if (Op.Has(ESubstrateBsdfFeature::ESubstrateBsdfFeature_SpecularProfile))
+						{
+							AddFeatureSlot(FText::FromString(TEXT("Spec. Profile")));
+						}
+
+						Features->AddSlot()
+						.AutoHeight()
+						.VAlign(VAlign_Fill)
+						.HAlign(HAlign_Fill)
+						.Padding(0.0f, 0.0f, 1.0f, 1.0f)
+						[
+							HorizontalOperator
+						];
+					}
+
+					OutFeatureWidget = Features.ToSharedPtr();
+				}
+			}
+
 			const EStyleColor Color = OverrideColor != EStyleColor::MAX ? OverrideColor : (bIsCurrent ? EStyleColor::AccentGreen : EStyleColor::AccentGray);
 			const FSlateColor SlateColor = OutputType == ESubstrateWidgetOutputType::Node ? FSlateColor(Color) : FSlateColor(FLinearColor(0.16, 0.015, 0.24));
 			auto BSDF = SNew(SErrorText)
@@ -236,6 +338,47 @@ static const TSharedRef<SWidget> InternalProcessOperator(
 	return TreeOperatorError->AsShared();
 }
 
+static const TSharedRef<SWidget> InternalProcessOperator(
+	const FSubstrateMaterialCompilationOutput& CompilationOutput,
+	const FSubstrateOperator& Op,
+	ESubstrateWidgetOutputType OutputType,
+	const TArray<FGuid>& InGuid,
+	EStyleColor OverrideColor)
+{
+	TSharedPtr<SWidget> FeatureWidget;
+	TSharedRef<SWidget> Topology = InternalProcessOperator(CompilationOutput, Op, OutputType, InGuid, OverrideColor, FeatureWidget);
+	if (FeatureWidget)
+	{
+		return SNew(SVerticalBox)
+		+SVerticalBox::Slot()
+		.AutoHeight()
+		.VAlign(VAlign_Fill)
+		.HAlign(HAlign_Fill)
+		.Padding(0.0f, 0.0f, 1.0f, 1.0f)
+		[
+			FeatureWidget.ToSharedRef()
+		]
+		+SVerticalBox::Slot()
+		.MinHeight(15)
+		.FillHeight(1.f)
+		.Padding(0.0f, 0.0f, 1.0f, 1.0f)
+		[
+			SNullWidget::NullWidget
+		]
+		+SVerticalBox::Slot()
+		.AutoHeight()
+		.VAlign(VAlign_Center)
+		.HAlign(HAlign_Center)
+		.Padding(0.0f, 0.0f, 1.0f, 1.0f)
+		[
+			Topology
+		];
+	}
+	else
+	{
+		return Topology;
+	}
+}
 const TSharedRef<SWidget> FSubstrateWidget::ProcessOperator(const FSubstrateMaterialCompilationOutput& CompilationOutput)
 {
 	return InternalProcessOperator(CompilationOutput, CompilationOutput.Operators[CompilationOutput.RootOperatorIndex], ESubstrateWidgetOutputType::DetailPanel, TArray<FGuid>(), EStyleColor::MAX);
diff --git a/Engine/Source/Runtime/Core/Private/UObject/FortniteMainBranchObjectVersion.cpp b/Engine/Source/Runtime/Core/Private/UObject/FortniteMainBranchObjectVersion.cpp
--- a/Engine/Source/Runtime/Core/Private/UObject/FortniteMainBranchObjectVersion.cpp
+++ b/Engine/Source/Runtime/Core/Private/UObject/FortniteMainBranchObjectVersion.cpp
@@ -14,6 +14,6 @@ TMap<FGuid, FGuid> FFortniteMainBranchObjectVersion::GetSystemGuids()
 	SystemGuids.Add(DevGuids.Niagara_LatestScriptCompileVersion, FGuid("670A2EBDEE1A4794BB77BBFF331EA7EF"));
 	SystemGuids.Add(DevGuids.SkeletalMeshDerivedDataVersion, FGuid("5E93192FCF3647EEB7F9D4FFF1750050"));
 	SystemGuids.Add(DevGuids.STATICMESH_DERIVEDDATA_VER, FGuid("3ABF32149F9D4D838750368AB95573BE"));
-	SystemGuids.Add(DevGuids.MaterialTranslationDDCVersion, FGuid("0195246791767C5080E1D4C3ED151B18"));
+	SystemGuids.Add(DevGuids.MaterialTranslationDDCVersion, FGuid("BB6334B9845A4A9E8E5DE2F7768127DB"));
 	return SystemGuids;
 }
diff --git a/Engine/Source/Runtime/Engine/Private/Materials/MaterialExpressionSubstrate.cpp b/Engine/Source/Runtime/Engine/Private/Materials/MaterialExpressionSubstrate.cpp
--- a/Engine/Source/Runtime/Engine/Private/Materials/MaterialExpressionSubstrate.cpp
+++ b/Engine/Source/Runtime/Engine/Private/Materials/MaterialExpressionSubstrate.cpp
@@ -21,6 +21,13 @@
 ///////////////////////////////////////////////////////////////////////////////
 // Substrate
 
+EMaterialSubSurfaceType SubstrateMergeSubSurfaceType(EMaterialSubSurfaceType A, EMaterialSubSurfaceType B)
+{
+	// Merge two type of SSS type, by using the most complex behavior, i.e., No < Wrap < Diffusion < Diffusion Profile
+	// This code needs to be in sync with SubstrateMergeSSSType() in Substrate.ush
+	return EMaterialSubSurfaceType(FMath::Max(uint32(A), uint32(B)));
+}
+
 #if WITH_EDITOR
 static int32 SubstrateBlendNormal(class FMaterialCompiler* Compiler, int32 NormalCodeChunk0, int32 NormalCodeChunk1, int32 MixCodeChunk)
 {
@@ -601,6 +608,7 @@ FSubstrateOperator* UMaterialExpressionSubstrateShadingModels::SubstrateGenerate
 			FSubstrateOperator& SlabOperator = Compiler->SubstrateCompilationRegisterOperator(SUBSTRATE_OPERATOR_BSDF_LEGACY, Compiler->SubstrateTreeStackGetPathUniqueId(), this, Parent, Compiler->SubstrateTreeStackGetParentPathUniqueId());
 			SlabOperator.BSDFType = SUBSTRATE_BSDF_TYPE_SLAB;
 			SlabOperator.BSDFFeatures = In | (Anisotropy.IsConnected() ? ESubstrateBsdfFeature_Anisotropy : ESubstrateBsdfFeature_None);
+			SlabOperator.SubSurfaceType = uint8(EMaterialSubSurfaceType::MSS_None);
 			SlabOperator.bBSDFWritesEmissive = EmissiveColor.IsConnected();
 			SlabOperator.ThicknessIndex = ThicknessIndex;
 
@@ -1255,6 +1263,7 @@ FSubstrateOperator* UMaterialExpressionSubstrateSlabBSDF::SubstrateGenerateMater
 	if (HasSpecularProfile())	{ SubstrateOperator.BSDFFeatures |= ESubstrateBsdfFeature_SpecularProfile; }
 	if (HasSSSProfile())		{ SubstrateOperator.BSDFFeatures |= ESubstrateBsdfFeature_SecondRoughnessOrSimpleClearCoat; } // If a Slab has a subsurface profile, it will have haziness in order to to support the dual-specular lobe from the profile.
 
+	SubstrateOperator.SubSurfaceType = SubSurfaceType;
 	SubstrateOperator.ThicknessIndex = Compiler->SubstrateThicknessStackGetThicknessIndex();
 	SubstrateOperator.bBSDFWritesEmissive = EmissiveColor.IsConnected();
 	return &SubstrateOperator;
@@ -1313,6 +1322,16 @@ bool UMaterialExpressionSubstrateSlabBSDF::HasSpecularProfile() const
 	return SpecularProfile != nullptr && Substrate::IsSpecularProfileEnabled();
 }
 
+void UMaterialExpressionSubstrateSlabBSDF::PostEditChangeProperty(FPropertyChangedEvent& PropertyChangedEvent)
+{
+	if (GraphNode && PropertyChangedEvent.Property != nullptr)
+	{
+		GraphNode->ReconstructNode();
+		GraphNode->Modify();
+	}
+	Super::PostEditChangeProperty(PropertyChangedEvent);
+}
+
 #endif // WITH_EDITOR
 
 
diff --git a/Engine/Source/Runtime/Engine/Public/Materials/MaterialExpressionSubstrate.h b/Engine/Source/Runtime/Engine/Public/Materials/MaterialExpressionSubstrate.h
--- a/Engine/Source/Runtime/Engine/Public/Materials/MaterialExpressionSubstrate.h
+++ b/Engine/Source/Runtime/Engine/Public/Materials/MaterialExpressionSubstrate.h
@@ -24,6 +24,8 @@ enum EMaterialSubSurfaceType : int
 	MSS_MAX
 };
 
+EMaterialSubSurfaceType SubstrateMergeSubSurfaceType(EMaterialSubSurfaceType A, EMaterialSubSurfaceType B);
+
 ///////////////////////////////////////////////////////////////////////////////
 // Functions
 
@@ -375,6 +377,7 @@ class UMaterialExpressionSubstrateSlabBSDF : public UMaterialExpressionSubstrate
 	virtual FName GetInputName(int32 InputIndex) const override;
 	virtual void GetConnectorToolTip(int32 InputIndex, int32 OutputIndex, TArray<FString>& OutToolTip) override;
 	virtual void GetExpressionToolTip(TArray<FString>& OutToolTip) override;
+	virtual void PostEditChangeProperty(FPropertyChangedEvent& PropertyChangedEvent) override;
 
 	bool HasEdgeColor() const;
 	bool HasFuzz() const;
diff --git a/Engine/Source/Runtime/Engine/Public/Rendering/SubstrateMaterialShared.cpp b/Engine/Source/Runtime/Engine/Public/Rendering/SubstrateMaterialShared.cpp
--- a/Engine/Source/Runtime/Engine/Public/Rendering/SubstrateMaterialShared.cpp
+++ b/Engine/Source/Runtime/Engine/Public/Rendering/SubstrateMaterialShared.cpp
@@ -1,7 +1,7 @@
 // Copyright Epic Games, Inc. All Rights Reserved.
 
 #include "SubstrateMaterialShared.h"
-
+#include "Materials/MaterialExpressionSubstrate.h"
 
 IMPLEMENT_TYPE_LAYOUT(FSubstrateMaterialCompilationOutput);
 
@@ -49,6 +49,7 @@ FSubstrateOperator::FSubstrateOperator()
 	BSDFRegisteredSharedLocalBasis = FSubstrateRegisteredSharedLocalBasis();
 	BSDFFeatures = 0;
 	SubUsage = SUBSTRATE_OPERATOR_SUBUSAGE_NONE;
+	SubSurfaceType = uint8(EMaterialSubSurfaceType::MSS_None);
 
 	bBSDFWritesEmissive = false;
 	bBSDFWritesAmbientOcclusion = false;
@@ -67,6 +68,7 @@ void FSubstrateOperator::CombineFlagsForParameterBlending(FSubstrateOperator& A,
 {
 #if WITH_EDITOR
 	BSDFFeatures = A.BSDFFeatures | B.BSDFFeatures;
+	SubSurfaceType = uint32(SubstrateMergeSubSurfaceType(EMaterialSubSurfaceType(A.SubSurfaceType), EMaterialSubSurfaceType(B.SubSurfaceType)));
 #endif
 }
 
@@ -74,6 +76,7 @@ void FSubstrateOperator::CopyFlagsForParameterBlending(FSubstrateOperator& A)
 {
 #if WITH_EDITOR
 	BSDFFeatures = A.BSDFFeatures;
+	SubSurfaceType = A.SubSurfaceType;
 #endif
 }
 
diff --git a/Engine/Source/Runtime/Engine/Public/Rendering/SubstrateMaterialShared.h b/Engine/Source/Runtime/Engine/Public/Rendering/SubstrateMaterialShared.h
--- a/Engine/Source/Runtime/Engine/Public/Rendering/SubstrateMaterialShared.h
+++ b/Engine/Source/Runtime/Engine/Public/Rendering/SubstrateMaterialShared.h
@@ -84,6 +84,7 @@ public:
 	LAYOUT_FIELD_EDITORONLY(FSubstrateRegisteredSharedLocalBasis, BSDFRegisteredSharedLocalBasis);
 	LAYOUT_FIELD_EDITORONLY(uint32,BSDFFeatures);
 	LAYOUT_FIELD_EDITORONLY(uint8, SubUsage);		// Sometimes, Unlit or Weight operators are used to transport data for other meaning (e.g. Light Function or ConvertToDecal)
+	LAYOUT_FIELD_EDITORONLY(uint8, SubSurfaceType);
 
 	LAYOUT_FIELD_EDITORONLY(uint8, bBSDFWritesEmissive);
 	LAYOUT_FIELD_EDITORONLY(uint8, bBSDFWritesAmbientOcclusion);
