Substrate - improve look for the "convert to substrate" action on the root node.

#rnx
#tests pc editor dx12

diff --git a/Engine/Source/Editor/GraphEditor/Private/SGraphEditorImpl.cpp b/Engine/Source/Editor/GraphEditor/Private/SGraphEditorImpl.cpp
--- a/Engine/Source/Editor/GraphEditor/Private/SGraphEditorImpl.cpp
+++ b/Engine/Source/Editor/GraphEditor/Private/SGraphEditorImpl.cpp
@@ -1182,8 +1182,8 @@ void SGraphEditorImpl::RegisterContextMenu(const UEdGraphSchema* Schema, FToolMe
 							const UMaterialGraphNode_Root* RootNode = Cast<UMaterialGraphNode_Root>(ContextObject->Node);
 
 							TSharedRef<SButton> SubstrateConvertButton = SNew(SButton)
-								.Text(LOCTEXT("SubstrateConvertButtonLabel", "Convert To Substrate"))
-								.ToolTipText(LOCTEXT("SubstrateConvertButtonToolTip", "Convert the root node to a simple Substrate graph for further authoring."))
+								.Text(LOCTEXT("SubstrateConvertButtonLabel", "\x26A0 Convert To Substrate \x26A0"))
+								.ToolTipText(LOCTEXT("SubstrateConvertButtonToolTip", "Convert the root node to a Substrate graph. This material will then use Substrate more advanced features and topology processes.\nTo go back to the previous Non-Substrate material, you can use CTRL+Z or you will need to manually relink the graph."))
 								.ButtonStyle(FAppStyle::Get(), "Button")
 								.OnClicked_Lambda([RootNode]()
 								{
