Unshelved from pending changelist '40605339':

Show the panel (node list) context menu when the user ends a drag connection operation on a mistmatched pin for blueprint graphs

[RN] When connecting incompatible pins in a blueprint graph the context menu is now shown, allowing the user to search for a compatible node. This new behavior can be disabled via the ShowPanelContextMenuForIncompatibleConnections Blueprint Editor Setting

diff --git a/Engine/Source/Editor/BlueprintGraph/Classes/EdGraphSchema_K2.h b/Engine/Source/Editor/BlueprintGraph/Classes/EdGraphSchema_K2.h
--- a/Engine/Source/Editor/BlueprintGraph/Classes/EdGraphSchema_K2.h
+++ b/Engine/Source/Editor/BlueprintGraph/Classes/EdGraphSchema_K2.h
@@ -1254,6 +1254,8 @@ private:
 	 * @return true if there are out parameters, else false
 	 */
 	bool DoesFunctionHaveOutParameters( const UFunction* Function ) const;
+	
+	bool ShouldShowPanelContextMenuForIncompatibleConnectionsImpl() const override;
 
 	static const UScriptStruct* VectorStruct;
 	static const UScriptStruct* Vector3fStruct;
diff --git a/Engine/Source/Editor/BlueprintGraph/Private/BlueprintEditorSettings.cpp b/Engine/Source/Editor/BlueprintGraph/Private/BlueprintEditorSettings.cpp
--- a/Engine/Source/Editor/BlueprintGraph/Private/BlueprintEditorSettings.cpp
+++ b/Engine/Source/Editor/BlueprintGraph/Private/BlueprintEditorSettings.cpp
@@ -50,6 +50,7 @@ UBlueprintEditorSettings::UBlueprintEditorSettings(const FObjectInitializer& Obj
 	, bNavigateToNativeFunctionsFromCallNodes(true)
 	, bDoubleClickNavigatesToParent(true)
 	, bEnableTypePromotion(true)
+	, bShowPanelContextMenuForIncompatibleConnections(true)
 	, TypePromotionPinDenyList { UEdGraphSchema_K2::PC_String, UEdGraphSchema_K2::PC_Text }
 	, BreakpointReloadMethod(EBlueprintBreakpointReloadMethod::RestoreAll)
 	, bEnablePinValueInspectionTooltips(true)
diff --git a/Engine/Source/Editor/BlueprintGraph/Private/EdGraphSchema_K2.cpp b/Engine/Source/Editor/BlueprintGraph/Private/EdGraphSchema_K2.cpp
--- a/Engine/Source/Editor/BlueprintGraph/Private/EdGraphSchema_K2.cpp
+++ b/Engine/Source/Editor/BlueprintGraph/Private/EdGraphSchema_K2.cpp
@@ -689,6 +689,13 @@ bool UEdGraphSchema_K2::DoesFunctionHaveOutParameters( const UFunction* Function
 	return false;
 }
 
+bool UEdGraphSchema_K2::ShouldShowPanelContextMenuForIncompatibleConnectionsImpl() const
+{
+	// Disabling this behavior for child types as it is experimental and
+	// there are no tests that would exercise this codepath for them:
+	return GetClass() == UEdGraphSchema_K2::StaticClass() && GetDefault<UBlueprintEditorSettings>()->bShowPanelContextMenuForIncompatibleConnections;
+}
+
 bool UEdGraphSchema_K2::CanFunctionBeUsedInGraph(const UClass* InClass, const UFunction* InFunction, const UEdGraph* InDestGraph, uint32 InAllowedFunctionTypes, bool bInCalledForEach, FText* OutReason) const
 {
 	if (CanUserKismetCallFunction(InFunction))
diff --git a/Engine/Source/Editor/BlueprintGraph/Public/BlueprintEditorSettings.h b/Engine/Source/Editor/BlueprintGraph/Public/BlueprintEditorSettings.h
--- a/Engine/Source/Editor/BlueprintGraph/Public/BlueprintEditorSettings.h
+++ b/Engine/Source/Editor/BlueprintGraph/Public/BlueprintEditorSettings.h
@@ -181,6 +181,10 @@ public:
 	UPROPERTY(config, EditAnywhere, Category = Workflow)
 	bool bEnableTypePromotion;
 
+	/** When enabled the the list of compatible nodes will be shown when a connection is attempted between incompatible pins */
+	UPROPERTY(config, EditAnywhere, Category = Workflow)
+	bool bShowPanelContextMenuForIncompatibleConnections;
+
 	/** If a pin type is within this list, then it will never be marked as a possible promotable function. */
 	UPROPERTY(config, EditAnywhere, Category = Workflow, meta=(EditCondition="bEnableTypePromotion"))
 	TSet<FName> TypePromotionPinDenyList;
diff --git a/Engine/Source/Editor/GraphEditor/Private/BlueprintConnectionDrawingPolicy.cpp b/Engine/Source/Editor/GraphEditor/Private/BlueprintConnectionDrawingPolicy.cpp
--- a/Engine/Source/Editor/GraphEditor/Private/BlueprintConnectionDrawingPolicy.cpp
+++ b/Engine/Source/Editor/GraphEditor/Private/BlueprintConnectionDrawingPolicy.cpp
@@ -638,7 +638,8 @@ void FKismetConnectionDrawingPolicy::SetIncompatiblePinDrawState(const TSharedPt
 		if (CheckPin != StartPin)
 		{
 			const FPinConnectionResponse Response = EdGraphSchema->CanCreateConnection(StartPinObj, CheckPin->GetPinObj());
-			if (Response.Response == CONNECT_RESPONSE_DISALLOW) //-V1051
+			if (Response.Response == CONNECT_RESPONSE_DISALLOW //-V1051
+				&& !EdGraphSchema->ShouldShowPanelContextMenuForIncompatibleConnections()) 
 			{
 				CheckPin->SetPinColorModifier(DisallowColor);
 			}
diff --git a/Engine/Source/Editor/GraphEditor/Private/DragConnection.cpp b/Engine/Source/Editor/GraphEditor/Private/DragConnection.cpp
--- a/Engine/Source/Editor/GraphEditor/Private/DragConnection.cpp
+++ b/Engine/Source/Editor/GraphEditor/Private/DragConnection.cpp
@@ -5,6 +5,7 @@
 
 #include "Containers/EnumAsByte.h"
 #include "Containers/Set.h"
+#include "EdGraphSchema_K2.h"
 #include "EdGraph/EdGraph.h"
 #include "EdGraph/EdGraphNode.h"
 #include "EdGraph/EdGraphPin.h"
@@ -79,13 +80,28 @@ void FDragConnection::HoverTargetChanged()
 
 				// Determine what the schema thinks
 				FPinConnectionResponse Response;
+				const UEdGraphSchema* Schema = GraphObj->GetSchema();
 				switch (DragMode)
 				{
 				case EDragMode::CreateConnection:
-					Response = GraphObj->GetSchema()->CanCreateConnection( StartingPinObj, TargetPinObj );
+					Response = Schema->CanCreateConnection( StartingPinObj, TargetPinObj );
+					if( Response.Response == CONNECT_RESPONSE_DISALLOW && 
+						Schema->ShouldShowPanelContextMenuForIncompatibleConnections())
+					{
+						// override the response, these pins aren't compatible but in this context
+						// we can show the panel context menu listing compatible nodes:
+						Response = FPinConnectionResponse(
+							CONNECT_RESPONSE_MAKE_WITH_CONVERSION_NODE, 
+							FString::Printf(
+								TEXT("Search for Nodes matching %s (%s) and %s (%s)"), 
+								*UEdGraphSchema_K2::TypeToText(StartingPinObj->PinType).ToString(), 
+								LexToString(StartingPinObj->Direction), 
+								*UEdGraphSchema_K2::TypeToText(TargetPinObj->PinType).ToString(), 
+								LexToString(TargetPinObj->Direction)));
+					}
 					break;
 				case EDragMode::RelinkConnection:
-					Response = GraphObj->GetSchema()->CanRelinkConnectionToPin(StartingPinObj, TargetPinObj);
+					Response = Schema->CanRelinkConnectionToPin(StartingPinObj, TargetPinObj);
 					break;
 				}
 
@@ -280,6 +296,7 @@ FReply FDragConnection::DroppedOnPin(FVector2D ScreenPosition, FVector2D GraphPo
 {
 	bool bError = false;
 	TSet<FEdGraphNodeHandle> NodeList;
+	const UEdGraphSchema* Schema = nullptr;
 
 	switch (DragMode)
 	{
@@ -307,7 +324,8 @@ FReply FDragConnection::DroppedOnPin(FVector2D ScreenPosition, FVector2D GraphPo
 				const UEdGraph* MyGraphObj = PinA.GetGraph();
 
 				// the pin may change during the creation of the link
-				if (MyGraphObj->GetSchema()->TryCreateConnection(PinA.GetPin(), PinB.GetPin()))
+				Schema = MyGraphObj->GetSchema();
+				if (Schema->TryCreateConnection(PinA.GetPin(), PinB.GetPin()))
 				{
 					if (PinA.GetPin() && !PinA.GetPin()->IsPendingKill())
 					{
@@ -345,8 +363,9 @@ FReply FDragConnection::DroppedOnPin(FVector2D ScreenPosition, FVector2D GraphPo
 			{
 				const FScopedTransaction Transaction(NSLOCTEXT("UnrealEd", "GraphEd_RelinkConnection", "Relink Connection"));
 
+				Schema = Graph->GetSchema();
 				const TArray<UEdGraphNode*> SelectedGraphNodes = GraphPanel->GetSelectedGraphNodes();
-				bError = !Graph->GetSchema()->TryRelinkConnectionTarget(SourceGraphPinHandle.GetPin(), TargetGraphPin, HoveredGraphPinHandle.GetPin(), SelectedGraphNodes);				if (!bError)
+				bError = !Schema->TryRelinkConnectionTarget(SourceGraphPinHandle.GetPin(), TargetGraphPin, HoveredGraphPinHandle.GetPin(), SelectedGraphNodes);				if (!bError)
 				{
 					// Send all nodes that received a new pin connection a notification
 					TArray<FEdGraphPinHandle> ModifiedGraphPinHandles = { SourceGraphPinHandle, TargetGraphPinHandle, HoveredGraphPinHandle };
@@ -374,6 +393,28 @@ FReply FDragConnection::DroppedOnPin(FVector2D ScreenPosition, FVector2D GraphPo
 		}
 	}
 
+	if( Schema && Schema->ShouldShowPanelContextMenuForIncompatibleConnections() && 
+		NodeList.IsEmpty())
+	{
+		// no connections made, lets show a list of compatible nodes if any:	
+		TArray<UEdGraphPin*> PinObjects;
+		ValidateGraphPinList(/*out*/ PinObjects); // this gathers source pins for the drag operation
+		if(UEdGraphPin* CurrentHoveredPin = GetHoveredPin())
+		{
+			PinObjects.Add(CurrentHoveredPin);
+		}
+
+		if(PinObjects.Num() > 0 )
+		{
+			TSharedPtr<SWidget> WidgetToFocus = GraphPanel->SummonContextMenu(
+				ScreenPosition, GraphPosition, nullptr, nullptr, PinObjects);		
+			if (WidgetToFocus.IsValid())
+			{
+				return FReply::Handled().SetUserFocus(WidgetToFocus.ToSharedRef(), EFocusCause::SetDirectly);
+			}
+		}
+	}
+
 	if (bError)
 	{
 		return FReply::Unhandled();
diff --git a/Engine/Source/Runtime/Engine/Classes/EdGraph/EdGraphSchema.h b/Engine/Source/Runtime/Engine/Classes/EdGraph/EdGraphSchema.h
--- a/Engine/Source/Runtime/Engine/Classes/EdGraph/EdGraphSchema.h
+++ b/Engine/Source/Runtime/Engine/Classes/EdGraph/EdGraphSchema.h
@@ -844,6 +844,16 @@ class UEdGraphSchema : public UObject
 	* @return	True if the promotion and connection were successful; False if the connection failed.
 	*/
 	ENGINE_API virtual bool CreatePromotedConnection(UEdGraphPin* A, UEdGraphPin* B) const;
+	
+#if WITH_EDITORONLY_DATA
+	/**
+	 * @return	True if the schema would like for failed connections to show the node panel menu so the user can pick an appropriate node
+	 *
+	 * Implementaion notes: Consider returning true if your node context menu handles muliple context pins for node indexing and creation.
+	 * When the user attempts to connect two incompatible pins this will allow them to find a node which fullfills the requirements.
+	 */
+	ENGINE_API bool ShouldShowPanelContextMenuForIncompatibleConnections() const;
+#endif
 
 	/**
 	 * Determine if the supplied pin default values would be valid.
@@ -1375,6 +1385,8 @@ class UEdGraphSchema : public UObject
 #endif
 	
 #if WITH_EDITORONLY_DATA
+private:
+	virtual bool ShouldShowPanelContextMenuForIncompatibleConnectionsImpl() const { return false; }
 protected:
 	/** Build an array containing all search types, return the index of the first non-localized entry. */
 	ENGINE_API int32 CollectSearchTextWeightInfo(const FEdGraphSchemaAction& InCurrentAction, const FGraphSchemaSearchWeightModifiers& InWeightModifiers,
diff --git a/Engine/Source/Runtime/Engine/Private/EdGraph/EdGraphSchema.cpp b/Engine/Source/Runtime/Engine/Private/EdGraph/EdGraphSchema.cpp
--- a/Engine/Source/Runtime/Engine/Private/EdGraph/EdGraphSchema.cpp
+++ b/Engine/Source/Runtime/Engine/Private/EdGraph/EdGraphSchema.cpp
@@ -543,6 +543,13 @@ bool UEdGraphSchema::CreatePromotedConnection(UEdGraphPin* PinA, UEdGraphPin* Pi
 	return false;
 }
 
+#if WITH_EDITORONLY_DATA
+bool UEdGraphSchema::ShouldShowPanelContextMenuForIncompatibleConnections() const
+{
+	return ShouldShowPanelContextMenuForIncompatibleConnectionsImpl();
+}
+#endif
+
 void UEdGraphSchema::TrySetDefaultValue(UEdGraphPin& Pin, const FString& NewDefaultValue, bool bMarkAsModified) const
 {
 	Pin.DefaultValue = NewDefaultValue;
