[Truncation Warnings] Update FConnectionDrawingPolicy functions to FVector2f types

diff --git a/Engine/Plugins/Editor/AssetManagerEditor/Source/AssetManagerEditor/Private/ReferenceViewer/ReferenceViewerSchema.cpp b/Engine/Plugins/Editor/AssetManagerEditor/Source/AssetManagerEditor/Private/ReferenceViewer/ReferenceViewerSchema.cpp
--- a/Engine/Plugins/Editor/AssetManagerEditor/Source/AssetManagerEditor/Private/ReferenceViewer/ReferenceViewerSchema.cpp
+++ b/Engine/Plugins/Editor/AssetManagerEditor/Source/AssetManagerEditor/Private/ReferenceViewer/ReferenceViewerSchema.cpp
@@ -115,10 +115,10 @@ public:
 	{
 	}
 
-	virtual FVector2D ComputeSplineTangent(const FVector2D& Start, const FVector2D& End) const override
+	virtual FVector2f ComputeSplineTangent(const FVector2f& Start, const FVector2f& End) const override
 	{
 		const int32 Tension = FMath::Abs<int32>(Start.X - End.X);
-		return Tension * FVector2D(1.0f, 0);
+		return Tension * FVector2f(1.0f, 0);
 	}
 
 	virtual void DetermineWiringStyle(UEdGraphPin* OutputPin, UEdGraphPin* InputPin, /*inout*/ FConnectionParams& Params) override
diff --git a/Engine/Plugins/Experimental/Animation/RigMapper/Source/RigMapperEditor/Private/RigMapperGraph/RigMapperDefinitionEditorGraphSchema.h b/Engine/Plugins/Experimental/Animation/RigMapper/Source/RigMapperEditor/Private/RigMapperGraph/RigMapperDefinitionEditorGraphSchema.h
--- a/Engine/Plugins/Experimental/Animation/RigMapper/Source/RigMapperEditor/Private/RigMapperGraph/RigMapperDefinitionEditorGraphSchema.h
+++ b/Engine/Plugins/Experimental/Animation/RigMapper/Source/RigMapperEditor/Private/RigMapperGraph/RigMapperDefinitionEditorGraphSchema.h
@@ -28,10 +28,10 @@ public:
 			ArrowRadius = FVector2D(0.0f, 0.0f);
 		}
 
-		virtual FVector2D ComputeSplineTangent(const FVector2D& Start, const FVector2D& End) const override
+		virtual FVector2f ComputeSplineTangent(const FVector2f& Start, const FVector2f& End) const override
 		{
 			const int32 Tension = FMath::Abs<int32>(Start.X - End.X);
-			return Tension * FVector2D(1.0f, 0);
+			return Tension * FVector2f(1.0f, 0);
 		}
 	};
 	
diff --git a/Engine/Plugins/Experimental/CommonConversation/Source/CommonConversationGraph/Private/ConversationGraphConnectionDrawingPolicy.cpp b/Engine/Plugins/Experimental/CommonConversation/Source/CommonConversationGraph/Private/ConversationGraphConnectionDrawingPolicy.cpp
--- a/Engine/Plugins/Experimental/CommonConversation/Source/CommonConversationGraph/Private/ConversationGraphConnectionDrawingPolicy.cpp
+++ b/Engine/Plugins/Experimental/CommonConversation/Source/CommonConversationGraph/Private/ConversationGraphConnectionDrawingPolicy.cpp
@@ -83,7 +83,7 @@ void FConversationGraphConnectionDrawingPolicy::Draw(TMap<TSharedRef<SWidget>, F
 	FConnectionDrawingPolicy::Draw(InPinGeometries, ArrangedNodes);
 }
 
-void FConversationGraphConnectionDrawingPolicy::DrawPreviewConnector(const FGeometry& PinGeometry, const FVector2D& StartPoint, const FVector2D& EndPoint, UEdGraphPin* Pin)
+void FConversationGraphConnectionDrawingPolicy::DrawPreviewConnector(const FGeometry& PinGeometry, const FVector2f& StartPoint, const FVector2f& EndPoint, UEdGraphPin* Pin)
 {
 	FConnectionParams Params;
 	DetermineWiringStyle(Pin, nullptr, /*inout*/ Params);
@@ -98,35 +98,35 @@ void FConversationGraphConnectionDrawingPolicy::DrawPreviewConnector(const FGeom
 	}
 }
 
-void FConversationGraphConnectionDrawingPolicy::DrawSplineWithArrow(const FVector2D& StartAnchorPoint, const FVector2D& EndAnchorPoint, const FConnectionParams& Params)
+void FConversationGraphConnectionDrawingPolicy::DrawSplineWithArrow(const FVector2f& StartAnchorPoint, const FVector2f& EndAnchorPoint, const FConnectionParams& Params)
 {
 	// bUserFlag1 indicates that we need to reverse the direction of connection (used by debugger)
-	const FVector2D& P0 = Params.bUserFlag1 ? EndAnchorPoint : StartAnchorPoint;
-	const FVector2D& P1 = Params.bUserFlag1 ? StartAnchorPoint : EndAnchorPoint;
+	const FVector2f& P0 = Params.bUserFlag1 ? EndAnchorPoint : StartAnchorPoint;
+	const FVector2f& P1 = Params.bUserFlag1 ? StartAnchorPoint : EndAnchorPoint;
 
 	Internal_DrawLineWithArrow(P0, P1, Params);
 }
 
-void FConversationGraphConnectionDrawingPolicy::Internal_DrawLineWithArrow(const FVector2D& StartAnchorPoint, const FVector2D& EndAnchorPoint, const FConnectionParams& Params)
+void FConversationGraphConnectionDrawingPolicy::Internal_DrawLineWithArrow(const FVector2f& StartAnchorPoint, const FVector2f& EndAnchorPoint, const FConnectionParams& Params)
 {
 	//@TODO: Should this be scaled by zoom factor?
 	const float LineSeparationAmount = 4.5f;
 
-	const FVector2D DeltaPos = EndAnchorPoint - StartAnchorPoint;
-	const FVector2D UnitDelta = DeltaPos.GetSafeNormal();
-	const FVector2D Normal = FVector2D(DeltaPos.Y, -DeltaPos.X).GetSafeNormal();
+	const FVector2f DeltaPos = EndAnchorPoint - StartAnchorPoint;
+	const FVector2f UnitDelta = DeltaPos.GetSafeNormal();
+	const FVector2f Normal = FVector2f(DeltaPos.Y, -DeltaPos.X).GetSafeNormal();
 
 	// Come up with the final start/end points
-	const FVector2D DirectionBias = Normal * LineSeparationAmount;
-	const FVector2D LengthBias = ArrowRadius.X * UnitDelta;
-	const FVector2D StartPoint = StartAnchorPoint + DirectionBias + LengthBias;
-	const FVector2D EndPoint = EndAnchorPoint + DirectionBias - LengthBias;
+	const FVector2f DirectionBias = Normal * LineSeparationAmount;
+	const FVector2f LengthBias = ArrowRadius.X * UnitDelta;
+	const FVector2f StartPoint = StartAnchorPoint + DirectionBias + LengthBias;
+	const FVector2f EndPoint = EndAnchorPoint + DirectionBias - LengthBias;
 
 	// Draw a line/spline
 	DrawConnection(WireLayerID, StartPoint, EndPoint, Params);
 
 	// Draw the arrow
-	const FVector2D ArrowDrawPos = EndPoint - ArrowRadius;
+	const FVector2f ArrowDrawPos = EndPoint - ArrowRadius;
 	const float AngleInRadians = FMath::Atan2(DeltaPos.Y, DeltaPos.X);
 
 	FSlateDrawElement::MakeRotatedBox(
@@ -136,7 +136,7 @@ void FConversationGraphConnectionDrawingPolicy::Internal_DrawLineWithArrow(const
 		ArrowImage,
 		ESlateDrawEffect::None,
 		AngleInRadians,
-		TOptional<FVector2D>(),
+		TOptional<FVector2f>(),
 		FSlateDrawElement::RelativeToElement,
 		Params.WireColor
 		);
@@ -145,21 +145,21 @@ void FConversationGraphConnectionDrawingPolicy::Internal_DrawLineWithArrow(const
 void FConversationGraphConnectionDrawingPolicy::DrawSplineWithArrow(const FGeometry& StartGeom, const FGeometry& EndGeom, const FConnectionParams& Params)
 {
 	// Get a reasonable seed point (halfway between the boxes)
-	const FVector2D StartCenter = FGeometryHelper::CenterOf(StartGeom);
-	const FVector2D EndCenter = FGeometryHelper::CenterOf(EndGeom);
-	const FVector2D SeedPoint = (StartCenter + EndCenter) * 0.5f;
+	const FVector2f StartCenter = FGeometryHelper::CenterOf(StartGeom);
+	const FVector2f EndCenter = FGeometryHelper::CenterOf(EndGeom);
+	const FVector2f SeedPoint = (StartCenter + EndCenter) * 0.5f;
 
 	// Find the (approximate) closest points between the two boxes
-	const FVector2D StartAnchorPoint = FGeometryHelper::FindClosestPointOnGeom(StartGeom, SeedPoint);
-	const FVector2D EndAnchorPoint = FGeometryHelper::FindClosestPointOnGeom(EndGeom, SeedPoint);
+	const FVector2f StartAnchorPoint = FGeometryHelper::FindClosestPointOnGeom(StartGeom, SeedPoint);
+	const FVector2f EndAnchorPoint = FGeometryHelper::FindClosestPointOnGeom(EndGeom, SeedPoint);
 
 	DrawSplineWithArrow(StartAnchorPoint, EndAnchorPoint, Params);
 }
 
-FVector2D FConversationGraphConnectionDrawingPolicy::ComputeSplineTangent(const FVector2D& Start, const FVector2D& End) const
+FVector2f FConversationGraphConnectionDrawingPolicy::ComputeSplineTangent(const FVector2f& Start, const FVector2f& End) const
 {
-	const FVector2D Delta = End - Start;
-	const FVector2D NormDelta = Delta.GetSafeNormal();
+	const FVector2f Delta = End - Start;
+	const FVector2f NormDelta = Delta.GetSafeNormal();
 
 	return NormDelta;
 }
diff --git a/Engine/Plugins/Experimental/CommonConversation/Source/CommonConversationGraph/Public/ConversationGraphConnectionDrawingPolicy.h b/Engine/Plugins/Experimental/CommonConversation/Source/CommonConversationGraph/Public/ConversationGraphConnectionDrawingPolicy.h
--- a/Engine/Plugins/Experimental/CommonConversation/Source/CommonConversationGraph/Public/ConversationGraphConnectionDrawingPolicy.h
+++ b/Engine/Plugins/Experimental/CommonConversation/Source/CommonConversationGraph/Public/ConversationGraphConnectionDrawingPolicy.h
@@ -20,11 +20,11 @@ public:
 	virtual void DetermineWiringStyle(UEdGraphPin* OutputPin, UEdGraphPin* InputPin, /*inout*/ FConnectionParams& Params) override;
 	virtual void Draw(TMap<TSharedRef<SWidget>, FArrangedWidget>& PinGeometries, FArrangedChildren& ArrangedNodes) override;
 	virtual void DrawSplineWithArrow(const FGeometry& StartGeom, const FGeometry& EndGeom, const FConnectionParams& Params) override;
-	virtual void DrawSplineWithArrow(const FVector2D& StartPoint, const FVector2D& EndPoint, const FConnectionParams& Params) override;
-	virtual void DrawPreviewConnector(const FGeometry& PinGeometry, const FVector2D& StartPoint, const FVector2D& EndPoint, UEdGraphPin* Pin) override;
-	virtual FVector2D ComputeSplineTangent(const FVector2D& Start, const FVector2D& End) const override;
+	virtual void DrawSplineWithArrow(const FVector2f& StartPoint, const FVector2f& EndPoint, const FConnectionParams& Params) override;
+	virtual void DrawPreviewConnector(const FGeometry& PinGeometry, const FVector2f& StartPoint, const FVector2f& EndPoint, UEdGraphPin* Pin) override;
+	virtual FVector2f ComputeSplineTangent(const FVector2f& Start, const FVector2f& End) const override;
 	// End of FConnectionDrawingPolicy interface
 
 protected:
-	void Internal_DrawLineWithArrow(const FVector2D& StartAnchorPoint, const FVector2D& EndAnchorPoint, const FConnectionParams& Params);
+	void Internal_DrawLineWithArrow(const FVector2f& StartAnchorPoint, const FVector2f& EndAnchorPoint, const FConnectionParams& Params);
 };
diff --git a/Engine/Plugins/Experimental/PluginReferenceViewer/Source/PluginReferenceViewer/Private/PluginReferenceViewerSchema.cpp b/Engine/Plugins/Experimental/PluginReferenceViewer/Source/PluginReferenceViewer/Private/PluginReferenceViewerSchema.cpp
--- a/Engine/Plugins/Experimental/PluginReferenceViewer/Source/PluginReferenceViewer/Private/PluginReferenceViewerSchema.cpp
+++ b/Engine/Plugins/Experimental/PluginReferenceViewer/Source/PluginReferenceViewer/Private/PluginReferenceViewerSchema.cpp
@@ -23,10 +23,10 @@ public:
 	{
 	}
 
-	virtual FVector2D ComputeSplineTangent(const FVector2D& Start, const FVector2D& End) const override
+	virtual FVector2f ComputeSplineTangent(const FVector2f& Start, const FVector2f& End) const override
 	{
 		const int32 Tension = FMath::Abs<int32>(Start.X - End.X);
-		return Tension * FVector2D(1.0f, 0);
+		return Tension * FVector2f(1.0f, 0);
 	}
 
 	virtual void DetermineWiringStyle(UEdGraphPin* OutputPin, UEdGraphPin* InputPin, /*inout*/ FConnectionParams& Params) override
diff --git a/Engine/Plugins/Runtime/Metasound/Source/MetasoundEditor/Private/MetasoundEditorGraphConnectionDrawingPolicy.cpp b/Engine/Plugins/Runtime/Metasound/Source/MetasoundEditor/Private/MetasoundEditorGraphConnectionDrawingPolicy.cpp
--- a/Engine/Plugins/Runtime/Metasound/Source/MetasoundEditor/Private/MetasoundEditorGraphConnectionDrawingPolicy.cpp
+++ b/Engine/Plugins/Runtime/Metasound/Source/MetasoundEditor/Private/MetasoundEditorGraphConnectionDrawingPolicy.cpp
@@ -38,14 +38,14 @@ namespace Metasound
 			// Draw the bounding box for debugging
 			void DrawDebugConnection(FSlateWindowElementList& InDrawElementsList, uint32 InLayerId, const FDrawConnectionData& InDrawConnectionData)
 			{
-				const FVector2D& TL = InDrawConnectionData.Bounds.Min;
-				const FVector2D& BR = InDrawConnectionData.Bounds.Max;
-				const FVector2D TR = FVector2D(InDrawConnectionData.Bounds.Max.X, InDrawConnectionData.Bounds.Min.Y);
-				const FVector2D BL = FVector2D(InDrawConnectionData.Bounds.Min.X, InDrawConnectionData.Bounds.Max.Y);
+				const FVector2f& TL = InDrawConnectionData.Bounds.Min;
+				const FVector2f& BR = InDrawConnectionData.Bounds.Max;
+				const FVector2f TR = FVector2f(InDrawConnectionData.Bounds.Max.X, InDrawConnectionData.Bounds.Min.Y);
+				const FVector2f BL = FVector2f(InDrawConnectionData.Bounds.Min.X, InDrawConnectionData.Bounds.Max.Y);
 
-				auto DrawSpaceLine = [&](const FVector2D& Point1, const FVector2D& Point2, const FLinearColor& InColor)
+				auto DrawSpaceLine = [&](const FVector2f& Point1, const FVector2f& Point2, const FLinearColor& InColor)
 				{
-					const FVector2D FakeTangent = (Point2 - Point1).GetSafeNormal();
+					const FVector2f FakeTangent = (Point2 - Point1).GetSafeNormal();
 					FSlateDrawElement::MakeDrawSpaceSpline(InDrawElementsList, InLayerId, Point1, FakeTangent, Point2, FakeTangent, 1.0f, ESlateDrawEffect::None);
 				};
 
@@ -159,11 +159,11 @@ namespace Metasound
 			}
 		} // namespace DrawingPolicyPrivate
 
-		FDrawConnectionData::FDrawConnectionData(const FVector2D& InStart, const FVector2D& InEnd, const FEditor& Editor, const FVector2D& InSplineTangent, const FConnectionParams& InParams, const UGraphEditorSettings& InSettings, const FVector2D& InMousePosition)
+		FDrawConnectionData::FDrawConnectionData(const UE::Slate::FDeprecateVector2DParameter& InStart, const UE::Slate::FDeprecateVector2DParameter& InEnd, const FEditor& Editor, const UE::Slate::FDeprecateVector2DParameter& InSplineTangent, const FConnectionParams& InParams, const UGraphEditorSettings& InSettings, const UE::Slate::FDeprecateVector2DParameter& InMousePosition)
 			: P0(InStart)
-			, P0Tangent(InParams.StartDirection == EGPD_Output ? InSplineTangent : -InSplineTangent)
+			, P0Tangent(InParams.StartDirection == EGPD_Output ? UE::Slate::CastToVector2f(InSplineTangent) : UE::Slate::CastToVector2f(-InSplineTangent))
 			, P1(InEnd)
-			, P1Tangent(InParams.EndDirection == EGPD_Input ? InSplineTangent : -InSplineTangent)
+			, P1Tangent(InParams.EndDirection == EGPD_Input ? UE::Slate::CastToVector2f(InSplineTangent) : UE::Slate::CastToVector2f(-InSplineTangent))
 			, Params(InParams)
 		{
 			if (Params.AssociatedPin1)
@@ -174,12 +174,12 @@ namespace Metasound
 			// The curve will include the endpoints but can extend out of a tight bounds because of the tangents
 			// P0Tangent coefficient maximizes to 4/27 at a=1/3, and P1Tangent minimizes to -4/27 at a=2/3.
 			constexpr float MaximumTangentContribution = 4.0f / 27.0f;
-			Bounds = FBox2D(ForceInit);
+			Bounds = FBox2f(ForceInit);
 
-			Bounds += FVector2D(P0);
-			Bounds += FVector2D(P0 + MaximumTangentContribution * P0Tangent);
-			Bounds += FVector2D(P1);
-			Bounds += FVector2D(P1 - MaximumTangentContribution * P1Tangent);
+			Bounds += P0;
+			Bounds += P0 + MaximumTangentContribution * P0Tangent;
+			Bounds += P1;
+			Bounds += P1 - MaximumTangentContribution * P1Tangent;
 
 			if (InSettings.bTreatSplinesLikePins)
 			{
@@ -188,19 +188,19 @@ namespace Metasound
 				bCloseToSpline = Bounds.ComputeSquaredDistanceToPoint(InMousePosition) < QueryDistanceForCloseSquared;
 			}
 
-			ClosestPoint = FVector2D(ForceInit);
+			ClosestPoint = FVector2f(ForceInit);
 			if (InSettings.bTreatSplinesLikePins)
 			{
 				if (bCloseToSpline)
 				{
 					const int32 NumStepsToTest = 16;
 					const float StepInterval = 1.0f / (float)NumStepsToTest;
-					FVector2D Point1 = FMath::CubicInterp(P0, P0Tangent, P1, P1Tangent, 0.0f);
+					FVector2f Point1 = FMath::CubicInterp(P0, P0Tangent, P1, P1Tangent, 0.0f);
 					for (float TestAlpha = 0.0f; TestAlpha < 1.0f; TestAlpha += StepInterval)
 					{
-						const FVector2D Point2 = FMath::CubicInterp(P0, P0Tangent, P1, P1Tangent, TestAlpha + StepInterval);
+						const FVector2f Point2 = FMath::CubicInterp(P0, P0Tangent, P1, P1Tangent, TestAlpha + StepInterval);
 
-						const FVector2D ClosestPointToSegment = FMath::ClosestPointOnSegment2D(InMousePosition, Point1, Point2);
+						const FVector2f ClosestPointToSegment = FMath::ClosestPointOnSegment2D(InMousePosition, Point1, Point2);
 						const float DistanceSquared = (InMousePosition - ClosestPointToSegment).SizeSquared();
 
 						if (DistanceSquared < ClosestDistanceSquared)
@@ -273,7 +273,7 @@ namespace Metasound
 		{
 			// Don't show arrow images in the MetaSound graph
 			ArrowImage = nullptr;
-			ArrowRadius = FVector2D::ZeroVector;
+			ArrowRadius = FVector2f::ZeroVector;
 
 			WireAnimationLayerID = (InFrontLayerID + InBackLayerID) / 2;
 
@@ -444,7 +444,7 @@ namespace Metasound
 			}
 		}
 
-		void FGraphConnectionDrawingPolicy::DrawConnection(int32 LayerId, const FVector2D& Start, const FVector2D& End, const FConnectionParams& Params)
+		void FGraphConnectionDrawingPolicy::DrawConnection(int32 LayerId, const FVector2f& Start, const FVector2f& End, const FConnectionParams& Params)
 		{
 			METASOUND_TRACE_CPUPROFILER_EVENT_SCOPE(Metasound::Editor::FGraphConnectionDrawingPolicy::DrawConnection);
 
@@ -552,15 +552,15 @@ namespace Metasound
 
 			// Determine the spline position for the midpoint
 			const float MidpointAlpha = InSplineReparamTable.Eval(InSplineLength * 0.5f, 0.f);
-			const FVector2D Midpoint = FMath::CubicInterp(InData.P0, InData.P0Tangent, InData.P1, InData.P1Tangent, MidpointAlpha);
+			const FVector2f Midpoint = FMath::CubicInterp(InData.P0, InData.P0Tangent, InData.P1, InData.P1Tangent, MidpointAlpha);
 
 			// Approximate the slope at the midpoint (to orient the midpoint image to the spline)
-			const FVector2D MidpointPlusE = FMath::CubicInterp(InData.P0, InData.P0Tangent, InData.P1, InData.P1Tangent, MidpointAlpha + KINDA_SMALL_NUMBER);
-			const FVector2D MidpointMinusE = FMath::CubicInterp(InData.P0, InData.P0Tangent, InData.P1, InData.P1Tangent, MidpointAlpha - KINDA_SMALL_NUMBER);
-			const FVector2D SlopeUnnormalized = MidpointPlusE - MidpointMinusE;
+			const FVector2f MidpointPlusE = FMath::CubicInterp(InData.P0, InData.P0Tangent, InData.P1, InData.P1Tangent, MidpointAlpha + KINDA_SMALL_NUMBER);
+			const FVector2f MidpointMinusE = FMath::CubicInterp(InData.P0, InData.P0Tangent, InData.P1, InData.P1Tangent, MidpointAlpha - KINDA_SMALL_NUMBER);
+			const FVector2f SlopeUnnormalized = MidpointPlusE - MidpointMinusE;
 
 			// Draw the arrow
-			const FVector2D MidpointDrawPos = Midpoint - MidpointRadius;
+			const FVector2f MidpointDrawPos = Midpoint - MidpointRadius;
 			const float AngleInRadians = SlopeUnnormalized.IsNearlyZero() ? 0.0f : FMath::Atan2(SlopeUnnormalized.Y, SlopeUnnormalized.X);
 
 			FSlateDrawElement::MakeRotatedBox(
@@ -570,7 +570,7 @@ namespace Metasound
 				MidpointImage,
 				ESlateDrawEffect::None,
 				AngleInRadians,
-				TOptional<FVector2D>(),
+				TOptional<FVector2f>(),
 				FSlateDrawElement::RelativeToElement,
 				InData.Params.WireColor
 			);
@@ -594,7 +594,7 @@ namespace Metasound
 			const float BubbleSpacing = InParams.SpacingFactor * ZoomFactor;
 			float BubbleSpeed = InParams.SpeedFactor * ZoomFactor;
 
-			const FVector2D BubbleSize = BubbleImage->ImageSize * ZoomFactor * 0.2f * InParams.ConnectionData.Params.WireThickness;
+			const FVector2f BubbleSize = BubbleImage->ImageSize * ZoomFactor * 0.2f * InParams.ConnectionData.Params.WireThickness;
 
 			const float BubbleOffset = FMath::Fmod(InParams.AppTime * BubbleSpeed, BubbleSpacing);
 			const int32 NumBubbles = FMath::CeilToInt(InParams.SplineLength / BubbleSpacing);
@@ -605,7 +605,7 @@ namespace Metasound
 				if (Distance < InParams.SplineLength)
 				{
 					const float Alpha = InParams.SplineReparamTable.Eval(Distance, 0.f);
-					FVector2D BubblePos = FMath::CubicInterp(InParams.ConnectionData.P0, InParams.ConnectionData.P0Tangent, InParams.ConnectionData.P1, InParams.ConnectionData.P1Tangent, Alpha);
+					FVector2f BubblePos = FMath::CubicInterp(InParams.ConnectionData.P0, InParams.ConnectionData.P0Tangent, InParams.ConnectionData.P1, InParams.ConnectionData.P1Tangent, Alpha);
 					BubblePos -= BubbleSize * 0.5f;
 
 					FSlateDrawElement::MakeBox(
@@ -720,16 +720,16 @@ namespace Metasound
 
 			TArray<FVector2f> PointsUp;
 			TArray<FVector2f> PointsDown;
-			FVector2D LastPointPos = FVector2D::Zero();
-			FVector2D LastTangent = FVector2D::Zero();
+			FVector2f LastPointPos = FVector2f::Zero();
+			FVector2f LastTangent = FVector2f::Zero();
 
 			auto AddPoint = [&](const float EnvMagnitude, const float Distance)
 			{
 				const float Alpha = InParams.SplineReparamTable.Eval(Distance, 0.f);
-				FVector2D PointPos = FMath::CubicInterp(InParams.ConnectionData.P0, InParams.ConnectionData.P0Tangent, InParams.ConnectionData.P1, InParams.ConnectionData.P1Tangent, Alpha);
+				FVector2f PointPos = FMath::CubicInterp(InParams.ConnectionData.P0, InParams.ConnectionData.P0Tangent, InParams.ConnectionData.P1, InParams.ConnectionData.P1Tangent, Alpha);
 
-				LastTangent = FVector2D(PointPos.X - LastPointPos.X, PointPos.Y - LastPointPos.Y);
-				const FVector2D Normal = FVector2D(LastTangent.Y, -1.0f * LastTangent.X).GetSafeNormal();
+				LastTangent = FVector2f(PointPos.X - LastPointPos.X, PointPos.Y - LastPointPos.Y);
+				const FVector2f Normal = FVector2f(LastTangent.Y, -1.0f * LastTangent.X).GetSafeNormal();
 				LastPointPos = PointPos;
 
 				const float HalfEnvMagnitude = EnvMagnitude * 0.5f;
diff --git a/Engine/Plugins/Runtime/Metasound/Source/MetasoundEditor/Private/MetasoundEditorGraphConnectionDrawingPolicy.h b/Engine/Plugins/Runtime/Metasound/Source/MetasoundEditor/Private/MetasoundEditorGraphConnectionDrawingPolicy.h
--- a/Engine/Plugins/Runtime/Metasound/Source/MetasoundEditor/Private/MetasoundEditorGraphConnectionDrawingPolicy.h
+++ b/Engine/Plugins/Runtime/Metasound/Source/MetasoundEditor/Private/MetasoundEditorGraphConnectionDrawingPolicy.h
@@ -31,14 +31,14 @@ namespace Metasound
 
 		struct FDrawConnectionData
 		{
-			FVector2D P0;
-			FVector2D P0Tangent;
-			FVector2D P1;
-			FVector2D P1Tangent;
+			FDeprecateSlateVector2D P0;
+			FDeprecateSlateVector2D P0Tangent;
+			FDeprecateSlateVector2D P1;
+			FDeprecateSlateVector2D P1Tangent;
 
 			FConnectionParams Params;
 
-			FBox2D Bounds;
+			FBox2f Bounds;
 
 			const UEdGraphPin* OutputPin = nullptr;
 			const FMetasoundFrontendVertex* OutputVertex = nullptr;
@@ -54,11 +54,11 @@ namespace Metasound
 			// dead zone to avoid mistakes caused by missing a double-click on a connection.
 			float QueryDistanceForCloseSquared = 0.0f;
 
-			FVector2D ClosestPoint;
+			FDeprecateSlateVector2D ClosestPoint;
 
 			const FMetasoundFrontendEdgeStyle* EdgeStyle = nullptr;
 
-			FDrawConnectionData(const FVector2D& InStart, const FVector2D& InEnd, const FEditor& Editor, const FVector2D& InSplineTangent, const FConnectionParams& InParams, const UGraphEditorSettings& InSettings, const FVector2D& InMousePosition);
+			FDrawConnectionData(const UE::Slate::FDeprecateVector2DParameter& InStart, const UE::Slate::FDeprecateVector2DParameter& InEnd, const FEditor& Editor, const UE::Slate::FDeprecateVector2DParameter& InSplineTangent, const FConnectionParams& InParams, const UGraphEditorSettings& InSettings, const UE::Slate::FDeprecateVector2DParameter& InMousePosition);
 
 			void UpdateSplineOverlap(FGraphSplineOverlapResult& OutResult) const;
 		};
@@ -112,7 +112,7 @@ namespace Metasound
 			FGraphConnectionDrawingPolicy(int32 InBackLayerID, int32 InFrontLayerID, float ZoomFactor, const FSlateRect& InClippingRect, FSlateWindowElementList& InDrawElements, UEdGraph* InGraphObj);
 
 			// FConnectionDrawingPolicy interface
-			virtual void DrawConnection(int32 LayerId, const FVector2D& Start, const FVector2D& End, const FConnectionParams& InParams) override;
+			virtual void DrawConnection(int32 LayerId, const FVector2f& Start, const FVector2f& End, const FConnectionParams& InParams) override;
 			virtual void DetermineWiringStyle(UEdGraphPin* InPinA, UEdGraphPin* InPinB, FConnectionParams& OutParams) override;
 			// End of FConnectionDrawingPolicy interface
 
diff --git a/Engine/Source/Editor/AIGraph/Private/AIGraphConnectionDrawingPolicy.cpp b/Engine/Source/Editor/AIGraph/Private/AIGraphConnectionDrawingPolicy.cpp
--- a/Engine/Source/Editor/AIGraph/Private/AIGraphConnectionDrawingPolicy.cpp
+++ b/Engine/Source/Editor/AIGraph/Private/AIGraphConnectionDrawingPolicy.cpp
@@ -54,7 +54,7 @@ void FAIGraphConnectionDrawingPolicy::Draw(TMap<TSharedRef<SWidget>, FArrangedWi
 	FConnectionDrawingPolicy::Draw(InPinGeometries, ArrangedNodes);
 }
 
-void FAIGraphConnectionDrawingPolicy::DrawPreviewConnector(const FGeometry& PinGeometry, const FVector2D& StartPoint, const FVector2D& EndPoint, UEdGraphPin* Pin)
+void FAIGraphConnectionDrawingPolicy::DrawPreviewConnector(const FGeometry& PinGeometry, const FVector2f& StartPoint, const FVector2f& EndPoint, UEdGraphPin* Pin)
 {
 	FConnectionParams Params;
 	DetermineWiringStyle(Pin, nullptr, /*inout*/ Params);
@@ -69,21 +69,21 @@ void FAIGraphConnectionDrawingPolicy::DrawPreviewConnector(const FGeometry& PinG
 	}
 }
 
-void FAIGraphConnectionDrawingPolicy::DrawSplineWithArrow(const FVector2D& StartAnchorPoint, const FVector2D& EndAnchorPoint, const FConnectionParams& Params)
+void FAIGraphConnectionDrawingPolicy::DrawSplineWithArrow(const FVector2f& StartAnchorPoint, const FVector2f& EndAnchorPoint, const FConnectionParams& Params)
 {
 	// bUserFlag1 indicates that we need to reverse the direction of connection (used by debugger)
-	const FVector2D& P0 = Params.bUserFlag1 ? EndAnchorPoint : StartAnchorPoint;
-	const FVector2D& P1 = Params.bUserFlag1 ? StartAnchorPoint : EndAnchorPoint;
+	const FVector2f& P0 = Params.bUserFlag1 ? EndAnchorPoint : StartAnchorPoint;
+	const FVector2f& P1 = Params.bUserFlag1 ? StartAnchorPoint : EndAnchorPoint;
 
 	Internal_DrawLineWithArrow(P0, P1, Params);
 }
 
-void FAIGraphConnectionDrawingPolicy::Internal_DrawLineWithArrow(const FVector2D& StartAnchorPoint, const FVector2D& EndAnchorPoint, const FConnectionParams& Params)
+void FAIGraphConnectionDrawingPolicy::Internal_DrawLineWithArrow(const UE::Slate::FDeprecateVector2DParameter& StartAnchorPoint, const UE::Slate::FDeprecateVector2DParameter& EndAnchorPoint, const FConnectionParams& Params)
 {
 	//@TODO: Should this be scaled by zoom factor?
 	const float LineSeparationAmount = 4.5f;
 
-	const FVector2f StartAnchorPointTemp(UE::Slate::CastToVector2f(StartAnchorPoint));
+	const FVector2f StartAnchorPointTemp(StartAnchorPoint);
 	const FVector2f EndAnchorPointTemp(static_cast<float>(EndAnchorPoint.X), static_cast<float>(EndAnchorPoint.Y));
 	const FVector2f DeltaPos = EndAnchorPointTemp - StartAnchorPointTemp;
 	const FVector2f UnitDelta = DeltaPos.GetSafeNormal();
@@ -92,11 +92,11 @@ void FAIGraphConnectionDrawingPolicy::Internal_DrawLineWithArrow(const FVector2D
 	// Come up with the final start/end points
 	const FVector2f DirectionBias = Normal * LineSeparationAmount;
 	const FVector2f LengthBias = ArrowRadius.X * UnitDelta;
-	const FDeprecateSlateVector2D StartPoint = StartAnchorPointTemp + DirectionBias + LengthBias;
-	const FDeprecateSlateVector2D EndPoint = EndAnchorPointTemp + DirectionBias - LengthBias;
+	const FVector2f StartPoint = StartAnchorPointTemp + DirectionBias + LengthBias;
+	const FVector2f EndPoint = EndAnchorPointTemp + DirectionBias - LengthBias;
 
 	// Draw a line/spline
-	DrawConnection(WireLayerID, FVector2D(StartPoint), FVector2D(EndPoint), Params);
+	DrawConnection(WireLayerID, StartPoint, EndPoint, Params);
 
 	// Draw the arrow
 	const FVector2f ArrowDrawPos = EndPoint - ArrowRadius;
@@ -118,21 +118,21 @@ void FAIGraphConnectionDrawingPolicy::Internal_DrawLineWithArrow(const FVector2D
 void FAIGraphConnectionDrawingPolicy::DrawSplineWithArrow(const FGeometry& StartGeom, const FGeometry& EndGeom, const FConnectionParams& Params)
 {
 	// Get a reasonable seed point (halfway between the boxes)
-	const FVector2D StartCenter = FGeometryHelper::CenterOf(StartGeom);
-	const FVector2D EndCenter = FGeometryHelper::CenterOf(EndGeom);
-	const FVector2D SeedPoint = (StartCenter + EndCenter) * 0.5f;
+	const FVector2f StartCenter = FGeometryHelper::CenterOf(StartGeom);
+	const FVector2f EndCenter = FGeometryHelper::CenterOf(EndGeom);
+	const FVector2f SeedPoint = (StartCenter + EndCenter) * 0.5f;
 
 	// Find the (approximate) closest points between the two boxes
-	const FVector2D StartAnchorPoint = FGeometryHelper::FindClosestPointOnGeom(StartGeom, SeedPoint);
-	const FVector2D EndAnchorPoint = FGeometryHelper::FindClosestPointOnGeom(EndGeom, SeedPoint);
+	const FVector2f StartAnchorPoint = FGeometryHelper::FindClosestPointOnGeom(StartGeom, SeedPoint);
+	const FVector2f EndAnchorPoint = FGeometryHelper::FindClosestPointOnGeom(EndGeom, SeedPoint);
 
 	DrawSplineWithArrow(StartAnchorPoint, EndAnchorPoint, Params);
 }
 
-FVector2D FAIGraphConnectionDrawingPolicy::ComputeSplineTangent(const FVector2D& Start, const FVector2D& End) const
+FVector2f FAIGraphConnectionDrawingPolicy::ComputeSplineTangent(const FVector2f& Start, const FVector2f& End) const
 {
-	const FVector2D Delta = End - Start;
-	const FVector2D NormDelta = Delta.GetSafeNormal();
+	const FVector2f Delta = End - Start;
+	const FVector2f NormDelta = Delta.GetSafeNormal();
 
 	return NormDelta;
 }
diff --git a/Engine/Source/Editor/AIGraph/Public/AIGraphConnectionDrawingPolicy.h b/Engine/Source/Editor/AIGraph/Public/AIGraphConnectionDrawingPolicy.h
--- a/Engine/Source/Editor/AIGraph/Public/AIGraphConnectionDrawingPolicy.h
+++ b/Engine/Source/Editor/AIGraph/Public/AIGraphConnectionDrawingPolicy.h
@@ -34,11 +34,11 @@ public:
 	virtual void DetermineWiringStyle(UEdGraphPin* OutputPin, UEdGraphPin* InputPin, /*inout*/ FConnectionParams& Params) override;
 	virtual void Draw(TMap<TSharedRef<SWidget>, FArrangedWidget>& PinGeometries, FArrangedChildren& ArrangedNodes) override;
 	virtual void DrawSplineWithArrow(const FGeometry& StartGeom, const FGeometry& EndGeom, const FConnectionParams& Params) override;
-	virtual void DrawSplineWithArrow(const FVector2D& StartPoint, const FVector2D& EndPoint, const FConnectionParams& Params) override;
-	virtual void DrawPreviewConnector(const FGeometry& PinGeometry, const FVector2D& StartPoint, const FVector2D& EndPoint, UEdGraphPin* Pin) override;
-	virtual FVector2D ComputeSplineTangent(const FVector2D& Start, const FVector2D& End) const override;
+	virtual void DrawSplineWithArrow(const FVector2f& StartPoint, const FVector2f& EndPoint, const FConnectionParams& Params) override;
+	virtual void DrawPreviewConnector(const FGeometry& PinGeometry, const FVector2f& StartPoint, const FVector2f& EndPoint, UEdGraphPin* Pin) override;
+	virtual FVector2f ComputeSplineTangent(const FVector2f& Start, const FVector2f& End) const override;
 	// End of FConnectionDrawingPolicy interface
 
 protected:
-	void Internal_DrawLineWithArrow(const FVector2D& StartAnchorPoint, const FVector2D& EndAnchorPoint, const FConnectionParams& Params);
+	void Internal_DrawLineWithArrow(const UE::Slate::FDeprecateVector2DParameter& StartAnchorPoint, const UE::Slate::FDeprecateVector2DParameter& EndAnchorPoint, const FConnectionParams& Params);
 };
diff --git a/Engine/Source/Editor/AnimationBlueprintEditor/Private/AnimGraphConnectionDrawingPolicy.cpp b/Engine/Source/Editor/AnimationBlueprintEditor/Private/AnimGraphConnectionDrawingPolicy.cpp
--- a/Engine/Source/Editor/AnimationBlueprintEditor/Private/AnimGraphConnectionDrawingPolicy.cpp
+++ b/Engine/Source/Editor/AnimationBlueprintEditor/Private/AnimGraphConnectionDrawingPolicy.cpp
@@ -145,7 +145,7 @@ void FAnimGraphConnectionDrawingPolicy::DetermineStyleOfExecWire(float& Thicknes
 	bDrawBubbles = true;
 }
 
-void FAnimGraphConnectionDrawingPolicy::DrawConnection(int32 LayerId, const FVector2D& Start, const FVector2D& End, const FConnectionParams& Params)
+void FAnimGraphConnectionDrawingPolicy::DrawConnection(int32 LayerId, const FVector2f& Start, const FVector2f& End, const FConnectionParams& Params) 
 {
 	bool bCompositeWire = false;
 
@@ -187,12 +187,12 @@ void FAnimGraphConnectionDrawingPolicy::DrawConnection(int32 LayerId, const FVec
 						const float AttributeWireThickness = FMath::Lerp(MinAttributeWireThickness, MaxAttributeWireThickness, ZoomLevelAlpha);
 						const float WireGap = FMath::Lerp(MinWireGap, MaxWireGap, ZoomLevelAlpha);
 
-						const FVector2D& P0 = Start;
-						const FVector2D& P1 = End;
+						const FVector2f& P0 = Start;
+						const FVector2f& P1 = End;
 
-						const FVector2D SplineTangent = ComputeSplineTangent(P0, P1);
-						const FVector2D P0Tangent = (Params.StartDirection == EGPD_Output) ? SplineTangent : -SplineTangent;
-						const FVector2D P1Tangent = (Params.EndDirection == EGPD_Input) ? SplineTangent : -SplineTangent;
+						const FVector2f SplineTangent = ComputeSplineTangent(P0, P1);
+						const FVector2f P0Tangent = (Params.StartDirection == EGPD_Output) ? SplineTangent : -SplineTangent;
+						const FVector2f P1Tangent = (Params.EndDirection == EGPD_Input) ? SplineTangent : -SplineTangent;
 
 						bCompositeWire = true;
 
@@ -241,9 +241,9 @@ void FAnimGraphConnectionDrawingPolicy::DrawConnection(int32 LayerId, const FVec
 						}
 
 						const float InitialOffset = TotalThickness * 0.5f;
-						FVector2D SubWireP0 = P0;
+						FVector2f SubWireP0 = P0;
 						SubWireP0.Y += InitialOffset;
-						FVector2D SubWireP1 = P1;
+						FVector2f SubWireP1 = P1;
 						SubWireP1.Y += InitialOffset;
 
 						// Draw in reverse order to get pose wires appearing on top
diff --git a/Engine/Source/Editor/AnimationBlueprintEditor/Private/AnimGraphConnectionDrawingPolicy.h b/Engine/Source/Editor/AnimationBlueprintEditor/Private/AnimGraphConnectionDrawingPolicy.h
--- a/Engine/Source/Editor/AnimationBlueprintEditor/Private/AnimGraphConnectionDrawingPolicy.h
+++ b/Engine/Source/Editor/AnimationBlueprintEditor/Private/AnimGraphConnectionDrawingPolicy.h
@@ -34,7 +34,7 @@ public:
 	virtual bool TreatWireAsExecutionPin(UEdGraphPin* InputPin, UEdGraphPin* OutputPin) const override;
 	virtual void BuildExecutionRoadmap() override;
 	virtual void DetermineStyleOfExecWire(float& Thickness, FLinearColor& WireColor, bool& bDrawBubbles, const FTimePair& Times) override;
-	virtual void DrawConnection(int32 LayerId, const FVector2D& Start, const FVector2D& End, const FConnectionParams& Params) override;
+	virtual void DrawConnection(int32 LayerId, const FVector2f& Start, const FVector2f& End, const FConnectionParams& Params) override;
 	virtual void BuildPinToPinWidgetMap(TMap<TSharedRef<SWidget>, FArrangedWidget>& InPinGeometries) override;
 	virtual void ApplyHoverDeemphasis(UEdGraphPin* OutputPin, UEdGraphPin* InputPin, /*inout*/ float& Thickness, /*inout*/ FLinearColor& WireColor) override;
 	// End of FKismetConnectionDrawingPolicy interface
diff --git a/Engine/Source/Editor/AnimationBlueprintEditor/Private/StateMachineConnectionDrawingPolicy.cpp b/Engine/Source/Editor/AnimationBlueprintEditor/Private/StateMachineConnectionDrawingPolicy.cpp
--- a/Engine/Source/Editor/AnimationBlueprintEditor/Private/StateMachineConnectionDrawingPolicy.cpp
+++ b/Engine/Source/Editor/AnimationBlueprintEditor/Private/StateMachineConnectionDrawingPolicy.cpp
@@ -147,20 +147,20 @@ void FStateMachineConnectionDrawingPolicy::Draw(TMap<TSharedRef<SWidget>, FArran
 	FConnectionDrawingPolicy::Draw(InPinGeometries, ArrangedNodes);
 }
 
-void FStateMachineConnectionDrawingPolicy::DrawPreviewConnector(const FGeometry& PinGeometry, const FVector2D& StartPoint, const FVector2D& EndPoint, UEdGraphPin* Pin)
+void FStateMachineConnectionDrawingPolicy::DrawPreviewConnector(const FGeometry& PinGeometry, const FVector2f& StartPoint, const FVector2f& EndPoint, UEdGraphPin* Pin)
 {
 	FConnectionParams Params;
 	DetermineWiringStyle(Pin, nullptr, /*inout*/ Params);
 
-	const FVector2D SeedPoint = EndPoint;
-	const FVector2D AdjustedStartPoint = FGeometryHelper::FindClosestPointOnGeom(PinGeometry, SeedPoint);
+	const FVector2f SeedPoint = EndPoint;
+	const FVector2f AdjustedStartPoint = FGeometryHelper::FindClosestPointOnGeom(PinGeometry, SeedPoint);
 
 	Params.bUserFlag2 = false; // bUserFlag2 is used to indicate whether the drawn arrow is a preview transition (the temporary transition when creating or relinking).
 	DrawSplineWithArrow(AdjustedStartPoint, EndPoint, Params);
 }
 
 
-void FStateMachineConnectionDrawingPolicy::DrawSplineWithArrow(const FVector2D& StartAnchorPoint, const FVector2D& EndAnchorPoint, const FConnectionParams& Params)
+void FStateMachineConnectionDrawingPolicy::DrawSplineWithArrow(const FVector2f& StartAnchorPoint, const FVector2f& EndAnchorPoint, const FConnectionParams& Params)
 {
 	Internal_DrawLineWithArrow(StartAnchorPoint, EndAnchorPoint, Params);
 
@@ -171,40 +171,38 @@ void FStateMachineConnectionDrawingPolicy::DrawSplineWithArrow(const FVector2D&
 	}
 }
 
-void FStateMachineConnectionDrawingPolicy::Internal_DrawLineWithArrow(const FVector2D& StartAnchorPoint, const FVector2D& EndAnchorPoint, const FConnectionParams& Params)
+void FStateMachineConnectionDrawingPolicy::Internal_DrawLineWithArrow(const FVector2f& StartAnchorPoint, const FVector2f& EndAnchorPoint, const FConnectionParams& Params)
 {
 	//@TODO: Should this be scaled by zoom factor?
 	const float LineSeparationAmount = 4.5f;
 
-	const FVector2D DeltaPos = EndAnchorPoint - StartAnchorPoint;
-	const FVector2D UnitDelta = DeltaPos.GetSafeNormal();
-	const FVector2D Normal = FVector2D(DeltaPos.Y, -DeltaPos.X).GetSafeNormal();
+	const FVector2f DeltaPos = EndAnchorPoint - StartAnchorPoint;
+	const FVector2f UnitDelta = DeltaPos.GetSafeNormal();
+	const FVector2f Normal = FVector2f(DeltaPos.Y, -DeltaPos.X).GetSafeNormal();
 
 	// Come up with the final start/end points
-	const FVector2D DirectionBias = Normal * LineSeparationAmount;
-	const FVector2D LengthBias = ArrowRadius.X * UnitDelta;
-	const FVector2D StartPoint = StartAnchorPoint + DirectionBias + LengthBias;
-	const FVector2D EndPoint = EndAnchorPoint + DirectionBias - LengthBias;
+	const FVector2f DirectionBias = Normal * LineSeparationAmount;
+	const FVector2f LengthBias = ArrowRadius.X * UnitDelta;
+	const FVector2f StartPoint = StartAnchorPoint + DirectionBias + LengthBias;
+	const FVector2f EndPoint = EndAnchorPoint + DirectionBias - LengthBias;
 	FLinearColor ArrowHeadColor = Params.WireColor;
 
 	// Draw a line/spline
 	DrawConnection(WireLayerID, StartPoint, EndPoint, Params);
 
-	const FVector2D ArrowDrawPos = EndPoint - ArrowRadius;
+	const FVector2f ArrowDrawPos = EndPoint - ArrowRadius;
 	const double AngleInRadians = FMath::Atan2(DeltaPos.Y, DeltaPos.X);
 
 	// Draw the transition grab handles in case the mouse is hovering the transition
 	bool StartHovered = false;
 	bool EndHovered = false;
-	const FVector FVecMousePos = FVector(LocalMousePosition.X, LocalMousePosition.Y, 0.0f);
-	const FVector ClosestPoint = FMath::ClosestPointOnSegment(FVecMousePos,
-		FVector(StartPoint.X, StartPoint.Y, 0.0f),
-		FVector(EndPoint.X, EndPoint.Y, 0.0f));
+	const FVector2f FVecMousePos = FVector2f(LocalMousePosition.X, LocalMousePosition.Y);
+	const FVector2f ClosestPoint = FMath::ClosestPointOnSegment2D(FVecMousePos, StartPoint, EndPoint);
 	if ((ClosestPoint - FVecMousePos).Length() < RelinkHandleHoverRadius)
 	{
-		StartHovered = FVector2D(StartPoint - LocalMousePosition).Length() < RelinkHandleHoverRadius;
-		EndHovered = FVector2D(EndPoint - LocalMousePosition).Length() < RelinkHandleHoverRadius;
-		FVector2D HoverIndicatorPosition = StartHovered ? StartPoint : EndPoint;
+		StartHovered = (StartPoint - LocalMousePosition).Length() < RelinkHandleHoverRadius;
+		EndHovered = (EndPoint - LocalMousePosition).Length() < RelinkHandleHoverRadius;
+		FVector2f HoverIndicatorPosition = StartHovered ? StartPoint : EndPoint;
 
 		// Set the hovered pin results. This will be used by the SGraphPanel again.
 		const float SquaredDistToPin1 = (Params.AssociatedPin1 != nullptr) ? (StartPoint - LocalMousePosition).SizeSquared() : FLT_MAX;
@@ -235,7 +233,7 @@ void FStateMachineConnectionDrawingPolicy::Internal_DrawLineWithArrow(const FVec
 				// Draw circle around the arrow in case the transition is hovered (mouse close or over transition line or arrow head).
 				const int CircleLineSegments = 16;
 				const float CircleRadius = 10.0f;
-				const FVector2D CircleCenter = EndPoint - UnitDelta * 2.0f;
+				const FVector2f CircleCenter = EndPoint - UnitDelta * 2.0f;
 				DrawCircle(CircleCenter, CircleRadius, Params.WireColor, CircleLineSegments);
 			}
 		}
@@ -256,8 +254,8 @@ void FStateMachineConnectionDrawingPolicy::Internal_DrawLineWithArrow(const FVec
 			}
 		}
 
-		const FVector2D TransitionCenter = StartAnchorPoint + DeltaPos * 0.5f;
-		const FVector2D TextPosition = TransitionCenter + Normal * 15.0f * ZoomFactor;
+		const FVector2f TransitionCenter = StartAnchorPoint + DeltaPos * 0.5f;
+		const FVector2f TextPosition = TransitionCenter + Normal * 15.0f * ZoomFactor;
 
 		FSlateDrawElement::MakeText(
 			DrawElementsList,
@@ -275,13 +273,13 @@ void FStateMachineConnectionDrawingPolicy::Internal_DrawLineWithArrow(const FVec
 		ArrowImage,
 		ESlateDrawEffect::None,
 		static_cast<float>(AngleInRadians),
-		TOptional<FVector2D>(),
+		TOptional<FVector2f>(),
 		FSlateDrawElement::RelativeToElement,
 		ArrowHeadColor
 	);
 }
 
-void FStateMachineConnectionDrawingPolicy::DrawCircle(const FVector2D& Center, float Radius, const FLinearColor& Color, const int NumLineSegments)
+void FStateMachineConnectionDrawingPolicy::DrawCircle(const FVector2f& Center, float Radius, const FLinearColor& Color, const int NumLineSegments)
 {
 	TempPoints.Empty();
 
@@ -290,7 +288,7 @@ void FStateMachineConnectionDrawingPolicy::DrawCircle(const FVector2D& Center, f
 	{
 		const float Angle = (i / NumFloatLineSegments) * TWO_PI;
 
-		FVector2D PointOnCircle;
+		FVector2f PointOnCircle;
 		PointOnCircle.X = cosf(Angle) * Radius;
 		PointOnCircle.Y = sinf(Angle) * Radius;
 		TempPoints.Add(PointOnCircle);
@@ -299,7 +297,7 @@ void FStateMachineConnectionDrawingPolicy::DrawCircle(const FVector2D& Center, f
 	FSlateDrawElement::MakeLines(
 		DrawElementsList,
 		ArrowLayerID + 1,
-		FPaintGeometry(Center, FVector2D(Radius, Radius) * ZoomFactor, ZoomFactor),
+		FPaintGeometry(Center, FVector2f(Radius, Radius) * ZoomFactor, ZoomFactor),
 		TempPoints,
 		ESlateDrawEffect::None,
 		Color
@@ -309,21 +307,21 @@ void FStateMachineConnectionDrawingPolicy::DrawCircle(const FVector2D& Center, f
 void FStateMachineConnectionDrawingPolicy::DrawSplineWithArrow(const FGeometry& StartGeom, const FGeometry& EndGeom, const FConnectionParams& Params)
 {
 	// Get a reasonable seed point (halfway between the boxes)
-	const FVector2D StartCenter = FGeometryHelper::CenterOf(StartGeom);
-	const FVector2D EndCenter = FGeometryHelper::CenterOf(EndGeom);
-	const FVector2D SeedPoint = (StartCenter + EndCenter) * 0.5f;
+	const FVector2f StartCenter = FGeometryHelper::CenterOf(StartGeom);
+	const FVector2f EndCenter = FGeometryHelper::CenterOf(EndGeom);
+	const FVector2f SeedPoint = (StartCenter + EndCenter) * 0.5f;
 	
 	// Find the (approximate) closest points between the two boxes
-	const FVector2D StartAnchorPoint = FGeometryHelper::FindClosestPointOnGeom(StartGeom, SeedPoint);
-	const FVector2D EndAnchorPoint = FGeometryHelper::FindClosestPointOnGeom(EndGeom, SeedPoint);
+	const FVector2f StartAnchorPoint = FGeometryHelper::FindClosestPointOnGeom(StartGeom, SeedPoint);
+	const FVector2f EndAnchorPoint = FGeometryHelper::FindClosestPointOnGeom(EndGeom, SeedPoint);
 
 	DrawSplineWithArrow(StartAnchorPoint, EndAnchorPoint, Params);
 }
 
-FVector2D FStateMachineConnectionDrawingPolicy::ComputeSplineTangent(const FVector2D& Start, const FVector2D& End) const
+FVector2f FStateMachineConnectionDrawingPolicy::ComputeSplineTangent(const FVector2f& Start, const FVector2f& End) const
 {
-	const FVector2D Delta = End - Start;
-	const FVector2D NormDelta = Delta.GetSafeNormal();
+	const FVector2f Delta = End - Start;
+	const FVector2f NormDelta = Delta.GetSafeNormal();
 
 	return NormDelta;
 }
diff --git a/Engine/Source/Editor/AnimationBlueprintEditor/Private/StateMachineConnectionDrawingPolicy.h b/Engine/Source/Editor/AnimationBlueprintEditor/Private/StateMachineConnectionDrawingPolicy.h
--- a/Engine/Source/Editor/AnimationBlueprintEditor/Private/StateMachineConnectionDrawingPolicy.h
+++ b/Engine/Source/Editor/AnimationBlueprintEditor/Private/StateMachineConnectionDrawingPolicy.h
@@ -46,17 +46,17 @@ public:
 		/*out*/ FArrangedWidget*& EndWidgetGeometry
 		) override;
 	virtual void DrawSplineWithArrow(const FGeometry& StartGeom, const FGeometry& EndGeom, const FConnectionParams& Params) override;
-	virtual void DrawSplineWithArrow(const FVector2D& StartPoint, const FVector2D& EndPoint, const FConnectionParams& Params) override;
-	virtual void DrawPreviewConnector(const FGeometry& PinGeometry, const FVector2D& StartPoint, const FVector2D& EndPoint, UEdGraphPin* Pin) override;
-	virtual FVector2D ComputeSplineTangent(const FVector2D& Start, const FVector2D& End) const override;
+	virtual void DrawSplineWithArrow(const FVector2f& StartPoint, const FVector2f& EndPoint, const FConnectionParams& Params) override;
+	virtual void DrawPreviewConnector(const FGeometry& PinGeometry, const FVector2f& StartPoint, const FVector2f& EndPoint, UEdGraphPin* Pin) override;
+	virtual FVector2f ComputeSplineTangent(const FVector2f& Start, const FVector2f& End) const override;
 	// End of FConnectionDrawingPolicy interface
 
 protected:
-	void Internal_DrawLineWithArrow(const FVector2D& StartAnchorPoint, const FVector2D& EndAnchorPoint, const FConnectionParams& Params);
+	void Internal_DrawLineWithArrow(const FVector2f& StartAnchorPoint, const FVector2f& EndAnchorPoint, const FConnectionParams& Params);
 
 	// Draw line-based circle (no solid filling)
-	void DrawCircle(const FVector2D& Center, float Radius, const FLinearColor& Color, const int NumLineSegments);
-	TArray<FVector2D> TempPoints;
+	void DrawCircle(const FVector2f& Center, float Radius, const FLinearColor& Color, const int NumLineSegments);
+	TArray<FVector2f> TempPoints;
 
 	static constexpr float RelinkHandleHoverRadius = 20.0f;
 };
diff --git a/Engine/Source/Editor/GraphEditor/Private/ConnectionDrawingPolicy.cpp b/Engine/Source/Editor/GraphEditor/Private/ConnectionDrawingPolicy.cpp
--- a/Engine/Source/Editor/GraphEditor/Private/ConnectionDrawingPolicy.cpp
+++ b/Engine/Source/Editor/GraphEditor/Private/ConnectionDrawingPolicy.cpp
@@ -5,69 +5,80 @@
 #include "Rendering/DrawElements.h"
 #include "Widgets/SToolTip.h"
 #include "Framework/Application/SlateApplication.h"
+#include "Algo/ForEach.h"
 
 DEFINE_LOG_CATEGORY(LogConnectionDrawingPolicy);
 
 //////////////////////////////////////////////////////////////////////////
 // FGeometryHelper
 
-FVector2D FGeometryHelper::VerticalMiddleLeftOf(const FGeometry& SomeGeometry)
+UE::Slate::FDeprecateVector2DResult FGeometryHelper::VerticalMiddleLeftOf(const FGeometry& SomeGeometry)
 {
-	const FVector2D GeometryDrawSize = SomeGeometry.GetDrawSize();
-	return FVector2D(
-		SomeGeometry.AbsolutePosition.X,
-		SomeGeometry.AbsolutePosition.Y + GeometryDrawSize.Y/2 );
+	const FVector2f GeometryDrawSize = SomeGeometry.GetDrawSize();
+	return FVector2f(SomeGeometry.AbsolutePosition.X,
+		SomeGeometry.AbsolutePosition.Y + GeometryDrawSize.Y/2.f);
 }
 
-FVector2D FGeometryHelper::VerticalMiddleRightOf(const FGeometry& SomeGeometry)
+UE::Slate::FDeprecateVector2DResult FGeometryHelper::VerticalMiddleRightOf(const FGeometry& SomeGeometry)
 {
-	const FVector2D GeometryDrawSize = SomeGeometry.GetDrawSize();
-	return FVector2D(
+	const FVector2f GeometryDrawSize = SomeGeometry.GetDrawSize();
+	return FVector2f(
 		SomeGeometry.AbsolutePosition.X + GeometryDrawSize.X,
-		SomeGeometry.AbsolutePosition.Y + GeometryDrawSize.Y/2 );
+		SomeGeometry.AbsolutePosition.Y + GeometryDrawSize.Y/2.f );
 }
 
-FVector2D FGeometryHelper::CenterOf(const FGeometry& SomeGeometry)
+UE::Slate::FDeprecateVector2DResult FGeometryHelper::CenterOf(const FGeometry& SomeGeometry)
 {
-	const FVector2D GeometryDrawSize = SomeGeometry.GetDrawSize();
-	return FVector2D(SomeGeometry.AbsolutePosition) + (GeometryDrawSize * 0.5f);
+	const FVector2f GeometryDrawSize = SomeGeometry.GetDrawSize();
+	return FVector2f(SomeGeometry.AbsolutePosition) + (GeometryDrawSize * 0.5f);
 }
 
 void FGeometryHelper::ConvertToPoints(const FGeometry& Geom, TArray<FVector2D>& Points)
 {
-	const FVector2D Size = Geom.GetDrawSize();
-	const FVector2D Location = FVector2D(Geom.AbsolutePosition);
+	TArray<FVector2f> PointsAsFloats;
+	ConvertToPoints(Geom, PointsAsFloats);
+
+	Algo::ForEach(PointsAsFloats, [&Points](const FVector2f& PointAsFloat)
+	{
+		Points.Emplace(FVector2D(PointAsFloat));
+	});
+}
+
+void FGeometryHelper::ConvertToPoints(const FGeometry& Geom, TArray<FVector2f>& Points)
+{
+	const FVector2f Size = Geom.GetDrawSize();
+	const FVector2f Location = FVector2f(Geom.AbsolutePosition);
 
 	int32 Index = Points.AddUninitialized(4);
 	Points[Index++] = Location;
-	Points[Index++] = Location + FVector2D(0.0f, Size.Y);
-	Points[Index++] = Location + FVector2D(Size.X, Size.Y);
-	Points[Index++] = Location + FVector2D(Size.X, 0.0f);
+	Points[Index++] = Location + FVector2f(0.0f, Size.Y);
+	Points[Index++] = Location + FVector2f(Size.X, Size.Y);
+	Points[Index++] = Location + FVector2f(Size.X, 0.0f);
 }
 
 /** Find the point on line segment from LineStart to LineEnd which is closest to Point */
-FVector2D FGeometryHelper::FindClosestPointOnLine(const FVector2D& LineStart, const FVector2D& LineEnd, const FVector2D& TestPoint)
+UE::Slate::FDeprecateVector2DResult FGeometryHelper::FindClosestPointOnLine(const UE::Slate::FDeprecateVector2DParameter& LineStart, const UE::Slate::FDeprecateVector2DParameter& LineEnd, const UE::Slate::FDeprecateVector2DParameter& TestPoint)
 {
-	const FVector2D LineVector = LineEnd - LineStart;
+	const FVector2f LineVector = LineEnd - LineStart;
 
-	const float A = -FVector2D::DotProduct(LineStart - TestPoint, LineVector);
+	const float A = -FVector2f::DotProduct(LineStart - TestPoint, LineVector);
 	const float B = LineVector.SizeSquared();
-	const float T = FMath::Clamp<float>(A / B, 0.0f, 1.0f);
+	const float T = FMath::Clamp(A / B, 0.0f, 1.0f);
 
 	// Generate closest point
 	return LineStart + (T * LineVector);
 }
 
-FVector2D FGeometryHelper::FindClosestPointOnGeom(const FGeometry& Geom, const FVector2D& TestPoint)
+UE::Slate::FDeprecateVector2DResult FGeometryHelper::FindClosestPointOnGeom(const FGeometry& Geom, const UE::Slate::FDeprecateVector2DParameter& TestPoint)
 {
-	TArray<FVector2D> Points;
+	TArray<FVector2f> Points;
 	FGeometryHelper::ConvertToPoints(Geom, Points);
 
 	float BestDistanceSquared = MAX_FLT;
-	FVector2D BestPoint = FVector2D::ZeroVector;
+	FVector2f BestPoint = FVector2f::ZeroVector;
 	for (int32 i = 0; i < Points.Num(); ++i)
 	{
-		const FVector2D Candidate = FindClosestPointOnLine(Points[i], Points[(i + 1) % Points.Num()], TestPoint);
+		const FVector2f Candidate = FindClosestPointOnLine(Points[i], Points[(i + 1) % Points.Num()], TestPoint);
 		const float CandidateDistanceSquared = (Candidate-TestPoint).SizeSquared();
 		if (CandidateDistanceSquared < BestDistanceSquared)
 		{
@@ -94,13 +105,23 @@ FConnectionDrawingPolicy::FConnectionDrawingPolicy(int32 InBackLayerID, int32 In
 	ArrowImage = FAppStyle::GetBrush( TEXT("Graph.Arrow") );
 	ArrowRadius = ArrowImage->ImageSize * ZoomFactor * 0.5f;
 	MidpointImage = nullptr;
-	MidpointRadius = FVector2D::ZeroVector;
+	MidpointRadius = FVector2f::ZeroVector;
 	HoverDeemphasisDarkFraction = 0.8f;
 
 	BubbleImage = FAppStyle::GetBrush( TEXT("Graph.ExecutionBubble") );
 }
 
 void FConnectionDrawingPolicy::DrawSplineWithArrow(const FVector2D& StartPoint, const FVector2D& EndPoint, const FConnectionParams& Params)
+{
+	DrawSplineWithArrow_DeprecationHelper(StartPoint, EndPoint, Params);
+}
+
+void FConnectionDrawingPolicy::DrawSplineWithArrow(const FVector2f& StartPoint, const FVector2f& EndPoint, const FConnectionParams& Params)
+{
+	DrawSplineWithArrow_DeprecationHelper(StartPoint, EndPoint, Params);
+}
+
+void FConnectionDrawingPolicy::DrawSplineWithArrow_DeprecationHelper(const UE::Slate::FDeprecateVector2DParameter& StartPoint, const UE::Slate::FDeprecateVector2DParameter& EndPoint, const FConnectionParams& Params)
 {
 	// Draw the spline
 	DrawConnection(
@@ -112,7 +133,7 @@ void FConnectionDrawingPolicy::DrawSplineWithArrow(const FVector2D& StartPoint,
 	// Draw the arrow
 	if (ArrowImage != nullptr)
 	{
-		FVector2D ArrowPoint = EndPoint - ArrowRadius;
+		FVector2f ArrowPoint = EndPoint - ArrowRadius;
 
 		FSlateDrawElement::MakeBox(
 			DrawElementsList,
@@ -121,7 +142,7 @@ void FConnectionDrawingPolicy::DrawSplineWithArrow(const FVector2D& StartPoint,
 			ArrowImage,
 			ESlateDrawEffect::None,
 			Params.WireColor
-			);
+		);
 	}
 }
 
@@ -131,8 +152,8 @@ void FConnectionDrawingPolicy::DrawSplineWithArrow(const FGeometry& StartGeom, c
 	// empty space inside of the pin brush images.
 	const float StartFudgeX = 4.0f;
 	const float EndFudgeX = 4.0f;
-	const FVector2D StartPoint = FGeometryHelper::VerticalMiddleRightOf(StartGeom) - FVector2D(StartFudgeX, 0.0f);
-	const FVector2D EndPoint = FGeometryHelper::VerticalMiddleLeftOf(EndGeom) - FVector2D(ArrowRadius.X - EndFudgeX, 0);
+	const FVector2f StartPoint = FGeometryHelper::VerticalMiddleRightOf(StartGeom) - FVector2f(StartFudgeX, 0.0f);
+	const FVector2f EndPoint = FGeometryHelper::VerticalMiddleLeftOf(EndGeom) - FVector2f(ArrowRadius.X - EndFudgeX, 0);
 
 	DrawSplineWithArrow(StartPoint, EndPoint, Params);
 }
@@ -176,7 +197,7 @@ void FConnectionDrawingPolicy::SetHoveredPins(const TSet< FEdGraphPinReference >
 	}
 }
 
-void FConnectionDrawingPolicy::SetMousePosition(const FVector2D& InMousePos)
+void FConnectionDrawingPolicy::SetMousePosition(const UE::Slate::FDeprecateVector2DParameter& InMousePos)
 {
 	LocalMousePosition = InMousePos;
 }
@@ -201,7 +222,7 @@ void FConnectionDrawingPolicy::SetMarkedPin(TWeakPtr<SGraphPin> InMarkedPin)
 }
 
 /** Util to make a 'distance->alpha' table and also return spline length */
-float FConnectionDrawingPolicy::MakeSplineReparamTable(const FVector2D& P0, const FVector2D& P0Tangent, const FVector2D& P1, const FVector2D& P1Tangent, FInterpCurve<float>& OutReparamTable)
+float FConnectionDrawingPolicy::MakeSplineReparamTable(const UE::Slate::FDeprecateVector2DParameter& P0, const UE::Slate::FDeprecateVector2DParameter& P0Tangent, const UE::Slate::FDeprecateVector2DParameter& P1, const UE::Slate::FDeprecateVector2DParameter& P1Tangent, FInterpCurve<float>& OutReparamTable)
 {
 	const int32 NumSteps = 10; // TODO: Make this adaptive...
 
@@ -213,7 +234,7 @@ float FConnectionDrawingPolicy::MakeSplineReparamTable(const FVector2D& P0, cons
 	const float Interval = (MaxInput - Param)/((float)(NumSteps-1)); 
 
 	// Add first entry, using first point on curve, total distance will be 0
-	FVector2D OldSplinePos = FMath::CubicInterp(P0, P0Tangent, P1, P1Tangent, Param);
+	FVector2f OldSplinePos = FMath::CubicInterp(P0, P0Tangent, P1, P1Tangent, Param);
 	float TotalDist = 0.f;
 	OutReparamTable.AddPoint(TotalDist, Param);
 	Param += Interval;
@@ -222,7 +243,7 @@ float FConnectionDrawingPolicy::MakeSplineReparamTable(const FVector2D& P0, cons
 	for(int32 i=1; i<NumSteps; i++)
 	{
 		// Iterate along spline at regular param intervals
-		const FVector2D NewSplinePos = FMath::CubicInterp(P0, P0Tangent, P1, P1Tangent, Param);
+		const FVector2f NewSplinePos = FMath::CubicInterp(P0, P0Tangent, P1, P1Tangent, Param);
 		TotalDist += (NewSplinePos - OldSplinePos).Size();
 		OldSplinePos = NewSplinePos;
 
@@ -235,19 +256,39 @@ float FConnectionDrawingPolicy::MakeSplineReparamTable(const FVector2D& P0, cons
 }
 
 FVector2D FConnectionDrawingPolicy::ComputeSplineTangent(const FVector2D& Start, const FVector2D& End) const
+{
+	return ComputeSplineTangent_DeprecationHelper(Start, End);
+}
+
+FVector2f FConnectionDrawingPolicy::ComputeSplineTangent(const FVector2f& Start, const FVector2f& End) const
+{
+	return ComputeSplineTangent_DeprecationHelper(Start, End);
+}
+
+UE::Slate::FDeprecateVector2DResult FConnectionDrawingPolicy::ComputeSplineTangent_DeprecationHelper(const UE::Slate::FDeprecateVector2DParameter& Start, const UE::Slate::FDeprecateVector2DParameter& End) const
 {
 	return Settings->ComputeSplineTangent(Start, End);
 }
 
 void FConnectionDrawingPolicy::DrawConnection(int32 LayerId, const FVector2D& Start, const FVector2D& End, const FConnectionParams& Params)
 {
-	const FVector2D& P0 = Start;
-	const FVector2D& P1 = End;
+	DrawConnection_DeprecationHelper(LayerId, Start, End, Params);
+}
+
+void FConnectionDrawingPolicy::DrawConnection(int32 LayerId, const FVector2f& Start, const FVector2f& End, const FConnectionParams& Params)
+{
+	DrawConnection_DeprecationHelper(LayerId, Start, End, Params);
+}
+
+void FConnectionDrawingPolicy::DrawConnection_DeprecationHelper(int32 LayerId, const UE::Slate::FDeprecateVector2DParameter& Start, const UE::Slate::FDeprecateVector2DParameter& End, const FConnectionParams& Params)
+{
+	const FVector2f& P0 = Start;
+	const FVector2f& P1 = End;
 
-	const FVector2D SplineTangent = ComputeSplineTangent(P0, P1);
+	const FVector2f SplineTangent = ComputeSplineTangent(P0, P1);
 
-	const FVector2D P0Tangent = (Params.StartTangent.IsNearlyZero()) ? ((Params.StartDirection == EGPD_Output) ? SplineTangent : -SplineTangent) : Params.StartTangent;
-	const FVector2D P1Tangent = (Params.EndTangent.IsNearlyZero()) ? ((Params.EndDirection == EGPD_Input) ? SplineTangent : -SplineTangent) : Params.EndTangent;
+	const FVector2f P0Tangent = (Params.StartTangent.IsNearlyZero()) ? ((Params.StartDirection == EGPD_Output) ? SplineTangent : -SplineTangent) : Params.StartTangent;
+	const FVector2f P1Tangent = (Params.EndTangent.IsNearlyZero()) ? ((Params.EndDirection == EGPD_Input) ? SplineTangent : -SplineTangent) : Params.EndTangent;
 
 	if (Settings->bTreatSplinesLikePins)
 	{
@@ -263,27 +304,27 @@ void FConnectionDrawingPolicy::DrawConnection(int32 LayerId, const FVector2D& St
 			// The curve will include the endpoints but can extend out of a tight bounds because of the tangents
 			// P0Tangent coefficient maximizes to 4/27 at a=1/3, and P1Tangent minimizes to -4/27 at a=2/3.
 			const float MaximumTangentContribution = 4.0f / 27.0f;
-			FBox2D Bounds(ForceInit);
+			FBox2f Bounds(ForceInit);
 
-			Bounds += FVector2D(P0);
-			Bounds += FVector2D(P0 + MaximumTangentContribution * P0Tangent);
-			Bounds += FVector2D(P1);
-			Bounds += FVector2D(P1 - MaximumTangentContribution * P1Tangent);
+			Bounds += FVector2f(P0);
+			Bounds += FVector2f(P0 + MaximumTangentContribution * P0Tangent);
+			Bounds += FVector2f(P1);
+			Bounds += FVector2f(P1 - MaximumTangentContribution * P1Tangent);
 
 			bCloseToSpline = Bounds.ComputeSquaredDistanceToPoint(LocalMousePosition) < QueryDistanceForCloseSquared;
 
 			// Draw the bounding box for debugging
 #if 0
-#define DrawSpaceLine(Point1, Point2, DebugWireColor) {const FVector2D FakeTangent = (Point2 - Point1).GetSafeNormal(); FSlateDrawElement::MakeDrawSpaceSpline(DrawElementsList, LayerId, Point1, FakeTangent, Point2, FakeTangent, ClippingRect, 1.0f, ESlateDrawEffect::None, DebugWireColor); }
+#define DrawSpaceLine(Point1, Point2, DebugWireColor) {const FVector2f FakeTangent = (Point2 - Point1).GetSafeNormal(); FSlateDrawElement::MakeDrawSpaceSpline(DrawElementsList, LayerId, Point1, FakeTangent, Point2, FakeTangent, ClippingRect, 1.0f, ESlateDrawEffect::None, DebugWireColor); }
 
 			if (bCloseToSpline)
 			{
 				const FLinearColor BoundsWireColor = bCloseToSpline ? FLinearColor::Green : FLinearColor::White;
 
-				FVector2D TL = Bounds.Min;
-				FVector2D BR = Bounds.Max;
-				FVector2D TR = FVector2D(Bounds.Max.X, Bounds.Min.Y);
-				FVector2D BL = FVector2D(Bounds.Min.X, Bounds.Max.Y);
+				FVector2f TL = Bounds.Min;
+				FVector2f BR = Bounds.Max;
+				FVector2f TR = FVector2f(Bounds.Max.X, Bounds.Min.Y);
+				FVector2f BL = FVector2f(Bounds.Min.X, Bounds.Max.Y);
 
 				DrawSpaceLine(TL, TR, BoundsWireColor);
 				DrawSpaceLine(TR, BR, BoundsWireColor);
@@ -296,17 +337,17 @@ void FConnectionDrawingPolicy::DrawConnection(int32 LayerId, const FVector2D& St
 		if (bCloseToSpline)
 		{
 			// Find the closest approach to the spline
-			FVector2D ClosestPoint(ForceInit);
+			FVector2f ClosestPoint(ForceInit);
 			float ClosestDistanceSquared = FLT_MAX;
 
 			const int32 NumStepsToTest = 16;
 			const float StepInterval = 1.0f / (float)NumStepsToTest;
-			FVector2D Point1 = FMath::CubicInterp(P0, P0Tangent, P1, P1Tangent, 0.0f);
+			FVector2f Point1 = FMath::CubicInterp(P0, P0Tangent, P1, P1Tangent, 0.0f);
 			for (float TestAlpha = 0.0f; TestAlpha < 1.0f; TestAlpha += StepInterval)
 			{
-				const FVector2D Point2 = FMath::CubicInterp(P0, P0Tangent, P1, P1Tangent, TestAlpha + StepInterval);
+				const FVector2f Point2 = FMath::CubicInterp(P0, P0Tangent, P1, P1Tangent, TestAlpha + StepInterval);
 
-				const FVector2D ClosestPointToSegment = FMath::ClosestPointOnSegment2D(LocalMousePosition, Point1, Point2);
+				const FVector2f ClosestPointToSegment = FMath::ClosestPointOnSegment2D(LocalMousePosition, Point1, Point2);
 				const float DistanceSquared = (LocalMousePosition - ClosestPointToSegment).SizeSquared();
 
 				if (DistanceSquared < ClosestDistanceSquared)
@@ -358,28 +399,28 @@ void FConnectionDrawingPolicy::DrawConnection(int32 LayerId, const FVector2D& St
 		{
 			const float BubbleSpacing = 64.f * ZoomFactor;
 			const float BubbleSpeed = 192.f * ZoomFactor;
-			const FVector2D BubbleSize = BubbleImage->ImageSize * ZoomFactor * 0.2f * Params.WireThickness;
+			const FVector2f BubbleSize = BubbleImage->ImageSize * ZoomFactor * 0.2f * Params.WireThickness;
 
-			float Time = (FPlatformTime::Seconds() - GStartTime);
+			float Time = static_cast<float>(FPlatformTime::Seconds() - GStartTime);
 			const float BubbleOffset = FMath::Fmod(Time * BubbleSpeed, BubbleSpacing);
-			const int32 NumBubbles = FMath::CeilToInt(SplineLength/BubbleSpacing);
+			const int32 NumBubbles = FMath::CeilToInt(SplineLength / BubbleSpacing);
 			for (int32 i = 0; i < NumBubbles; ++i)
 			{
 				const float Distance = ((float)i * BubbleSpacing) + BubbleOffset;
 				if (Distance < SplineLength)
 				{
 					const float Alpha = SplineReparamTable.Eval(Distance, 0.f);
-					FVector2D BubblePos = FMath::CubicInterp(P0, P0Tangent, P1, P1Tangent, Alpha);
+					FVector2f BubblePos = FMath::CubicInterp(P0, P0Tangent, P1, P1Tangent, Alpha);
 					BubblePos -= (BubbleSize * 0.5f);
 
 					FSlateDrawElement::MakeBox(
 						DrawElementsList,
 						LayerId,
-						FPaintGeometry( BubblePos, BubbleSize, ZoomFactor  ),
+						FPaintGeometry(BubblePos, BubbleSize, ZoomFactor),
 						BubbleImage,
 						ESlateDrawEffect::None,
 						Params.WireColor
-						);
+					);
 				}
 			}
 		}
@@ -389,15 +430,15 @@ void FConnectionDrawingPolicy::DrawConnection(int32 LayerId, const FVector2D& St
 		{
 			// Determine the spline position for the midpoint
 			const float MidpointAlpha = SplineReparamTable.Eval(SplineLength * 0.5f, 0.f);
-			const FVector2D Midpoint = FMath::CubicInterp(P0, P0Tangent, P1, P1Tangent, MidpointAlpha);
+			const FVector2f Midpoint = FMath::CubicInterp(P0, P0Tangent, P1, P1Tangent, MidpointAlpha);
 
 			// Approximate the slope at the midpoint (to orient the midpoint image to the spline)
-			const FVector2D MidpointPlusE = FMath::CubicInterp(P0, P0Tangent, P1, P1Tangent, MidpointAlpha + KINDA_SMALL_NUMBER);
-			const FVector2D MidpointMinusE = FMath::CubicInterp(P0, P0Tangent, P1, P1Tangent, MidpointAlpha - KINDA_SMALL_NUMBER);
-			const FVector2D SlopeUnnormalized = MidpointPlusE - MidpointMinusE;
+			const FVector2f MidpointPlusE = FMath::CubicInterp(P0, P0Tangent, P1, P1Tangent, MidpointAlpha + KINDA_SMALL_NUMBER);
+			const FVector2f MidpointMinusE = FMath::CubicInterp(P0, P0Tangent, P1, P1Tangent, MidpointAlpha - KINDA_SMALL_NUMBER);
+			const FVector2f SlopeUnnormalized = MidpointPlusE - MidpointMinusE;
 
 			// Draw the arrow
-			const FVector2D MidpointDrawPos = Midpoint - MidpointRadius;
+			const FVector2f MidpointDrawPos = Midpoint - MidpointRadius;
 			const float AngleInRadians = SlopeUnnormalized.IsNearlyZero() ? 0.0f : FMath::Atan2(SlopeUnnormalized.Y, SlopeUnnormalized.X);
 
 			FSlateDrawElement::MakeRotatedBox(
@@ -407,16 +448,25 @@ void FConnectionDrawingPolicy::DrawConnection(int32 LayerId, const FVector2D& St
 				MidpointImage,
 				ESlateDrawEffect::None,
 				AngleInRadians,
-				TOptional<FVector2D>(),
+				TOptional<FVector2f>(),
 				FSlateDrawElement::RelativeToElement,
 				Params.WireColor
-				);
+			);
 		}
 	}
-
 }
 
 void FConnectionDrawingPolicy::DrawPreviewConnector(const FGeometry& PinGeometry, const FVector2D& StartPoint, const FVector2D& EndPoint, UEdGraphPin* Pin)
+{
+	DrawPreviewConnector_DeprecationHelper(PinGeometry, StartPoint, EndPoint, Pin);
+}
+
+void FConnectionDrawingPolicy::DrawPreviewConnector(const FGeometry& PinGeometry, const FVector2f& StartPoint, const FVector2f& EndPoint, UEdGraphPin* Pin)
+{
+	DrawPreviewConnector_DeprecationHelper(PinGeometry, StartPoint, EndPoint, Pin);
+}
+
+void FConnectionDrawingPolicy::DrawPreviewConnector_DeprecationHelper(const FGeometry& PinGeometry, const UE::Slate::FDeprecateVector2DParameter& StartPoint, const UE::Slate::FDeprecateVector2DParameter& EndPoint, UEdGraphPin* Pin)
 {
 	FConnectionParams Params;
 	DetermineWiringStyle(Pin, nullptr, /*inout*/ Params);
diff --git a/Engine/Source/Editor/GraphEditor/Private/GraphEditorSettings.cpp b/Engine/Source/Editor/GraphEditor/Private/GraphEditorSettings.cpp
--- a/Engine/Source/Editor/GraphEditor/Private/GraphEditorSettings.cpp
+++ b/Engine/Source/Editor/GraphEditor/Private/GraphEditorSettings.cpp
@@ -115,36 +115,36 @@ void UGraphEditorSettings::PostEditChangeProperty(FPropertyChangedEvent& Propert
 
 FMargin UGraphEditorSettings::GetInputPinPadding() const
 {
-	const float HPad = FMath::Max<float>(PaddingTowardsNodeEdge, 0.0f);
+	const float HPad = FMath::Max(PaddingTowardsNodeEdge, 0.0f);
 	return FMargin(HPad, PaddingAbovePin, PaddingRightOfInput, PaddingBelowPin);
 }
 
 FMargin UGraphEditorSettings::GetOutputPinPadding() const
 {
-	const float HPad = FMath::Max<float>(PaddingTowardsNodeEdge, 0.0f);
+	const float HPad = FMath::Max(PaddingTowardsNodeEdge, 0.0f);
 	return FMargin(PaddingLeftOfOutput, PaddingAbovePin, HPad, PaddingBelowPin);
 }
 
 FMargin UGraphEditorSettings::GetNonPinNodeBodyPadding() const
 {
-	const float NegativeHPad = FMath::Max<float>(-PaddingTowardsNodeEdge, 0.0f);
+	const float NegativeHPad = FMath::Max(-PaddingTowardsNodeEdge, 0.0f);
 	return FMargin(NegativeHPad, 0.0f, NegativeHPad, 0.0f);
 }
 
-FVector2D UGraphEditorSettings::GetShadowDeltaSize() const
+UE::Slate::FDeprecateVector2DResult UGraphEditorSettings::GetShadowDeltaSize() const
 {
-	FVector2D ShadowSize = FAppStyle::GetVector(TEXT("Graph.Node.ShadowSize"));
-	ShadowSize.X += FMath::Min<float>(PaddingTowardsNodeEdge, 0.0f);
+	FVector2f ShadowSize = FAppStyle::GetVector(TEXT("Graph.Node.ShadowSize"));
+	ShadowSize.X += FMath::Min(PaddingTowardsNodeEdge, 0.0f);
 	return ShadowSize;
 }
 
-FVector2D UGraphEditorSettings::ComputeSplineTangent(const FVector2D& Start, const FVector2D& End) const
+UE::Slate::FDeprecateVector2DResult UGraphEditorSettings::ComputeSplineTangent(const UE::Slate::FDeprecateVector2DParameter& Start, const UE::Slate::FDeprecateVector2DParameter& End) const
 {
-	const FVector2D DeltaPos = End - Start;
+	const FVector2f DeltaPos = End - Start;
 	const bool bGoingForward = DeltaPos.X >= 0.0f;
 
-	const float ClampedTensionX = FMath::Min<float>(FMath::Abs<float>(DeltaPos.X), bGoingForward ? ForwardSplineHorizontalDeltaRange : BackwardSplineHorizontalDeltaRange);
-	const float ClampedTensionY = FMath::Min<float>(FMath::Abs<float>(DeltaPos.Y), bGoingForward ? ForwardSplineVerticalDeltaRange : BackwardSplineVerticalDeltaRange);
+	const float ClampedTensionX = FMath::Min(FMath::Abs(DeltaPos.X), bGoingForward ? ForwardSplineHorizontalDeltaRange : BackwardSplineHorizontalDeltaRange);
+	const float ClampedTensionY = FMath::Min(FMath::Abs(DeltaPos.Y), bGoingForward ? ForwardSplineVerticalDeltaRange : BackwardSplineVerticalDeltaRange);
 
 	if (bGoingForward)
 	{
diff --git a/Engine/Source/Editor/GraphEditor/Private/SGraphPanel.cpp b/Engine/Source/Editor/GraphEditor/Private/SGraphPanel.cpp
--- a/Engine/Source/Editor/GraphEditor/Private/SGraphPanel.cpp
+++ b/Engine/Source/Editor/GraphEditor/Private/SGraphPanel.cpp
@@ -565,13 +565,13 @@ int32 SGraphPanel::OnPaint( const FPaintArgs& Args, const FGeometry& AllottedGeo
 
 				if (PinGeometry != nullptr)
 				{
-					FVector2D StartPoint;
-					FVector2D EndPoint;
+					FVector2f StartPoint;
+					FVector2f EndPoint;
 
 					if (CurrentStartPin->GetDirection() == EGPD_Input)
 					{
 						StartPoint = AllottedGeometry.LocalToAbsolute(PreviewConnectorEndpoint);
-						EndPoint = FGeometryHelper::VerticalMiddleLeftOf( PinGeometry->Geometry ) - FVector2D(ConnectionDrawingPolicy->ArrowRadius.X, 0);
+						EndPoint = FGeometryHelper::VerticalMiddleLeftOf( PinGeometry->Geometry ) - FVector2f(ConnectionDrawingPolicy->ArrowRadius.X, 0);
 					}
 					else
 					{
diff --git a/Engine/Source/Editor/GraphEditor/Public/ConnectionDrawingPolicy.h b/Engine/Source/Editor/GraphEditor/Public/ConnectionDrawingPolicy.h
--- a/Engine/Source/Editor/GraphEditor/Public/ConnectionDrawingPolicy.h
+++ b/Engine/Source/Editor/GraphEditor/Public/ConnectionDrawingPolicy.h
@@ -40,15 +40,18 @@ GRAPHEDITOR_API DECLARE_LOG_CATEGORY_EXTERN(LogConnectionDrawingPolicy, Log, All
 class GRAPHEDITOR_API FGeometryHelper
 {
 public:
-	static FVector2D VerticalMiddleLeftOf(const FGeometry& SomeGeometry);
-	static FVector2D VerticalMiddleRightOf(const FGeometry& SomeGeometry);
-	static FVector2D CenterOf(const FGeometry& SomeGeometry);
+	static UE::Slate::FDeprecateVector2DResult VerticalMiddleLeftOf(const FGeometry& SomeGeometry);
+	static UE::Slate::FDeprecateVector2DResult VerticalMiddleRightOf(const FGeometry& SomeGeometry);
+
+	static UE::Slate::FDeprecateVector2DResult CenterOf(const FGeometry& SomeGeometry);
+	UE_DEPRECATED(5.6, "Use Array of FVector2f instead of FVector2D for slate positions")
 	static void ConvertToPoints(const FGeometry& Geom, TArray<FVector2D>& Points);
+	static void ConvertToPoints(const FGeometry& Geom, TArray<FVector2f>& Points);
 
 	/** Find the point on line segment from LineStart to LineEnd which is closest to Point */
-	static FVector2D FindClosestPointOnLine(const FVector2D& LineStart, const FVector2D& LineEnd, const FVector2D& TestPoint);
+	static UE::Slate::FDeprecateVector2DResult FindClosestPointOnLine(const UE::Slate::FDeprecateVector2DParameter& LineStart, const UE::Slate::FDeprecateVector2DParameter& LineEnd, const UE::Slate::FDeprecateVector2DParameter& TestPoint);
 
-	static FVector2D FindClosestPointOnGeom(const FGeometry& Geom, const FVector2D& TestPoint);
+	static UE::Slate::FDeprecateVector2DResult FindClosestPointOnGeom(const FGeometry& Geom, const UE::Slate::FDeprecateVector2DParameter& TestPoint);
 };
 
 /////////////////////////////////////////////////////
@@ -68,8 +71,8 @@ struct GRAPHEDITOR_API FConnectionParams
 	EEdGraphPinDirection StartDirection;
 	EEdGraphPinDirection EndDirection;
 
-	FVector2D StartTangent;
-	FVector2D EndTangent;
+	FDeprecateSlateVector2D StartTangent;
+	FDeprecateSlateVector2D EndTangent;
 
 	FConnectionParams()
 		: WireColor(FLinearColor::White)
@@ -81,12 +84,14 @@ struct GRAPHEDITOR_API FConnectionParams
 		, bUserFlag2(false)
 		, StartDirection(EGPD_Output)
 		, EndDirection(EGPD_Input)
-		, StartTangent(FVector2D::ZeroVector)
-		, EndTangent(FVector2D::ZeroVector)
+		, StartTangent(FVector2f::ZeroVector)
+		, EndTangent(FVector2f::ZeroVector)
 	{
 	}
 };
 
+#define END_DEPRECATED_VIRTUAL_FN
+
 /////////////////////////////////////////////////////
 // FConnectionDrawingPolicy
 
@@ -104,8 +109,8 @@ protected:
 	const UGraphEditorSettings* Settings;
 
 public:
-	FVector2D ArrowRadius;
-	FVector2D MidpointRadius;
+	FDeprecateSlateVector2D ArrowRadius;
+	FDeprecateSlateVector2D MidpointRadius;
 
 	FGraphSplineOverlapResult SplineOverlapResult;
 
@@ -125,7 +130,7 @@ protected:
 	TSet< FEdGraphPinReference > HoveredPins;
 	TMap<TSharedRef<SWidget>, FArrangedWidget>* PinGeometries;
 	double LastHoverTimeEvent;
-	FVector2D LocalMousePosition;
+	FDeprecateSlateVector2D LocalMousePosition;
 
 	/** List of currently relinked connections. */
 	TArray<FRelinkConnection> RelinkConnections;
@@ -140,7 +145,7 @@ public:
 	// Update the drawing policy with the set of hovered pins (which can be empty)
 	void SetHoveredPins(const TSet< FEdGraphPinReference >& InHoveredPins, const TArray< TSharedPtr<SGraphPin> >& OverridePins, double HoverTime);
 
-	void SetMousePosition(const FVector2D& InMousePos);
+	void SetMousePosition(const UE::Slate::FDeprecateVector2DParameter& InMousePos);
 
 	// Update the drawing policy with the marked pin (which may not be valid)
 	void SetMarkedPin(TWeakPtr<SGraphPin> InMarkedPin);
@@ -151,15 +156,26 @@ public:
 	// Set the list of currently relinked connections.
 	void SetRelinkConnections(const TArray<FRelinkConnection>& Connections) { RelinkConnections = Connections; }
 
-	static float MakeSplineReparamTable(const FVector2D& P0, const FVector2D& P0Tangent, const FVector2D& P1, const FVector2D& P1Tangent, FInterpCurve<float>& OutReparamTable);
+	static float MakeSplineReparamTable(const UE::Slate::FDeprecateVector2DParameter& P0, const UE::Slate::FDeprecateVector2DParameter& P0Tangent, const UE::Slate::FDeprecateVector2DParameter& P1, const UE::Slate::FDeprecateVector2DParameter& P1Tangent, FInterpCurve<float>& OutReparamTable);
 
-	virtual void DrawSplineWithArrow(const FVector2D& StartPoint, const FVector2D& EndPoint, const FConnectionParams& Params);
+	UE_DEPRECATED(5.6, "Use the version of the function accepting FVector2f; this Slate API no longer interfaces directly with double-precision scalars and vectors.")
+	virtual void DrawSplineWithArrow(const FVector2D& StartPoint, const FVector2D& EndPoint, const FConnectionParams& Params) UE_SLATE_DEPRECATED_VECTOR_VIRTUAL_FUNCTION;
+	virtual void DrawSplineWithArrow(const FVector2f& StartPoint, const FVector2f& EndPoint, const FConnectionParams& Params);
 	
 	virtual void DrawSplineWithArrow(const FGeometry& StartGeom, const FGeometry& EndGeom, const FConnectionParams& Params);
 
-	virtual FVector2D ComputeSplineTangent(const FVector2D& Start, const FVector2D& End) const;
-	virtual void DrawConnection(int32 LayerId, const FVector2D& Start, const FVector2D& End, const FConnectionParams& Params);
-	virtual void DrawPreviewConnector(const FGeometry& PinGeometry, const FVector2D& StartPoint, const FVector2D& EndPoint, UEdGraphPin* Pin);
+	UE_DEPRECATED(5.6, "Use the version of the function accepting FVector2f; this Slate API no longer interfaces directly with double-precision scalars and vectors.")
+	virtual FVector2D ComputeSplineTangent(const FVector2D& Start, const FVector2D& End) const UE_SLATE_DEPRECATED_VECTOR_VIRTUAL_FUNCTION;
+	virtual FVector2f ComputeSplineTangent(const FVector2f& Start, const FVector2f& End) const;
+
+
+	UE_DEPRECATED(5.6, "Use the version of the function accepting FVector2f; this Slate API no longer interfaces directly with double-precision scalars and vectors.")
+	virtual void DrawConnection(int32 LayerId, const FVector2D& Start, const FVector2D& End, const FConnectionParams& Params) UE_SLATE_DEPRECATED_VECTOR_VIRTUAL_FUNCTION;
+	virtual void DrawConnection(int32 LayerId, const FVector2f& Start, const FVector2f& End, const FConnectionParams& Params);
+
+	UE_DEPRECATED(5.6, "Use the version of the function accepting FVector2f; this Slate API no longer interfaces directly with double-precision scalars and vectors.")
+	virtual void DrawPreviewConnector(const FGeometry& PinGeometry, const FVector2D& StartPoint, const FVector2D& EndPoint, UEdGraphPin* Pin) UE_SLATE_DEPRECATED_VECTOR_VIRTUAL_FUNCTION;
+	virtual void DrawPreviewConnector(const FGeometry& PinGeometry, const FVector2f& StartPoint, const FVector2f& EndPoint, UEdGraphPin* Pin);
 
 	// Give specific editor modes a chance to highlight this connection or darken non-interesting connections
 	virtual void DetermineWiringStyle(UEdGraphPin* OutputPin, UEdGraphPin* InputPin, /*inout*/ FConnectionParams& Params);
@@ -193,4 +209,11 @@ protected:
 
 	// Helper function used by Draw(). Iterates over the pin geometries, drawing connections between them
 	virtual void DrawPinGeometries(TMap<TSharedRef<SWidget>, FArrangedWidget>& InPinGeometries, FArrangedChildren& ArrangedNodes);
+
+private:
+	// Note - these helpers are temporary, to be replaced entirely by the public API's 2f/float functions once the 2D/double deprecated functions can be removed
+	void DrawSplineWithArrow_DeprecationHelper(const UE::Slate::FDeprecateVector2DParameter& StartPoint, const UE::Slate::FDeprecateVector2DParameter& EndPoint, const FConnectionParams& Params);
+	UE::Slate::FDeprecateVector2DResult ComputeSplineTangent_DeprecationHelper(const UE::Slate::FDeprecateVector2DParameter& Start, const UE::Slate::FDeprecateVector2DParameter& End) const;
+	void DrawConnection_DeprecationHelper(int32 LayerId, const UE::Slate::FDeprecateVector2DParameter& Start, const UE::Slate::FDeprecateVector2DParameter& End, const FConnectionParams& Params);
+	void DrawPreviewConnector_DeprecationHelper(const FGeometry& PinGeometry, const UE::Slate::FDeprecateVector2DParameter& StartPoint, const UE::Slate::FDeprecateVector2DParameter& EndPoint, UEdGraphPin* Pin);
 };
diff --git a/Engine/Source/Editor/GraphEditor/Public/GraphEditorSettings.h b/Engine/Source/Editor/GraphEditor/Public/GraphEditorSettings.h
--- a/Engine/Source/Editor/GraphEditor/Public/GraphEditorSettings.h
+++ b/Engine/Source/Editor/GraphEditor/Public/GraphEditorSettings.h
@@ -90,7 +90,7 @@ public:
 	FMargin GetNonPinNodeBodyPadding() const;
 
 	/** Returns the size adjustment from node render geometry to shadow render geometry */
-	FVector2D GetShadowDeltaSize() const;
+	UE::Slate::FDeprecateVector2DResult GetShadowDeltaSize() const;
 
 public:
 
@@ -116,11 +116,11 @@ public:
 
 	/** The amount that the horizontal delta affects the generated tangent handle of splines (when the wire is moving forward) */
 	UPROPERTY(config, EditAnywhere, Category=Splines, AdvancedDisplay)
-	FVector2D ForwardSplineTangentFromHorizontalDelta;
+	FDeprecateSlateVector2D ForwardSplineTangentFromHorizontalDelta;
 
 	/** The amount that the vertical delta affects the generated tangent handle of splines (when the wire is moving forward) */
 	UPROPERTY(config, EditAnywhere, Category=Splines, AdvancedDisplay)
-	FVector2D ForwardSplineTangentFromVerticalDelta;
+	FDeprecateSlateVector2D ForwardSplineTangentFromVerticalDelta;
 
 	/** The maximum value to clamp the absolute value of the horizontal distance between endpoints when calculating tangents (when the wire is moving backwards) */
 	UPROPERTY(config, EditAnywhere, Category=Splines, AdvancedDisplay)
@@ -132,14 +132,14 @@ public:
 
 	/** The amount that the horizontal delta affects the generated tangent handle of splines (when the wire is moving backwards) */
 	UPROPERTY(config, EditAnywhere, Category=Splines, AdvancedDisplay)
-	FVector2D BackwardSplineTangentFromHorizontalDelta;
+	FDeprecateSlateVector2D BackwardSplineTangentFromHorizontalDelta;
 
 	/** The amount that the vertical delta affects the generated tangent handle of splines (when the wire is moving backwards) */
 	UPROPERTY(config, EditAnywhere, Category=Splines, AdvancedDisplay)
-	FVector2D BackwardSplineTangentFromVerticalDelta;
+	FDeprecateSlateVector2D BackwardSplineTangentFromVerticalDelta;
 
 	/** Compute a tangent given the start and end points of the spline */
-	FVector2D ComputeSplineTangent(const FVector2D& Start, const FVector2D& End) const;
+	UE::Slate::FDeprecateVector2DResult ComputeSplineTangent(const UE::Slate::FDeprecateVector2DParameter& Start, const UE::Slate::FDeprecateVector2DParameter& End) const;
 
 public:
 
diff --git a/Engine/Source/Editor/PhysicsAssetEditor/Private/PhysicsAssetGraph/PhysicsAssetConnectionDrawingPolicy.h b/Engine/Source/Editor/PhysicsAssetEditor/Private/PhysicsAssetGraph/PhysicsAssetConnectionDrawingPolicy.h
--- a/Engine/Source/Editor/PhysicsAssetEditor/Private/PhysicsAssetGraph/PhysicsAssetConnectionDrawingPolicy.h
+++ b/Engine/Source/Editor/PhysicsAssetEditor/Private/PhysicsAssetGraph/PhysicsAssetConnectionDrawingPolicy.h
@@ -10,14 +10,14 @@ public:
 	FPhysicsAssetConnectionDrawingPolicy(int32 InBackLayerID, int32 InFrontLayerID, float InZoomFactor, const FSlateRect& InClippingRect, FSlateWindowElementList& InDrawElements)
 		: FConnectionDrawingPolicy(InBackLayerID, InFrontLayerID, InZoomFactor, InClippingRect, InDrawElements)
 	{
-		// We dont draw arrows here
+		// We don't draw arrows here
 		ArrowImage = nullptr;
-		ArrowRadius = FVector2D(0.0f, 0.0f);
+		ArrowRadius = FVector2f::ZeroVector;
 	}
 
-	virtual FVector2D ComputeSplineTangent(const FVector2D& Start, const FVector2D& End) const override
+	virtual FVector2f ComputeSplineTangent(const FVector2f& Start, const FVector2f& End) const override
 	{
 		const int32 Tension = FMath::Abs<int32>(Start.X - End.X);
-		return Tension * FVector2D(1.0f, 0);
+		return Tension * FVector2f(1.0f, 0.0f);
 	}
 };
diff --git a/Engine/Source/Runtime/SlateCore/Private/Input/HittestGrid.cpp b/Engine/Source/Runtime/SlateCore/Private/Input/HittestGrid.cpp
--- a/Engine/Source/Runtime/SlateCore/Private/Input/HittestGrid.cpp
+++ b/Engine/Source/Runtime/SlateCore/Private/Input/HittestGrid.cpp
@@ -49,21 +49,21 @@ FVector2f ClosestPointOnSlateRotatedRect(const FVector2f &Point, const FSlateRot
 	}
 
 	const static int32 NumOfCorners = 4;
-	FVector2d Corners[NumOfCorners];
-	Corners[0] = FVector2d(RotatedRect.TopLeft);
-	Corners[1] = FVector2d(Corners[0]) + FVector2d(RotatedRect.ExtentX);
-	Corners[2] = FVector2d(Corners[1]) + FVector2d(RotatedRect.ExtentY);
-	Corners[3] = FVector2d(Corners[0]) + FVector2d(RotatedRect.ExtentY);
+	FVector2f Corners[NumOfCorners];
+	Corners[0] = RotatedRect.TopLeft;
+	Corners[1] = Corners[0] + RotatedRect.ExtentX;
+	Corners[2] = Corners[1] + RotatedRect.ExtentY;
+	Corners[3] = Corners[0] + RotatedRect.ExtentY;
 
 	FVector2f RetPoint = FVector2f::ZeroVector;
 	float ClosestDistSq = FLT_MAX;
 	for (int32 i = 0; i < NumOfCorners; ++i)
 	{
 		//grab the closest point along the line segment
-		const FVector2d ClosestPoint = FMath::ClosestPointOnSegment2D(FVector2d(Point), Corners[i], Corners[(i + 1) % NumOfCorners]);
+		const FVector2f ClosestPoint = FMath::ClosestPointOnSegment2D(Point, Corners[i], Corners[(i + 1) % NumOfCorners]);
 
 		//get the distance between the two
-		const float TestDist = FVector2d::DistSquared(FVector2d(Point), ClosestPoint);
+		const float TestDist = FVector2f::DistSquared(Point, ClosestPoint);
 
 		//if the distance is smaller than the current smallest, update our closest
 		if (TestDist < ClosestDistSq)
diff --git a/Engine/Source/Runtime/SlateCore/Private/Styling/SlateStyleSet.cpp b/Engine/Source/Runtime/SlateCore/Private/Styling/SlateStyleSet.cpp
--- a/Engine/Source/Runtime/SlateCore/Private/Styling/SlateStyleSet.cpp
+++ b/Engine/Source/Runtime/SlateCore/Private/Styling/SlateStyleSet.cpp
@@ -199,7 +199,7 @@ float FSlateStyleSet::GetFloat(const FName PropertyName, const ANSICHAR* Specifi
 	return Result ? *Result : DefaultValue;
 }
 
-FVector2D FSlateStyleSet::GetVector(const FName PropertyName, const ANSICHAR* Specifier, FVector2D DefaultValue, const ISlateStyle* RequestingStyle) const
+UE::Slate::FDeprecateVector2DResult FSlateStyleSet::GetVector(const FName PropertyName, const ANSICHAR* Specifier, UE::Slate::FDeprecateVector2DParameter DefaultValue, const ISlateStyle* RequestingStyle) const
 {
 	const FVector2f* const Result = Vector2DValues.Find(Join(PropertyName, Specifier));
 
@@ -218,7 +218,7 @@ FVector2D FSlateStyleSet::GetVector(const FName PropertyName, const ANSICHAR* Sp
 		ReportingStyle->LogMissingResource(EStyleMessageSeverity::Warning, FText::Format(NSLOCTEXT("SlateStyleSet", "UnknownSlateVector", "Unable to find vector property '{0}' in style."), FText::FromName(PropertyName)), PropertyName);
 	}
 
-	return Result ? FVector2D(*Result) : DefaultValue;
+	return Result ? FVector2f(*Result) : DefaultValue;
 }
 
 const FLinearColor& FSlateStyleSet::GetColor(const FName PropertyName, const ANSICHAR* Specifier, const FLinearColor& DefaultValue, const ISlateStyle* RequestingStyle) const
diff --git a/Engine/Source/Runtime/SlateCore/Private/Styling/StyleDefaults.cpp b/Engine/Source/Runtime/SlateCore/Private/Styling/StyleDefaults.cpp
--- a/Engine/Source/Runtime/SlateCore/Private/Styling/StyleDefaults.cpp
+++ b/Engine/Source/Runtime/SlateCore/Private/Styling/StyleDefaults.cpp
@@ -15,8 +15,8 @@ const FSlateFontInfo FStyleDefaults::GetFontInfo(uint16 Size)
 }
 
 float FStyleDefaults::DefaultFloat = 0.0f;
-FVector2D FStyleDefaults::DefaultFVector2D = FVector2D( 0, 0 );
-FLinearColor FStyleDefaults::DefaultColor = FLinearColor( 1, 1, 1 );
+FVector2f FStyleDefaults::DefaultFVector2D = FVector2f::ZeroVector;
+FLinearColor FStyleDefaults::DefaultColor = FLinearColor( 1.0f, 1.0f, 1.0f );
 FMargin FStyleDefaults::DefaultMargin = 0.0f;
 FSlateSound FStyleDefaults::DefaultSound;
 FSlateColor FStyleDefaults::DefaultSlateColor = FSlateColor::UseForeground();
\ No newline at end of file
diff --git a/Engine/Source/Runtime/SlateCore/Public/Styling/AppStyle.h b/Engine/Source/Runtime/SlateCore/Public/Styling/AppStyle.h
--- a/Engine/Source/Runtime/SlateCore/Public/Styling/AppStyle.h
+++ b/Engine/Source/Runtime/SlateCore/Public/Styling/AppStyle.h
@@ -43,7 +43,7 @@ public:
 		return FAppStyle::Get().GetFloat(PropertyName, Specifier);
 	}
 
-	static FVector2D GetVector(FName PropertyName, const ANSICHAR* Specifier = NULL)
+	static UE::Slate::FDeprecateVector2DResult GetVector(FName PropertyName, const ANSICHAR* Specifier = NULL)
 	{
 		return FAppStyle::Get().GetVector(PropertyName, Specifier);
 	}
diff --git a/Engine/Source/Runtime/SlateCore/Public/Styling/ISlateStyle.h b/Engine/Source/Runtime/SlateCore/Public/Styling/ISlateStyle.h
--- a/Engine/Source/Runtime/SlateCore/Public/Styling/ISlateStyle.h
+++ b/Engine/Source/Runtime/SlateCore/Public/Styling/ISlateStyle.h
@@ -93,7 +93,7 @@ public:
 	 *
 	 * @return a FVector2D property.
 	 */
-	virtual FVector2D GetVector(const FName PropertyName, const ANSICHAR* Specifier = nullptr, FVector2D DefaultValue = FStyleDefaults::GetVector2D(), const ISlateStyle* RequestingStyle = nullptr) const = 0;
+	virtual UE::Slate::FDeprecateVector2DResult GetVector(const FName PropertyName, const ANSICHAR* Specifier = nullptr, UE::Slate::FDeprecateVector2DParameter DefaultValue = FStyleDefaults::GetVector2D(), const ISlateStyle* RequestingStyle = nullptr) const = 0;
 
 	/**
 	 * @param PropertyName   Name of the property to get
diff --git a/Engine/Source/Runtime/SlateCore/Public/Styling/SlateStyle.h b/Engine/Source/Runtime/SlateCore/Public/Styling/SlateStyle.h
--- a/Engine/Source/Runtime/SlateCore/Public/Styling/SlateStyle.h
+++ b/Engine/Source/Runtime/SlateCore/Public/Styling/SlateStyle.h
@@ -60,7 +60,7 @@ public:
 	SLATECORE_API virtual FString RootToCoreContentDir(const FString& RelativePath);
 
 	SLATECORE_API virtual float GetFloat(const FName PropertyName, const ANSICHAR* Specifier = nullptr, float DefaultValue = FStyleDefaults::GetFloat(), const ISlateStyle* RequestingStyle = nullptr) const override;
-	SLATECORE_API virtual FVector2D GetVector(const FName PropertyName, const ANSICHAR* Specifier = nullptr, FVector2D DefaultValue = FStyleDefaults::GetVector2D(), const ISlateStyle* RequestingStyle = nullptr) const override;
+	SLATECORE_API virtual UE::Slate::FDeprecateVector2DResult GetVector(const FName PropertyName, const ANSICHAR* Specifier = nullptr, UE::Slate::FDeprecateVector2DParameter DefaultValue = FStyleDefaults::GetVector2D(), const ISlateStyle* RequestingStyle = nullptr) const override;
 	SLATECORE_API virtual const FLinearColor& GetColor(const FName PropertyName, const ANSICHAR* Specifier = nullptr, const FLinearColor& DefaultValue = FStyleDefaults::GetColor(), const ISlateStyle* RequestingStyle = nullptr) const override;
 	SLATECORE_API virtual const FSlateColor GetSlateColor(const FName PropertyName, const ANSICHAR* Specifier = nullptr, const FSlateColor& DefaultValue = FStyleDefaults::GetSlateColor(), const ISlateStyle* RequestingStyle = nullptr) const override;
 	SLATECORE_API virtual const FMargin& GetMargin(const FName PropertyName, const ANSICHAR* Specifier = nullptr, const FMargin& DefaultValue = FStyleDefaults::GetMargin(), const ISlateStyle* RequestingStyle = nullptr) const override;
diff --git a/Engine/Source/Runtime/SlateCore/Public/Styling/StyleDefaults.h b/Engine/Source/Runtime/SlateCore/Public/Styling/StyleDefaults.h
--- a/Engine/Source/Runtime/SlateCore/Public/Styling/StyleDefaults.h
+++ b/Engine/Source/Runtime/SlateCore/Public/Styling/StyleDefaults.h
@@ -25,7 +25,7 @@ public:
 	 * Get default FVector2D value.
 	 * @return - The default FVector2D.
 	 */
-	const static FVector2D GetVector2D()
+	const static UE::Slate::FDeprecateVector2DResult GetVector2D()
 	{
 		return DefaultFVector2D;
 	}
@@ -80,7 +80,7 @@ public:
 private:
 
 	static SLATECORE_API float DefaultFloat;
-	static SLATECORE_API FVector2D DefaultFVector2D;
+	static SLATECORE_API FVector2f DefaultFVector2D;
 	static SLATECORE_API FLinearColor DefaultColor;
 	static SLATECORE_API FMargin DefaultMargin;
 	static SLATECORE_API FSlateSound DefaultSound;
