Fix graph panel not returning early when breaking single pin link on alt+click

diff --git a/Engine/Source/Editor/GraphEditor/Private/SGraphPanel.cpp b/Engine/Source/Editor/GraphEditor/Private/SGraphPanel.cpp
--- a/Engine/Source/Editor/GraphEditor/Private/SGraphPanel.cpp
+++ b/Engine/Source/Editor/GraphEditor/Private/SGraphPanel.cpp
@@ -971,6 +971,7 @@ FReply SGraphPanel::OnMouseButtonDown(const FGeometry& MyGeometry, const FPointe
 		{
 			const UEdGraphSchema* Schema = GraphObj->GetSchema();
 			Schema->BreakSinglePinLink(Pin1, Pin2);
+			return FReply::Handled();
 		}
 		else if (SGraphPin* BestPinFromHoveredSpline = GetBestPinFromHoveredSpline())
 		{
