[Truncation Warnings] Update SNodePanel, SGraphPanel and dependent classes to use FVector2f

diff --git a/Engine/Plugins/AI/EnvironmentQueryEditor/Source/EnvironmentQueryEditor/Private/SGraphNode_EnvironmentQuery.cpp b/Engine/Plugins/AI/EnvironmentQueryEditor/Source/EnvironmentQueryEditor/Private/SGraphNode_EnvironmentQuery.cpp
--- a/Engine/Plugins/AI/EnvironmentQueryEditor/Source/EnvironmentQueryEditor/Private/SGraphNode_EnvironmentQuery.cpp
+++ b/Engine/Plugins/AI/EnvironmentQueryEditor/Source/EnvironmentQueryEditor/Private/SGraphNode_EnvironmentQuery.cpp
@@ -414,8 +414,8 @@ void SGraphNode_EnvironmentQuery::UpdateGraphNode()
 		.IsGraphNodeHovered(this, &SGraphNode::IsHovered);
 
 	GetOrAddSlot(ENodeZone::TopCenter)
-		.SlotOffset(TAttribute<FVector2D>(CommentBubble.Get(), &SCommentBubble::GetOffset))
-		.SlotSize(TAttribute<FVector2D>(CommentBubble.Get(), &SCommentBubble::GetSize))
+		.SlotOffset2f(TAttribute<FVector2f>(CommentBubble.Get(), &SCommentBubble::GetOffset2f))
+		.SlotSize2f(TAttribute<FVector2f>(CommentBubble.Get(), &SCommentBubble::GetSize2f))
 		.AllowScaling(TAttribute<bool>(CommentBubble.Get(), &SCommentBubble::IsScalingAllowed))
 		.VAlign(VAlign_Top)
 		[
diff --git a/Engine/Plugins/Cameras/GameplayCameras/Source/GameplayCamerasEditor/Private/Editors/SObjectTreeGraphCommentNode.cpp b/Engine/Plugins/Cameras/GameplayCameras/Source/GameplayCamerasEditor/Private/Editors/SObjectTreeGraphCommentNode.cpp
--- a/Engine/Plugins/Cameras/GameplayCameras/Source/GameplayCamerasEditor/Private/Editors/SObjectTreeGraphCommentNode.cpp
+++ b/Engine/Plugins/Cameras/GameplayCameras/Source/GameplayCamerasEditor/Private/Editors/SObjectTreeGraphCommentNode.cpp
@@ -12,7 +12,7 @@ void SObjectTreeGraphCommentNode::Construct(const FArguments& InArgs)
 	ObjectGraphNode = InArgs._GraphNode;
 }
 
-void SObjectTreeGraphCommentNode::MoveTo(const FVector2D& NewPosition, FNodeSet& NodeFilter, bool bMarkDirty)
+void SObjectTreeGraphCommentNode::MoveTo(const FVector2f& NewPosition, FNodeSet& NodeFilter, bool bMarkDirty)
 {
 	SGraphNodeComment::MoveTo(NewPosition, NodeFilter, bMarkDirty);
 
diff --git a/Engine/Plugins/Cameras/GameplayCameras/Source/GameplayCamerasEditor/Private/Editors/SObjectTreeGraphNode.cpp b/Engine/Plugins/Cameras/GameplayCameras/Source/GameplayCamerasEditor/Private/Editors/SObjectTreeGraphNode.cpp
--- a/Engine/Plugins/Cameras/GameplayCameras/Source/GameplayCamerasEditor/Private/Editors/SObjectTreeGraphNode.cpp
+++ b/Engine/Plugins/Cameras/GameplayCameras/Source/GameplayCamerasEditor/Private/Editors/SObjectTreeGraphNode.cpp
@@ -27,7 +27,7 @@ void SObjectTreeGraphNode::Construct(const FArguments& InArgs)
 	UpdateGraphNode();
 }
 
-void SObjectTreeGraphNode::MoveTo(const FVector2D& NewPosition, FNodeSet& NodeFilter, bool bMarkDirty)
+void SObjectTreeGraphNode::MoveTo(const FVector2f& NewPosition, FNodeSet& NodeFilter, bool bMarkDirty)
 {
 	SGraphNode::MoveTo(NewPosition, NodeFilter, bMarkDirty);
 
diff --git a/Engine/Plugins/Cameras/GameplayCameras/Source/GameplayCamerasEditor/Public/Editors/SObjectTreeGraphCommentNode.h b/Engine/Plugins/Cameras/GameplayCameras/Source/GameplayCamerasEditor/Public/Editors/SObjectTreeGraphCommentNode.h
--- a/Engine/Plugins/Cameras/GameplayCameras/Source/GameplayCamerasEditor/Public/Editors/SObjectTreeGraphCommentNode.h
+++ b/Engine/Plugins/Cameras/GameplayCameras/Source/GameplayCamerasEditor/Public/Editors/SObjectTreeGraphCommentNode.h
@@ -26,7 +26,7 @@ public:
 public:
 
 	// SNodePanel::SNode interface.
-	virtual void MoveTo(const FVector2D& NewPosition, FNodeSet& NodeFilter, bool bMarkDirty=true) override;
+	virtual void MoveTo(const FVector2f& NewPosition, FNodeSet& NodeFilter, bool bMarkDirty=true) override;
 
 protected:
 
diff --git a/Engine/Plugins/Cameras/GameplayCameras/Source/GameplayCamerasEditor/Public/Editors/SObjectTreeGraphNode.h b/Engine/Plugins/Cameras/GameplayCameras/Source/GameplayCamerasEditor/Public/Editors/SObjectTreeGraphNode.h
--- a/Engine/Plugins/Cameras/GameplayCameras/Source/GameplayCamerasEditor/Public/Editors/SObjectTreeGraphNode.h
+++ b/Engine/Plugins/Cameras/GameplayCameras/Source/GameplayCamerasEditor/Public/Editors/SObjectTreeGraphNode.h
@@ -37,7 +37,7 @@ public:
 	virtual const FSlateBrush* GetNodeBodyBrush() const override;
 
 	// SNodePanel::SNode interface.
-	virtual void MoveTo(const FVector2D& NewPosition, FNodeSet& NodeFilter, bool bMarkDirty=true) override;
+	virtual void MoveTo(const FVector2f& NewPosition, FNodeSet& NodeFilter, bool bMarkDirty=true) override;
 
 protected:
 
diff --git a/Engine/Plugins/Editor/AssetManagerEditor/Source/AssetManagerEditor/Private/ReferenceViewer/SReferenceNode.cpp b/Engine/Plugins/Editor/AssetManagerEditor/Source/AssetManagerEditor/Private/ReferenceViewer/SReferenceNode.cpp
--- a/Engine/Plugins/Editor/AssetManagerEditor/Source/AssetManagerEditor/Private/ReferenceViewer/SReferenceNode.cpp
+++ b/Engine/Plugins/Editor/AssetManagerEditor/Source/AssetManagerEditor/Private/ReferenceViewer/SReferenceNode.cpp
@@ -295,8 +295,8 @@ void SReferenceNode::UpdateGraphNode()
 			.Text(this, &SGraphNode::GetNodeComment);
 
 		GetOrAddSlot( ENodeZone::TopCenter )
-		.SlotOffset( TAttribute<FVector2D>( CommentBubble.Get(), &SCommentBubble::GetOffset ))
-		.SlotSize( TAttribute<FVector2D>( CommentBubble.Get(), &SCommentBubble::GetSize ))
+		.SlotOffset2f( TAttribute<FVector2f>( CommentBubble.Get(), &SCommentBubble::GetOffset2f ))
+		.SlotSize2f( TAttribute<FVector2f>( CommentBubble.Get(), &SCommentBubble::GetSize2f ))
 		.AllowScaling( TAttribute<bool>( CommentBubble.Get(), &SCommentBubble::IsScalingAllowed ))
 		.VAlign( VAlign_Top )
 		[
diff --git a/Engine/Plugins/Editor/AssetManagerEditor/Source/AssetManagerEditor/Private/ReferenceViewer/SReferencedPropertiesNode.h b/Engine/Plugins/Editor/AssetManagerEditor/Source/AssetManagerEditor/Private/ReferenceViewer/SReferencedPropertiesNode.h
--- a/Engine/Plugins/Editor/AssetManagerEditor/Source/AssetManagerEditor/Private/ReferenceViewer/SReferencedPropertiesNode.h
+++ b/Engine/Plugins/Editor/AssetManagerEditor/Source/AssetManagerEditor/Private/ReferenceViewer/SReferencedPropertiesNode.h
@@ -32,7 +32,7 @@ public:
 	// ~End SGraphNode
 
 	// ~Begin SNode
-	virtual bool CanBeSelected(const FVector2D& InMousePositionInNode) const override { return false; }
+	virtual bool CanBeSelected(const FVector2f& InMousePositionInNode) const override { return false; }
 	// ~End SNode
 
 	// ~Begin SWidget
diff --git a/Engine/Plugins/Enterprise/DataprepEditor/Source/DataprepEditor/Private/Widgets/DataprepGraph/SDataprepGraphActionNode.cpp b/Engine/Plugins/Enterprise/DataprepEditor/Source/DataprepEditor/Private/Widgets/DataprepGraph/SDataprepGraphActionNode.cpp
--- a/Engine/Plugins/Enterprise/DataprepEditor/Source/DataprepEditor/Private/Widgets/DataprepGraph/SDataprepGraphActionNode.cpp
+++ b/Engine/Plugins/Enterprise/DataprepEditor/Source/DataprepEditor/Private/Widgets/DataprepGraph/SDataprepGraphActionNode.cpp
@@ -70,8 +70,8 @@ public:
 
 	virtual FSlateRect GetTitleRect() const override
 	{
-		const FVector2D NodePosition = GetPosition();
-		const FVector2D NodeSize = InlineEditableText.IsValid() ? InlineEditableText->GetDesiredSize() : GetDesiredSize();
+		const FVector2f NodePosition = GetPosition2f();
+		const FVector2f NodeSize = InlineEditableText.IsValid() ? InlineEditableText->GetDesiredSize() : GetDesiredSize();
 
 		return FSlateRect( NodePosition.X, NodePosition.Y + NodeSize.Y, NodePosition.X + NodeSize.X, NodePosition.Y );
 	}
@@ -407,14 +407,14 @@ FVector2D SDataprepGraphBaseActionNode::ComputeDesiredSize(float f) const
 	return FVector2D(FMath::Max(Size.X, UserSize.X), Size.Y);
 }
 
-FVector2D SDataprepGraphBaseActionNode::GetNodeMinimumSize() const
+FVector2f SDataprepGraphBaseActionNode::GetNodeMinimumSize2f() const
 {
-	return FVector2D( SDataprepGraphActionNode::DefaultWidth, SDataprepGraphActionNode::DefaultHeight );
+	return FVector2f( SDataprepGraphActionNode::DefaultWidth, SDataprepGraphActionNode::DefaultHeight );
 }
 
-FVector2D SDataprepGraphBaseActionNode::GetNodeMaximumSize() const
+FVector2f SDataprepGraphBaseActionNode::GetNodeMaximumSize2f() const
 {
-	return FVector2D( UserSize.X + 100, UserSize.Y + 0 );
+	return FVector2f( UserSize.X + 100.0f, UserSize.Y + 0.0f );
 }
 
 void SDataprepGraphActionNode::Construct(const FArguments& InArgs, UDataprepGraphActionNode* InActionNode)
diff --git a/Engine/Plugins/Enterprise/DataprepEditor/Source/DataprepEditor/Private/Widgets/DataprepGraph/SDataprepGraphActionNode.h b/Engine/Plugins/Enterprise/DataprepEditor/Source/DataprepEditor/Private/Widgets/DataprepGraph/SDataprepGraphActionNode.h
--- a/Engine/Plugins/Enterprise/DataprepEditor/Source/DataprepEditor/Private/Widgets/DataprepGraph/SDataprepGraphActionNode.h
+++ b/Engine/Plugins/Enterprise/DataprepEditor/Source/DataprepEditor/Private/Widgets/DataprepGraph/SDataprepGraphActionNode.h
@@ -51,8 +51,8 @@ public:
 	//~ End SPanel Interface
 
 	//~ Begin SGraphNodeResizable Interface
-	virtual FVector2D GetNodeMinimumSize() const override;
-	virtual FVector2D GetNodeMaximumSize() const override;
+	virtual FVector2f GetNodeMinimumSize2f() const override;
+	virtual FVector2f GetNodeMaximumSize2f() const override;
 	//~ Begin SGraphNodeResizable Interface
 
 protected:
diff --git a/Engine/Plugins/Enterprise/DataprepEditor/Source/DataprepEditor/Private/Widgets/DataprepGraph/SDataprepGraphTrackNode.cpp b/Engine/Plugins/Enterprise/DataprepEditor/Source/DataprepEditor/Private/Widgets/DataprepGraph/SDataprepGraphTrackNode.cpp
--- a/Engine/Plugins/Enterprise/DataprepEditor/Source/DataprepEditor/Private/Widgets/DataprepGraph/SDataprepGraphTrackNode.cpp
+++ b/Engine/Plugins/Enterprise/DataprepEditor/Source/DataprepEditor/Private/Widgets/DataprepGraph/SDataprepGraphTrackNode.cpp
@@ -310,7 +310,7 @@ void SDataprepGraphTrackNode::Construct(const FArguments& InArgs, UDataprepGraph
 	check(DataprepAssetPtr.IsValid());
 
 	SNodePanel::SNode::FNodeSet NodeFilter;
-	SGraphNode::MoveTo( TrackAnchor, NodeFilter);
+	SGraphNode::MoveTo( UE::Slate::CastToVector2f(TrackAnchor), NodeFilter);
 
 	InNode->SetWidget(SharedThis(this));
 
@@ -555,10 +555,10 @@ bool SDataprepGraphTrackNode::RefreshLayout()
 	return false;
 }
 
-void SDataprepGraphTrackNode::MoveTo(const FVector2D& NewPosition, FNodeSet& NodeFilter, bool bMarkDirty)
+void SDataprepGraphTrackNode::MoveTo(const FVector2f& NewPosition, FNodeSet& NodeFilter, bool bMarkDirty)
 {
 	// Block track node to specific position
-	SGraphNode::MoveTo( TrackAnchor, NodeFilter, bMarkDirty);
+	SGraphNode::MoveTo( UE::Slate::CastToVector2f(TrackAnchor), NodeFilter, bMarkDirty);
 }
 
 const FSlateBrush* SDataprepGraphTrackNode::GetShadowBrush(bool bSelected) const
@@ -656,8 +656,8 @@ FSlateRect SDataprepGraphTrackNode::Update()
 
 	if(SDataprepGraphTrackWidget* TrackWidget = TrackWidgetPtr.Get())
 	{
-		const FVector2D& WorkingSize = TrackWidget->WorkingSize;
-		const FVector2D TrackPosition = GetPosition().GetAbs();
+		const FVector2f WorkingSize = UE::Slate::CastToVector2f(TrackWidget->WorkingSize);
+		const FVector2f TrackPosition = GetPosition2f().GetAbs();
 
 		if(TrackWidget->DragIndicatorIndex == INDEX_NONE && !bNodeDragging)
 		{
@@ -666,9 +666,9 @@ FSlateRect SDataprepGraphTrackNode::Update()
 			SGraphPanel& GraphPanel = *OwnerGraphPanelPtr.Pin();
 
 			// Determine canvas offset attribute in track widget's coordinates
-			const FVector2D PanelSize = GraphPanel.GetTickSpaceGeometry().GetLocalSize() / GraphPanel.GetZoomAmount();
+			const FVector2f PanelSize = GraphPanel.GetTickSpaceGeometry().GetLocalSize() / GraphPanel.GetZoomAmount();
 
-			const FVector2D TargetSize = FVector2D(FMath::Max(WorkingSize.X, PanelSize.X), FMath::Max(WorkingSize.Y, PanelSize.Y)) + TrackPosition + 20.f;
+			const FVector2f TargetSize = FVector2f(FMath::Max(WorkingSize.X, PanelSize.X), FMath::Max(WorkingSize.Y, PanelSize.Y)) + TrackPosition + 20.f;
 
 			FMargin& CanvasOffset = TrackWidgetPtr->CanvasOffset;
 
@@ -688,7 +688,7 @@ FSlateRect SDataprepGraphTrackNode::Update()
 			TrackWidgetPtr->CanvasSlot->SetOffset(CanvasOffset);
 		}
 
-		WorkingArea = FSlateRect(FVector2D::ZeroVector, WorkingSize + TrackPosition);
+		WorkingArea = FSlateRect(FVector2f::ZeroVector, WorkingSize + TrackPosition);
 	}
 
 	return WorkingArea;
diff --git a/Engine/Plugins/Enterprise/DataprepEditor/Source/DataprepEditor/Private/Widgets/DataprepGraph/SDataprepGraphTrackNode.h b/Engine/Plugins/Enterprise/DataprepEditor/Source/DataprepEditor/Private/Widgets/DataprepGraph/SDataprepGraphTrackNode.h
--- a/Engine/Plugins/Enterprise/DataprepEditor/Source/DataprepEditor/Private/Widgets/DataprepGraph/SDataprepGraphTrackNode.h
+++ b/Engine/Plugins/Enterprise/DataprepEditor/Source/DataprepEditor/Private/Widgets/DataprepGraph/SDataprepGraphTrackNode.h
@@ -64,12 +64,12 @@ public:
 	// End of SGraphNode interface
 
 	// SNodePanel::SNode interface
-	virtual bool CanBeSelected(const FVector2D& /*MousePositionInNode*/) const override
+	virtual bool CanBeSelected(const FVector2f& /*MousePositionInNode*/) const override
 	{
 		return false;
 	}
 	virtual void SetOwner(const TSharedRef<SGraphPanel>& OwnerPanel) override;
-	virtual void MoveTo( const FVector2D& NewPosition, FNodeSet& NodeFilter, bool bMarkDirty = true ) override;
+	virtual void MoveTo( const FVector2f& NewPosition, FNodeSet& NodeFilter, bool bMarkDirty = true ) override;
 	virtual const FSlateBrush* GetShadowBrush(bool bSelected) const override;
 	virtual FReply OnDrop( const FGeometry& MyGeometry, const FDragDropEvent& DragDropEvent ) override;
 	virtual bool ShouldAllowCulling() const override { return false; }
diff --git a/Engine/Plugins/Experimental/BlueprintSnapNodes/Source/BlueprintSnapNodes/Private/SGraphNodeSnapContainer.cpp b/Engine/Plugins/Experimental/BlueprintSnapNodes/Source/BlueprintSnapNodes/Private/SGraphNodeSnapContainer.cpp
--- a/Engine/Plugins/Experimental/BlueprintSnapNodes/Source/BlueprintSnapNodes/Private/SGraphNodeSnapContainer.cpp
+++ b/Engine/Plugins/Experimental/BlueprintSnapNodes/Source/BlueprintSnapNodes/Private/SGraphNodeSnapContainer.cpp
@@ -160,8 +160,8 @@ void SGraphNodeSnapContainer::UpdateGraphNode()
 	.IsGraphNodeHovered( this, &SGraphNode::IsHovered );
 
 	GetOrAddSlot( ENodeZone::TopCenter )
-	.SlotOffset( TAttribute<FVector2D>( CommentBubble.Get(), &SCommentBubble::GetOffset ))
-	.SlotSize( TAttribute<FVector2D>( CommentBubble.Get(), &SCommentBubble::GetSize ))
+	.SlotOffset2f( TAttribute<FVector2f>( CommentBubble.Get(), &SCommentBubble::GetOffset2f ))
+	.SlotSize2f( TAttribute<FVector2f>( CommentBubble.Get(), &SCommentBubble::GetSize2f ))
 	.AllowScaling( TAttribute<bool>( CommentBubble.Get(), &SCommentBubble::IsScalingAllowed ))
 	.VAlign( VAlign_Top )
 	[
diff --git a/Engine/Plugins/Experimental/CommonConversation/Source/CommonConversationGraph/Private/SConversationGraphNode.cpp b/Engine/Plugins/Experimental/CommonConversation/Source/CommonConversationGraph/Private/SConversationGraphNode.cpp
--- a/Engine/Plugins/Experimental/CommonConversation/Source/CommonConversationGraph/Private/SConversationGraphNode.cpp
+++ b/Engine/Plugins/Experimental/CommonConversation/Source/CommonConversationGraph/Private/SConversationGraphNode.cpp
@@ -562,8 +562,8 @@ void SConversationGraphNode::UpdateGraphNode()
 	.IsGraphNodeHovered( this, &SGraphNode::IsHovered );
 
 	GetOrAddSlot( ENodeZone::TopCenter )
-	.SlotOffset( TAttribute<FVector2D>( CommentBubble.Get(), &SCommentBubble::GetOffset ))
-	.SlotSize( TAttribute<FVector2D>( CommentBubble.Get(), &SCommentBubble::GetSize ))
+	.SlotOffset2f( TAttribute<FVector2f>( CommentBubble.Get(), &SCommentBubble::GetOffset2f ))
+	.SlotSize2f( TAttribute<FVector2f>( CommentBubble.Get(), &SCommentBubble::GetSize2f ))
 	.AllowScaling( TAttribute<bool>( CommentBubble.Get(), &SCommentBubble::IsScalingAllowed ))
 	.VAlign( VAlign_Top )
 	[
@@ -959,7 +959,7 @@ EVisibility SConversationGraphNode::GetBlueprintIconVisibility() const
 	return (bCanShowIcon && (!MyOwnerPanel.IsValid() || MyOwnerPanel->GetCurrentLOD() > EGraphRenderingLOD::LowDetail)) ? EVisibility::Visible : EVisibility::Collapsed;
 }
 
-void SConversationGraphNode::GetOverlayBrushes(bool bSelected, const FVector2D WidgetSize, TArray<FOverlayBrushInfo>& Brushes) const
+void SConversationGraphNode::GetOverlayBrushes(bool bSelected, const FVector2f& WidgetSize, TArray<FOverlayBrushInfo>& Brushes) const
 {
 	UConversationGraphNode* ConversationGraphNode = CastChecked<UConversationGraphNode>(GraphNode);
 
@@ -1021,13 +1021,13 @@ void SConversationGraphNode::GetOverlayBrushes(bool bSelected, const FVector2D W
 // 	}
 }
 
-TArray<FOverlayWidgetInfo> SConversationGraphNode::GetOverlayWidgets(bool bSelected, const FVector2D& WidgetSize) const
+TArray<FOverlayWidgetInfo> SConversationGraphNode::GetOverlayWidgets(bool bSelected, const FVector2f& WidgetSize) const
 {
 	TArray<FOverlayWidgetInfo> Widgets;
 
 	check(NodeBody.IsValid());
 
-	FVector2D Origin(0.0f, 0.0f);
+	FVector2f Origin(0.0f, 0.0f);
 
 	// build overlays for sub-nodes (above)
 	for (const TSharedPtr<SGraphNode>& ChildWidget : ChildNodeStuff.FindChecked(ESubNodeWidgetLocation::Above).ChildNodeWidgets)
@@ -1044,7 +1044,7 @@ TArray<FOverlayWidgetInfo> SConversationGraphNode::GetOverlayWidgets(bool bSelec
 	if (IndexOverlay.IsValid())
 	{
 		FOverlayWidgetInfo Overlay(IndexOverlay);
-		Overlay.OverlayOffset = FVector2D(WidgetSize.X - (IndexOverlay->GetDesiredSize().X * 0.5f), Origin.Y);
+		Overlay.OverlayOffset = FVector2f(WidgetSize.X - (IndexOverlay->GetDesiredSize().X * 0.5f), Origin.Y);
 		Widgets.Add(Overlay);
 	}
 
@@ -1071,7 +1071,7 @@ TSharedRef<SGraphNode> SConversationGraphNode::GetNodeUnderMouse(const FGeometry
 	return SubNode.IsValid() ? SubNode.ToSharedRef() : StaticCastSharedRef<SGraphNode>(AsShared());
 }
 
-void SConversationGraphNode::MoveTo(const FVector2D& NewPosition, FNodeSet& NodeFilter, bool bMarkDirty)
+void SConversationGraphNode::MoveTo(const FVector2f& NewPosition, FNodeSet& NodeFilter, bool bMarkDirty)
 {
 	SGraphNodeAI::MoveTo(NewPosition, NodeFilter, bMarkDirty);
 
diff --git a/Engine/Plugins/Experimental/CommonConversation/Source/CommonConversationGraph/Private/SConversationGraphNode.h b/Engine/Plugins/Experimental/CommonConversation/Source/CommonConversationGraph/Private/SConversationGraphNode.h
--- a/Engine/Plugins/Experimental/CommonConversation/Source/CommonConversationGraph/Private/SConversationGraphNode.h
+++ b/Engine/Plugins/Experimental/CommonConversation/Source/CommonConversationGraph/Private/SConversationGraphNode.h
@@ -45,10 +45,10 @@ public:
 	virtual void CreatePinWidgets() override;
 	virtual void AddPin(const TSharedRef<SGraphPin>& PinToAdd) override;
 	virtual TSharedPtr<SToolTip> GetComplexTooltip() override;
-	virtual void GetOverlayBrushes(bool bSelected, const FVector2D WidgetSize, TArray<FOverlayBrushInfo>& Brushes) const override;
-	virtual TArray<FOverlayWidgetInfo> GetOverlayWidgets(bool bSelected, const FVector2D& WidgetSize) const override;
+	virtual void GetOverlayBrushes(bool bSelected, const FVector2f& WidgetSize, TArray<FOverlayBrushInfo>& Brushes) const override;
+	virtual TArray<FOverlayWidgetInfo> GetOverlayWidgets(bool bSelected, const FVector2f& WidgetSize) const override;
 	virtual TSharedRef<SGraphNode> GetNodeUnderMouse(const FGeometry& MyGeometry, const FPointerEvent& MouseEvent) override;
-	virtual void MoveTo(const FVector2D& NewPosition, FNodeSet& NodeFilter, bool bMarkDirty = true) override;
+	virtual void MoveTo(const FVector2f& NewPosition, FNodeSet& NodeFilter, bool bMarkDirty = true) override;
 	virtual FReply OnDrop(const FGeometry& MyGeometry, const FDragDropEvent& DragDropEvent) override;
 	// End of SGraphNode interface
 
diff --git a/Engine/Plugins/Experimental/Dataflow/Source/DataflowEditor/Private/Dataflow/DataflowSNode.cpp b/Engine/Plugins/Experimental/Dataflow/Source/DataflowEditor/Private/Dataflow/DataflowSNode.cpp
--- a/Engine/Plugins/Experimental/Dataflow/Source/DataflowEditor/Private/Dataflow/DataflowSNode.cpp
+++ b/Engine/Plugins/Experimental/Dataflow/Source/DataflowEditor/Private/Dataflow/DataflowSNode.cpp
@@ -228,7 +228,7 @@ TSharedRef<SWidget> SDataflowEdNode::CreateTitleRightWidget()
 	return SNullWidget::NullWidget;
 }
 
-TArray<FOverlayWidgetInfo> SDataflowEdNode::GetOverlayWidgets(bool bSelected, const FVector2D& WidgetSize) const
+TArray<FOverlayWidgetInfo> SDataflowEdNode::GetOverlayWidgets(bool bSelected, const FVector2f& WidgetSize) const
 {
 	TArray<FOverlayWidgetInfo> Widgets = SGraphNode::GetOverlayWidgets(bSelected, WidgetSize);
 
@@ -236,8 +236,8 @@ TArray<FOverlayWidgetInfo> SDataflowEdNode::GetOverlayWidgets(bool bSelected, co
 	{
 		if (DataflowInterface->NodesHaveFreezeWidget())
 		{
-			static const FVector2D FreezeImageOverhang(10.f);  // The icon must slightly overhang to make space for the lower output pins
-			const FVector2D FreezeImageSize = FreezeImageWidget->GetDesiredSize();
+			static const FVector2f FreezeImageOverhang(10.f);  // The icon must slightly overhang to make space for the lower output pins
+			const FVector2f FreezeImageSize = FreezeImageWidget->GetDesiredSize();
 			FOverlayWidgetInfo FreezeImageInfo;
 			FreezeImageInfo.OverlayOffset = WidgetSize - FreezeImageSize + FreezeImageOverhang;
 			FreezeImageInfo.Widget = FreezeImageWidget;
@@ -271,9 +271,8 @@ TArray<FOverlayWidgetInfo> SDataflowEdNode::GetOverlayWidgets(bool bSelected, co
 						PerfWidget->SetText(FText());
 					}
 
-					static const FVector2D FreezeImageOverhang(10.f);  // The icon must slightly overhang to make space for the lower output pins
 					FOverlayWidgetInfo PerfWidgetInfo;
-					PerfWidgetInfo.OverlayOffset = FVector2D{ 0, -20 };
+					PerfWidgetInfo.OverlayOffset = FVector2f{ 0.0f, -20.0f };
 					PerfWidgetInfo.Widget = PerfWidget;
 					Widgets.Add(PerfWidgetInfo);
 				}
@@ -300,7 +299,7 @@ TArray<FOverlayWidgetInfo> SDataflowEdNode::GetOverlayWidgets(bool bSelected, co
 				WatchWidget->SetText(FText::FromString(WatchString));
 
 				FOverlayWidgetInfo WatchWidgetInfo;
-				WatchWidgetInfo.OverlayOffset = FVector2D{ 0, (double)WidgetSize.Y + 4.0 };
+				WatchWidgetInfo.OverlayOffset = FVector2f{ 0.0f, WidgetSize.Y + 4.0f };
 				WatchWidgetInfo.Widget = WatchWidget;
 				Widgets.Add(WatchWidgetInfo);
 			}
diff --git a/Engine/Plugins/Experimental/Dataflow/Source/DataflowEditor/Public/Dataflow/DataflowSNode.h b/Engine/Plugins/Experimental/Dataflow/Source/DataflowEditor/Public/Dataflow/DataflowSNode.h
--- a/Engine/Plugins/Experimental/Dataflow/Source/DataflowEditor/Public/Dataflow/DataflowSNode.h
+++ b/Engine/Plugins/Experimental/Dataflow/Source/DataflowEditor/Public/Dataflow/DataflowSNode.h
@@ -35,7 +35,7 @@ public:
 
 	// SGraphNode interface
 	virtual FReply OnMouseButtonDoubleClick(const FGeometry& InMyGeometry, const FPointerEvent& InMouseEvent) override;
-	virtual TArray<FOverlayWidgetInfo> GetOverlayWidgets(bool bSelected, const FVector2D& WidgetSize) const override;
+	virtual TArray<FOverlayWidgetInfo> GetOverlayWidgets(bool bSelected, const FVector2f& WidgetSize) const override;
 	virtual void UpdateErrorInfo() override;
 	virtual void RequestRenameOnSpawn() override { /* No auto rename on spawn, because it can interfere with Copy/Paste and cause a crash */ }
 
diff --git a/Engine/Plugins/Experimental/PluginReferenceViewer/Source/PluginReferenceViewer/Private/SPluginReferenceNode.cpp b/Engine/Plugins/Experimental/PluginReferenceViewer/Source/PluginReferenceViewer/Private/SPluginReferenceNode.cpp
--- a/Engine/Plugins/Experimental/PluginReferenceViewer/Source/PluginReferenceViewer/Private/SPluginReferenceNode.cpp
+++ b/Engine/Plugins/Experimental/PluginReferenceViewer/Source/PluginReferenceViewer/Private/SPluginReferenceNode.cpp
@@ -261,8 +261,8 @@ void SPluginReferenceNode::UpdateGraphNode()
 			.Text(this, &SGraphNode::GetNodeComment);
 
 		GetOrAddSlot(ENodeZone::TopCenter)
-			.SlotOffset(TAttribute<FVector2D>(CommentBubble.Get(), &SCommentBubble::GetOffset))
-			.SlotSize(TAttribute<FVector2D>(CommentBubble.Get(), &SCommentBubble::GetSize))
+			.SlotOffset2f(TAttribute<FVector2f>(CommentBubble.Get(), &SCommentBubble::GetOffset2f))
+			.SlotSize2f(TAttribute<FVector2f>(CommentBubble.Get(), &SCommentBubble::GetSize2f))
 			.AllowScaling(TAttribute<bool>(CommentBubble.Get(), &SCommentBubble::IsScalingAllowed))
 			.VAlign(VAlign_Top)
 			[
diff --git a/Engine/Plugins/FX/Niagara/Source/NiagaraEditor/Private/Widgets/SNiagaraGraphNode.cpp b/Engine/Plugins/FX/Niagara/Source/NiagaraEditor/Private/Widgets/SNiagaraGraphNode.cpp
--- a/Engine/Plugins/FX/Niagara/Source/NiagaraEditor/Private/Widgets/SNiagaraGraphNode.cpp
+++ b/Engine/Plugins/FX/Niagara/Source/NiagaraEditor/Private/Widgets/SNiagaraGraphNode.cpp
@@ -335,8 +335,8 @@ void SNiagaraGraphNode::UpdateGraphNodeCompact()
 	.IsGraphNodeHovered(this, &SGraphNode::IsHovered);
 
 	GetOrAddSlot( ENodeZone::TopCenter )
-	.SlotOffset( TAttribute<FVector2D>( CommentBubble.Get(), &SCommentBubble::GetOffset ))
-	.SlotSize( TAttribute<FVector2D>( CommentBubble.Get(), &SCommentBubble::GetSize ))
+	.SlotOffset2f( TAttribute<FVector2f>( CommentBubble.Get(), &SCommentBubble::GetOffset2f ))
+	.SlotSize2f( TAttribute<FVector2f>( CommentBubble.Get(), &SCommentBubble::GetSize2f ))
 	.AllowScaling( TAttribute<bool>( CommentBubble.Get(), &SCommentBubble::IsScalingAllowed ))
 	.VAlign( VAlign_Top )
 	[
diff --git a/Engine/Plugins/PCG/Source/PCGEditor/Private/Nodes/SPCGEditorGraphNode.cpp b/Engine/Plugins/PCG/Source/PCGEditor/Private/Nodes/SPCGEditorGraphNode.cpp
--- a/Engine/Plugins/PCG/Source/PCGEditor/Private/Nodes/SPCGEditorGraphNode.cpp
+++ b/Engine/Plugins/PCG/Source/PCGEditor/Private/Nodes/SPCGEditorGraphNode.cpp
@@ -276,7 +276,7 @@ void SPCGEditorGraphNode::AddPin(const TSharedRef<SGraphPin>& PinToAdd)
 	}
 }
 
-TArray<FOverlayWidgetInfo> SPCGEditorGraphNode::GetOverlayWidgets(bool bSelected, const FVector2D& WidgetSize) const
+TArray<FOverlayWidgetInfo> SPCGEditorGraphNode::GetOverlayWidgets(bool bSelected, const FVector2f& WidgetSize) const
 {
 	TArray<FOverlayWidgetInfo> OverlayWidgets = SGraphNode::GetOverlayWidgets(bSelected, WidgetSize);
 
@@ -293,11 +293,11 @@ TArray<FOverlayWidgetInfo> SPCGEditorGraphNode::GetOverlayWidgets(bool bSelected
 	return OverlayWidgets;
 }
 
-void SPCGEditorGraphNode::GetOverlayBrushes(bool bSelected, const FVector2D WidgetSize, TArray<FOverlayBrushInfo>& Brushes) const
+void SPCGEditorGraphNode::GetOverlayBrushes(bool bSelected, const FVector2f& WidgetSize, TArray<FOverlayBrushInfo>& Brushes) const
 {
 	check(PCGEditorGraphNode);
 	
-	FVector2D OverlayOffset(0.0, 0.0);
+	FVector2f OverlayOffset(0.0f, 0.0f);
 
 	auto AddOverlayBrush = [&OverlayOffset, &Brushes](const FName& BrushName)
 	{
@@ -307,7 +307,7 @@ void SPCGEditorGraphNode::GetOverlayBrushes(bool bSelected, const FVector2D Widg
 		{
 			FOverlayBrushInfo BrushInfo;
 			BrushInfo.Brush = Brush;
-			BrushInfo.OverlayOffset = OverlayOffset - Brush->GetImageSize() / 2.0;
+			BrushInfo.OverlayOffset = OverlayOffset - Brush->GetImageSize() / 2.0f;
 			Brushes.Add(BrushInfo);
 
 			OverlayOffset.Y += Brush->GetImageSize().Y;
diff --git a/Engine/Plugins/PCG/Source/PCGEditor/Private/Nodes/SPCGEditorGraphNode.h b/Engine/Plugins/PCG/Source/PCGEditor/Private/Nodes/SPCGEditorGraphNode.h
--- a/Engine/Plugins/PCG/Source/PCGEditor/Private/Nodes/SPCGEditorGraphNode.h
+++ b/Engine/Plugins/PCG/Source/PCGEditor/Private/Nodes/SPCGEditorGraphNode.h
@@ -31,8 +31,8 @@ public:
 	//~ End SGraphNode Interface
 
 	//~ Begin SNodePanel::SNode Interface
-	virtual TArray<FOverlayWidgetInfo> GetOverlayWidgets(bool bSelected, const FVector2D& WidgetSize) const override;
-	virtual void GetOverlayBrushes(bool bSelected, const FVector2D WidgetSize, TArray<FOverlayBrushInfo>& Brushes) const override;
+	virtual TArray<FOverlayWidgetInfo> GetOverlayWidgets(bool bSelected, const FVector2f& WidgetSize) const override;
+	virtual void GetOverlayBrushes(bool bSelected, const FVector2f& WidgetSize, TArray<FOverlayBrushInfo>& Brushes) const override;
 	//~ End SNodePanel::SNode Interface
 
 protected:
diff --git a/Engine/Plugins/PCG/Source/PCGEditor/Private/Nodes/SPCGEditorGraphNodeCompact.cpp b/Engine/Plugins/PCG/Source/PCGEditor/Private/Nodes/SPCGEditorGraphNodeCompact.cpp
--- a/Engine/Plugins/PCG/Source/PCGEditor/Private/Nodes/SPCGEditorGraphNodeCompact.cpp
+++ b/Engine/Plugins/PCG/Source/PCGEditor/Private/Nodes/SPCGEditorGraphNodeCompact.cpp
@@ -289,8 +289,8 @@ void SPCGEditorGraphNodeCompact::UpdateGraphNode()
 	.IsGraphNodeHovered(this, &SGraphNode::IsHovered);
 
 	GetOrAddSlot(ENodeZone::TopCenter)
-		.SlotOffset(TAttribute<FVector2D>(CommentBubble.Get(), &SCommentBubble::GetOffset))
-		.SlotSize(TAttribute<FVector2D>(CommentBubble.Get(), &SCommentBubble::GetSize))
+		.SlotOffset2f(TAttribute<FVector2f>(CommentBubble.Get(), &SCommentBubble::GetOffset2f))
+		.SlotSize2f(TAttribute<FVector2f>(CommentBubble.Get(), &SCommentBubble::GetSize2f))
 		.AllowScaling(TAttribute<bool>(CommentBubble.Get(), &SCommentBubble::IsScalingAllowed))
 		.VAlign(VAlign_Top)
 		[
diff --git a/Engine/Plugins/Runtime/Metasound/Source/MetasoundEditor/Private/SMetasoundGraphNode.cpp b/Engine/Plugins/Runtime/Metasound/Source/MetasoundEditor/Private/SMetasoundGraphNode.cpp
--- a/Engine/Plugins/Runtime/Metasound/Source/MetasoundEditor/Private/SMetasoundGraphNode.cpp
+++ b/Engine/Plugins/Runtime/Metasound/Source/MetasoundEditor/Private/SMetasoundGraphNode.cpp
@@ -482,7 +482,7 @@ namespace Metasound
 			return TitleBoxWidget.ToSharedRef();
 		}
 
-		void SMetaSoundGraphNode::GetOverlayBrushes(bool bSelected, const FVector2D WidgetSize, TArray<FOverlayBrushInfo>& Brushes) const
+		void SMetaSoundGraphNode::GetOverlayBrushes(bool bSelected, const FVector2f& WidgetSize, TArray<FOverlayBrushInfo>& Brushes) const
 		{
 			FName CornerIcon = GetMetaSoundNode().GetCornerIcon();
 			if (CornerIcon != NAME_None)
@@ -638,13 +638,13 @@ namespace Metasound
 			}
 		}
 
-		void SMetaSoundGraphNode::MoveTo(const FVector2D& NewPosition, FNodeSet& NodeFilter, bool bMarkDirty)
+		void SMetaSoundGraphNode::MoveTo(const FVector2f& NewPosition, FNodeSet& NodeFilter, bool bMarkDirty)
 		{
 			SGraphNode::MoveTo(NewPosition, NodeFilter, bMarkDirty);
 
 			UMetasoundEditorGraphNode& Node = GetMetaSoundNode();
 			Node.GetMetasoundChecked().Modify();
-			Node.UpdateFrontendNodeLocation(NewPosition);
+			Node.UpdateFrontendNodeLocation(FVector2D(NewPosition));
 			Node.SyncLocationFromFrontendNode();
 		}
 
@@ -1349,13 +1349,13 @@ namespace Metasound
 			UpdateGraphNode();
 		}
 
-		void SMetaSoundGraphNodeKnot::MoveTo(const FVector2D& NewPosition, FNodeSet& NodeFilter, bool bMarkDirty)
+		void SMetaSoundGraphNodeKnot::MoveTo(const FVector2f& NewPosition, FNodeSet& NodeFilter, bool bMarkDirty)
 		{
 			SGraphNode::MoveTo(NewPosition, NodeFilter, bMarkDirty);
 
 			UMetasoundEditorGraphNode& Node = GetMetaSoundNode();
 			Node.GetMetasoundChecked().Modify();
-			Node.UpdateFrontendNodeLocation(NewPosition);
+			Node.UpdateFrontendNodeLocation(FVector2D(NewPosition));
 			Node.SyncLocationFromFrontendNode();
 		}
 
diff --git a/Engine/Plugins/Runtime/Metasound/Source/MetasoundEditor/Private/SMetasoundGraphNode.h b/Engine/Plugins/Runtime/Metasound/Source/MetasoundEditor/Private/SMetasoundGraphNode.h
--- a/Engine/Plugins/Runtime/Metasound/Source/MetasoundEditor/Private/SMetasoundGraphNode.h
+++ b/Engine/Plugins/Runtime/Metasound/Source/MetasoundEditor/Private/SMetasoundGraphNode.h
@@ -54,9 +54,9 @@ namespace Metasound::Editor
 		virtual TSharedRef<SWidget> CreateTitleRightWidget() override;
 		virtual EVisibility IsAddPinButtonVisible() const override;
 		virtual FReply OnAddPin() override;
-		virtual void MoveTo(const FVector2D& NewPosition, FNodeSet& NodeFilter, bool bMarkDirty = true) override;
+		virtual void MoveTo(const FVector2f& NewPosition, FNodeSet& NodeFilter, bool bMarkDirty = true) override;
 		virtual void SetDefaultTitleAreaWidget(TSharedRef<SOverlay> DefaultTitleAreaWidget) override;
-		virtual void GetOverlayBrushes(bool bSelected, const FVector2D WidgetSize, TArray<FOverlayBrushInfo>& Brushes) const override;
+		virtual void GetOverlayBrushes(bool bSelected, const FVector2f& WidgetSize, TArray<FOverlayBrushInfo>& Brushes) const override;
 		virtual void OnCommentBubbleToggled(bool bInCommentBubbleVisible) override;
 		virtual void OnCommentTextCommitted(const FText& NewComment, ETextCommit::Type CommitInfo) override;
 		virtual void OnAdvancedViewChanged(const ECheckBoxState NewCheckedState) override;
@@ -125,7 +125,7 @@ namespace Metasound::Editor
 		void Construct(const FArguments& InArgs, class UEdGraphNode* InNode);
 
 		virtual TSharedPtr<SGraphPin> CreatePinWidget(UEdGraphPin* Pin) const override;
-		virtual void MoveTo(const FVector2D& NewPosition, FNodeSet& NodeFilter, bool bMarkDirty) override;
+		virtual void MoveTo(const FVector2f& NewPosition, FNodeSet& NodeFilter, bool bMarkDirty) override;
 
 		UMetasoundEditorGraphNode& GetMetaSoundNode();
 		const UMetasoundEditorGraphNode& GetMetaSoundNode() const;
diff --git a/Engine/Plugins/Runtime/Metasound/Source/MetasoundEditor/Private/SMetasoundGraphNodeComment.cpp b/Engine/Plugins/Runtime/Metasound/Source/MetasoundEditor/Private/SMetasoundGraphNodeComment.cpp
--- a/Engine/Plugins/Runtime/Metasound/Source/MetasoundEditor/Private/SMetasoundGraphNodeComment.cpp
+++ b/Engine/Plugins/Runtime/Metasound/Source/MetasoundEditor/Private/SMetasoundGraphNodeComment.cpp
@@ -10,7 +10,7 @@ namespace Metasound
 {
 	namespace Editor
 	{
-		void SMetasoundGraphNodeComment::MoveTo(const FVector2D& NewPosition, FNodeSet& NodeFilter, bool bMarkDirty)
+		void SMetasoundGraphNodeComment::MoveTo(const FVector2f& NewPosition, FNodeSet& NodeFilter, bool bMarkDirty)
 		{
 			SGraphNodeComment::MoveTo(NewPosition, NodeFilter, bMarkDirty);
 
@@ -31,7 +31,7 @@ namespace Metasound
 				UMetasoundEditorGraphCommentNode* CommentNode = Cast<UMetasoundEditorGraphCommentNode>(GraphNode);
 				if (CommentNode && CommentNode->MoveMode == ECommentBoxMode::GroupMovement)
 				{
-					FVector2D PositionDelta = NewPosition - GetPosition();
+					FVector2f PositionDelta = NewPosition - GetPosition2f();
 
 					// Now update any nodes which are touching the comment but *not* selected
 					// Selected nodes will be moved as part of the normal selection code
@@ -40,10 +40,10 @@ namespace Metasound
 					{
 						if (UMetasoundEditorGraphNode* MetasoundGraphNode = Cast<UMetasoundEditorGraphNode>(*NodeIt))
 						{
-							FVector2D MetasoundNodePosition = FVector2D(MetasoundGraphNode->NodePosX, MetasoundGraphNode->NodePosY);
+							FVector2f MetasoundNodePosition = FVector2f(MetasoundGraphNode->NodePosX, MetasoundGraphNode->NodePosY);
 							MetasoundNodePosition += PositionDelta;
 							MetasoundGraphNode->GetMetasoundChecked().Modify();
-							MetasoundGraphNode->UpdateFrontendNodeLocation(MetasoundNodePosition);
+							MetasoundGraphNode->UpdateFrontendNodeLocation(FVector2D(MetasoundNodePosition));
 						}
 						else if (UMetasoundEditorGraphCommentNode* MetasoundCommentNode = Cast<UMetasoundEditorGraphCommentNode>(*NodeIt))
 						{
diff --git a/Engine/Plugins/Runtime/Metasound/Source/MetasoundEditor/Private/SMetasoundGraphNodeComment.h b/Engine/Plugins/Runtime/Metasound/Source/MetasoundEditor/Private/SMetasoundGraphNodeComment.h
--- a/Engine/Plugins/Runtime/Metasound/Source/MetasoundEditor/Private/SMetasoundGraphNodeComment.h
+++ b/Engine/Plugins/Runtime/Metasound/Source/MetasoundEditor/Private/SMetasoundGraphNodeComment.h
@@ -12,7 +12,7 @@ namespace Metasound
 		{
 		public:
 			// SNodePanel::SNode interface
-			virtual void MoveTo(const FVector2D& NewPosition, FNodeSet& NodeFilter, bool bMarkDirty = true) override;
+			virtual void MoveTo(const FVector2f& NewPosition, FNodeSet& NodeFilter, bool bMarkDirty = true) override;
 			// End of SNodePanel::SNode interface
 		};
 	} // namespace Editor
diff --git a/Engine/Plugins/Runtime/RigVM/Source/RigVMEditor/Private/Widgets/SRigVMGraphNode.cpp b/Engine/Plugins/Runtime/RigVM/Source/RigVMEditor/Private/Widgets/SRigVMGraphNode.cpp
--- a/Engine/Plugins/Runtime/RigVM/Source/RigVMEditor/Private/Widgets/SRigVMGraphNode.cpp
+++ b/Engine/Plugins/Runtime/RigVM/Source/RigVMEditor/Private/Widgets/SRigVMGraphNode.cpp
@@ -425,7 +425,7 @@ void SRigVMGraphNode::EndUserInteraction() const
 	SGraphNode::EndUserInteraction();
 }
 
-void SRigVMGraphNode::MoveTo(const FVector2D& NewPosition, FNodeSet& NodeFilter, bool bMarkDirty)
+void SRigVMGraphNode::MoveTo(const FVector2f& NewPosition, FNodeSet& NodeFilter, bool bMarkDirty)
 {
 	if (!NodeFilter.Find(SharedThis(this)))
 	{
@@ -433,7 +433,7 @@ void SRigVMGraphNode::MoveTo(const FVector2D& NewPosition, FNodeSet& NodeFilter,
 		{
 			if (const URigVMEdGraphSchema* RigSchema = Cast<URigVMEdGraphSchema>(GraphNode->GetSchema()))
 			{
-				RigSchema->SetNodePosition(GraphNode, NewPosition, false);
+				RigSchema->SetNodePosition(GraphNode, FVector2D(NewPosition), false);
 			}
 		}
 	}
@@ -800,7 +800,7 @@ FReply SRigVMGraphNode::HandleAddArrayElement(FString InModelPinPath)
 }
 
 /** Populate the brushes array with any overlay brushes to render */
-void SRigVMGraphNode::GetOverlayBrushes(bool bSelected, const FVector2D WidgetSize, TArray<FOverlayBrushInfo>& Brushes) const
+void SRigVMGraphNode::GetOverlayBrushes(bool bSelected, const FVector2f& WidgetSize, TArray<FOverlayBrushInfo>& Brushes) const
 {
 	URigVMEdGraphNode* RigGraphNode = Cast<URigVMEdGraphNode>(GraphNode);
 
@@ -833,7 +833,7 @@ void SRigVMGraphNode::GetOverlayBrushes(bool bSelected, const FVector2D WidgetSi
 					IPOverlayInfo.OverlayOffset.Y = (Overlap - IPOverlayInfo.Brush->ImageSize.Y);
 				}
 
-				IPOverlayInfo.AnimationEnvelope = FVector2D(0.f, 10.f);
+				IPOverlayInfo.AnimationEnvelope = FVector2f(0.f, 10.f);
 
 				Brushes.Add(IPOverlayInfo);
 			}
@@ -981,7 +981,7 @@ void SRigVMGraphNode::GetNodeInfoPopups(FNodeInfoContext* Context, TArray<FGraph
 	}
 }
 
-TArray<FOverlayWidgetInfo> SRigVMGraphNode::GetOverlayWidgets(bool bSelected, const FVector2D& WidgetSize) const
+TArray<FOverlayWidgetInfo> SRigVMGraphNode::GetOverlayWidgets(bool bSelected, const FVector2f& WidgetSize) const
 {
 	TArray<FOverlayWidgetInfo> Widgets = SGraphNode::GetOverlayWidgets(bSelected, WidgetSize);
 
@@ -1025,10 +1025,10 @@ TArray<FOverlayWidgetInfo> SRigVMGraphNode::GetOverlayWidgets(bool bSelected, co
 
 							   if (Widgets.Num() == PreviousNumWidgets)
 							   {
-								   const FVector2D ImageSize = VisualDebugIndicatorWidget->GetDesiredSize();
+								   const FVector2f ImageSize = VisualDebugIndicatorWidget->GetDesiredSize();
 
 								   FOverlayWidgetInfo Info;
-								   Info.OverlayOffset = FVector2D(WidgetSize.X - ImageSize.X - 6.f, 6.f);
+								   Info.OverlayOffset = FVector2f(WidgetSize.X - ImageSize.X - 6.f, 6.f);
 								   Info.Widget = VisualDebugIndicatorWidget;
 
 								   Widgets.Add(Info);
@@ -1056,9 +1056,9 @@ TArray<FOverlayWidgetInfo> SRigVMGraphNode::GetOverlayWidgets(bool bSelected, co
 						if((Count > Blueprint->RigGraphDisplaySettings.NodeRunLowerBound) || bShowInstructionIndex)
 						{
 							const int32 VOffset = bSelected ? -2 : 2;
-							const FVector2D TextSize = InstructionCountTextBlockWidget->GetDesiredSize();
+							const FVector2f TextSize = InstructionCountTextBlockWidget->GetDesiredSize();
 							FOverlayWidgetInfo Info;
-							Info.OverlayOffset = FVector2D(WidgetSize.X - TextSize.X - 8.f, VOffset - TextSize.Y);
+							Info.OverlayOffset = FVector2f(WidgetSize.X - TextSize.X - 8.f, VOffset - TextSize.Y);
 							Info.Widget = InstructionCountTextBlockWidget;
 							Widgets.Add(Info);
 						}
@@ -1070,9 +1070,9 @@ TArray<FOverlayWidgetInfo> SRigVMGraphNode::GetOverlayWidgets(bool bSelected, co
 						if(MicroSeconds >= 0.0)
 						{
 							const int32 VOffset = bSelected ? -2 : 2;
-							const FVector2D TextSize = InstructionDurationTextBlockWidget->GetDesiredSize();
+							const FVector2f TextSize = InstructionDurationTextBlockWidget->GetDesiredSize();
 							FOverlayWidgetInfo Info;
-							Info.OverlayOffset = FVector2D(8.f, VOffset - TextSize.Y);
+							Info.OverlayOffset = FVector2f(8.f, VOffset - TextSize.Y);
 							Info.Widget = InstructionDurationTextBlockWidget;
 							Widgets.Add(Info);
 						}
diff --git a/Engine/Plugins/Runtime/RigVM/Source/RigVMEditor/Private/Widgets/SRigVMGraphNodeComment.cpp b/Engine/Plugins/Runtime/RigVM/Source/RigVMEditor/Private/Widgets/SRigVMGraphNodeComment.cpp
--- a/Engine/Plugins/Runtime/RigVM/Source/RigVMEditor/Private/Widgets/SRigVMGraphNodeComment.cpp
+++ b/Engine/Plugins/Runtime/RigVM/Source/RigVMEditor/Private/Widgets/SRigVMGraphNodeComment.cpp
@@ -143,7 +143,7 @@ void SRigVMGraphNodeComment::Tick(const FGeometry& AllottedGeometry, const doubl
 	SGraphNodeComment::Tick(AllottedGeometry, InCurrentTime, InDeltaTime);
 }
 
-void SRigVMGraphNodeComment::MoveTo(const FVector2D& NewPosition, FNodeSet& NodeFilter, bool bMarkDirty)
+void SRigVMGraphNodeComment::MoveTo(const FVector2f& NewPosition, FNodeSet& NodeFilter, bool bMarkDirty)
 {
 	if (!NodeFilter.Find(SharedThis(this)))
 	{
@@ -151,8 +151,8 @@ void SRigVMGraphNodeComment::MoveTo(const FVector2D& NewPosition, FNodeSet& Node
 		{
 			if (const URigVMEdGraphSchema* RigSchema = Cast<URigVMEdGraphSchema>(GraphNode->GetSchema()))
 			{
-				const FVector2D CommentOldPosition = RigSchema->GetNodePositionAtStartOfInteraction(GraphNode);
-				const FVector2D Delta = NewPosition - CommentOldPosition;
+				const FVector2f CommentOldPosition = UE::Slate::CastToVector2f(RigSchema->GetNodePositionAtStartOfInteraction(GraphNode));
+				const FVector2f Delta = NewPosition - CommentOldPosition;
 
 				TArray<UEdGraphNode*> NodesToMove = URigVMEdGraphSchema::GetNodesToMoveForNode(GraphNode);
 				for(UEdGraphNode* NodeToMove : NodesToMove)
@@ -163,15 +163,15 @@ void SRigVMGraphNodeComment::MoveTo(const FVector2D& NewPosition, FNodeSet& Node
 						continue;
 					}
 					
-					const FVector2D OldPosition = RigSchema->GetNodePositionAtStartOfInteraction(NodeToMove);
+					const FVector2f OldPosition = UE::Slate::CastToVector2f(RigSchema->GetNodePositionAtStartOfInteraction(NodeToMove));
 
 					// Don't drag note content if either of the shift keys are down.
 					FModifierKeysState KeysState = FSlateApplication::Get().GetModifierKeys();
 
-					if((NodeToMove == GraphNode) || (IsNodeUnderComment(CommentOldPosition, OldPosition) && !KeysState.IsShiftDown()))
+					if((NodeToMove == GraphNode) || (IsNodeUnderComment(FVector2D(CommentOldPosition), FVector2D(OldPosition)) && !KeysState.IsShiftDown()))
 					{ 
-						const FVector2D AccumulatedPosition = OldPosition + Delta; 
-						RigSchema->SetNodePosition(NodeToMove, AccumulatedPosition, false);
+						const FVector2f AccumulatedPosition = OldPosition + Delta; 
+						RigSchema->SetNodePosition(NodeToMove, FVector2D(AccumulatedPosition), false);
 					}
 				}
 			}
@@ -206,7 +206,7 @@ void SRigVMGraphNodeComment::OnCommentTextCommitted(const FText& NewComment, ETe
 
 bool SRigVMGraphNodeComment::IsNodeUnderComment(UEdGraphNode_Comment* InCommentNode, const TSharedRef<SGraphNode> InNodeWidget) const
 {
-	return IsNodeUnderComment(InNodeWidget->GetPosition());
+	return IsNodeUnderComment(FVector2D(InNodeWidget->GetPosition2f()));
 }
 
 
@@ -217,7 +217,7 @@ bool SRigVMGraphNodeComment::IsNodeUnderComment(UEdGraphNode* InNode) const
 
 bool SRigVMGraphNodeComment::IsNodeUnderComment(const FVector2D& InNodePosition) const
 {
-	return IsNodeUnderComment(GetPosition(), InNodePosition);
+	return IsNodeUnderComment(FVector2D(GetPosition2f()), InNodePosition);
 }
 
 bool SRigVMGraphNodeComment::IsNodeUnderComment(const FVector2D& InCommentPosition, const FVector2D& InNodePosition) const
diff --git a/Engine/Plugins/Runtime/RigVM/Source/RigVMEditor/Private/Widgets/SRigVMGraphNodeKnot.cpp b/Engine/Plugins/Runtime/RigVM/Source/RigVMEditor/Private/Widgets/SRigVMGraphNodeKnot.cpp
--- a/Engine/Plugins/Runtime/RigVM/Source/RigVMEditor/Private/Widgets/SRigVMGraphNodeKnot.cpp
+++ b/Engine/Plugins/Runtime/RigVM/Source/RigVMEditor/Private/Widgets/SRigVMGraphNodeKnot.cpp
@@ -39,7 +39,7 @@ void SRigVMGraphNodeKnot::EndUserInteraction() const
 	SGraphNodeKnot::EndUserInteraction();
 }
 
-void SRigVMGraphNodeKnot::MoveTo(const FVector2D& NewPosition, FNodeSet& NodeFilter, bool bMarkDirty)
+void SRigVMGraphNodeKnot::MoveTo(const FVector2f& NewPosition, FNodeSet& NodeFilter, bool bMarkDirty)
 {
 	if (!NodeFilter.Find(SharedThis(this)))
 	{
@@ -47,7 +47,7 @@ void SRigVMGraphNodeKnot::MoveTo(const FVector2D& NewPosition, FNodeSet& NodeFil
 		{
 			if (const URigVMEdGraphSchema* RigSchema = Cast<URigVMEdGraphSchema>(GraphNode->GetSchema()))
 			{
-				RigSchema->SetNodePosition(GraphNode, NewPosition, false);
+				RigSchema->SetNodePosition(GraphNode, FVector2D(NewPosition), false);
 			}
 		}
 	}
diff --git a/Engine/Plugins/Runtime/RigVM/Source/RigVMEditor/Public/Widgets/SRigVMGraphNode.h b/Engine/Plugins/Runtime/RigVM/Source/RigVMEditor/Public/Widgets/SRigVMGraphNode.h
--- a/Engine/Plugins/Runtime/RigVM/Source/RigVMEditor/Public/Widgets/SRigVMGraphNode.h
+++ b/Engine/Plugins/Runtime/RigVM/Source/RigVMEditor/Public/Widgets/SRigVMGraphNode.h
@@ -33,7 +33,7 @@ public:
 	virtual TSharedRef<SWidget> CreateTitleWidget(TSharedPtr<SNodeTitle> InNodeTitle) override;
 
 	virtual void EndUserInteraction() const override;
-	virtual void MoveTo(const FVector2D& NewPosition, FNodeSet& NodeFilter, bool bMarkDirty = true) override;
+	virtual void MoveTo(const FVector2f& NewPosition, FNodeSet& NodeFilter, bool bMarkDirty = true) override;
 	virtual void AddPin( const TSharedRef<SGraphPin>& PinToAdd ) override;
 	virtual void CreateStandardPinWidget(UEdGraphPin* CurPin) override;
 	virtual void SetDefaultTitleAreaWidget(TSharedRef<SOverlay> DefaultTitleAreaWidget) override
@@ -45,9 +45,9 @@ public:
 	virtual FReply OnMouseButtonDoubleClick(const FGeometry& InMyGeometry, const FPointerEvent& InMouseEvent) override;
 
 	virtual TSharedRef<SWidget> CreateNodeContentArea() override;
-	virtual void GetOverlayBrushes(bool bSelected, const FVector2D WidgetSize, TArray<FOverlayBrushInfo>& Brushes) const override;
+	virtual void GetOverlayBrushes(bool bSelected, const FVector2f& WidgetSize, TArray<FOverlayBrushInfo>& Brushes) const override;
 	virtual void GetNodeInfoPopups(FNodeInfoContext* Context, TArray<FGraphInformationPopupInfo>& Popups) const override;
-	virtual TArray<FOverlayWidgetInfo> GetOverlayWidgets(bool bSelected, const FVector2D& WidgetSize) const override;
+	virtual TArray<FOverlayWidgetInfo> GetOverlayWidgets(bool bSelected, const FVector2f& WidgetSize) const override;
 
 	virtual void RefreshErrorInfo() override;
 	virtual void Tick(const FGeometry& AllottedGeometry, const double InCurrentTime, const float InDeltaTime) override;
diff --git a/Engine/Plugins/Runtime/RigVM/Source/RigVMEditor/Public/Widgets/SRigVMGraphNodeComment.h b/Engine/Plugins/Runtime/RigVM/Source/RigVMEditor/Public/Widgets/SRigVMGraphNodeComment.h
--- a/Engine/Plugins/Runtime/RigVM/Source/RigVMEditor/Public/Widgets/SRigVMGraphNodeComment.h
+++ b/Engine/Plugins/Runtime/RigVM/Source/RigVMEditor/Public/Widgets/SRigVMGraphNodeComment.h
@@ -16,7 +16,7 @@ public:
 	virtual FReply OnMouseButtonUp(const FGeometry& MyGeometry, const FPointerEvent& MouseEvent) override;
 	virtual void EndUserInteraction() const override;
 	virtual void Tick(const FGeometry& AllottedGeometry, const double InCurrentTime, const float InDeltaTime) override;
-	virtual void MoveTo(const FVector2D& NewPosition, FNodeSet& NodeFilter, bool bMarkDirty = true) override;
+	virtual void MoveTo(const FVector2f& NewPosition, FNodeSet& NodeFilter, bool bMarkDirty = true) override;
 
 protected:
 
diff --git a/Engine/Plugins/Runtime/RigVM/Source/RigVMEditor/Public/Widgets/SRigVMGraphNodeKnot.h b/Engine/Plugins/Runtime/RigVM/Source/RigVMEditor/Public/Widgets/SRigVMGraphNodeKnot.h
--- a/Engine/Plugins/Runtime/RigVM/Source/RigVMEditor/Public/Widgets/SRigVMGraphNodeKnot.h
+++ b/Engine/Plugins/Runtime/RigVM/Source/RigVMEditor/Public/Widgets/SRigVMGraphNodeKnot.h
@@ -15,7 +15,7 @@ class RIGVMEDITOR_API SRigVMGraphNodeKnot : public SGraphNodeKnot
 
 	// SGraphNode interface
 	virtual void EndUserInteraction() const override;
-	virtual void MoveTo(const FVector2D& NewPosition, FNodeSet& NodeFilter, bool bMarkDirty = true) override;
+	virtual void MoveTo(const FVector2f& NewPosition, FNodeSet& NodeFilter, bool bMarkDirty = true) override;
 
 	void HandleNodeBeginRemoval();
 };
\ No newline at end of file
diff --git a/Engine/Plugins/Runtime/nDisplay/Source/DisplayClusterConfigurator/Private/Views/OutputMapping/GraphNodes/SDisplayClusterConfiguratorBaseNode.cpp b/Engine/Plugins/Runtime/nDisplay/Source/DisplayClusterConfigurator/Private/Views/OutputMapping/GraphNodes/SDisplayClusterConfiguratorBaseNode.cpp
--- a/Engine/Plugins/Runtime/nDisplay/Source/DisplayClusterConfigurator/Private/Views/OutputMapping/GraphNodes/SDisplayClusterConfiguratorBaseNode.cpp
+++ b/Engine/Plugins/Runtime/nDisplay/Source/DisplayClusterConfigurator/Private/Views/OutputMapping/GraphNodes/SDisplayClusterConfiguratorBaseNode.cpp
@@ -132,8 +132,8 @@ void SDisplayClusterConfiguratorBaseNode::UpdateGraphNode()
 	SetEnabled(MakeAttributeSP(this, &SDisplayClusterConfiguratorBaseNode::IsNodeEnabled));
 
 	GetOrAddSlot(ENodeZone::BottomRight)
-		.SlotSize(FVector2D(ResizeHandleSize))
-		.SlotOffset(TAttribute<FVector2D>(this, &SDisplayClusterConfiguratorBaseNode::GetReizeHandleOffset))
+		.SlotSize2f(FVector2f(ResizeHandleSize))
+		.SlotOffset2f(TAttribute<FVector2f>(this, &SDisplayClusterConfiguratorBaseNode::GetReizeHandleOffset))
 		.VAlign(VAlign_Top)
 		.HAlign(HAlign_Left)
 		.AllowScaling(false)
@@ -157,7 +157,7 @@ FVector2D SDisplayClusterConfiguratorBaseNode::ComputeDesiredSize(float) const
 	return GetSize();
 }
 
-void SDisplayClusterConfiguratorBaseNode::MoveTo(const FVector2D& NewPosition, FNodeSet& NodeFilter, bool bMarkDirty)
+void SDisplayClusterConfiguratorBaseNode::MoveTo(const FVector2f& NewPosition, FNodeSet& NodeFilter, bool bMarkDirty)
 {
 	if (!IsNodeUnlocked())
 	{
@@ -170,8 +170,8 @@ void SDisplayClusterConfiguratorBaseNode::MoveTo(const FVector2D& NewPosition, F
 
 	TSharedRef<IDisplayClusterConfiguratorViewOutputMapping> OutputMapping = Toolkit->GetViewOutputMapping();
 
-	FVector2D CurrentPosition = FVector2D(GraphNode->NodePosX, GraphNode->NodePosY);
-	const FVector2D Offset = NewPosition - CurrentPosition;
+	FVector2f CurrentPosition = FVector2f(GraphNode->NodePosX, GraphNode->NodePosY);
+	const FVector2f Offset = NewPosition - CurrentPosition;
 
 	UDisplayClusterConfiguratorBaseNode* EdNode = GetGraphNodeChecked<UDisplayClusterConfiguratorBaseNode>();
 	UDisplayClusterConfiguratorBaseNode* ParentEdNode = EdNode->GetParent();
@@ -229,7 +229,7 @@ void SDisplayClusterConfiguratorBaseNode::MoveTo(const FVector2D& NewPosition, F
 		return;
 	}
 
-	FVector2D BestOffset = Offset;
+	FVector2D BestOffset = FVector2D(Offset);
 	if (!bIsNodeFiltered)
 	{
 		const bool bIsOverlappingAllowed = OutputMapping->GetOutputMappingSettings().bAllowClusterItemOverlap;
@@ -250,7 +250,7 @@ void SDisplayClusterConfiguratorBaseNode::MoveTo(const FVector2D& NewPosition, F
 		}
 	}
 
-	FVector2D AlignmentOffset = FVector2D::ZeroVector;
+	FVector2f AlignmentOffset = FVector2f::ZeroVector;
 	if (CanSnapAlign() && CanNodeBeSnapAligned() && !bIsNodeFiltered)
 	{
 		const FNodeAlignmentSettings& AlignmentSettings = OutputMapping->GetNodeAlignmentSettings();
@@ -262,7 +262,7 @@ void SDisplayClusterConfiguratorBaseNode::MoveTo(const FVector2D& NewPosition, F
 		Params.SnapAdjacentEdgesPadding = AlignmentSettings.AdjacentEdgesSnapPadding;
 
 		FNodeAlignmentPair Alignments = EdNode->GetTranslationAlignments(BestOffset, Params, SelectedBaseNodes);
-		AlignmentOffset = Alignments.GetOffset();
+		AlignmentOffset = UE::Slate::CastToVector2f(Alignments.GetOffset());
 
 		UpdateAlignmentTarget(XAlignmentTarget, Alignments.XAlignment, Alignments.XAlignment.TargetNode == ParentEdNode);
 		UpdateAlignmentTarget(YAlignmentTarget, Alignments.YAlignment, Alignments.YAlignment.TargetNode == ParentEdNode);
@@ -273,7 +273,7 @@ void SDisplayClusterConfiguratorBaseNode::MoveTo(const FVector2D& NewPosition, F
 		YAlignmentTarget.TargetNode.Reset();
 	}
 
-	SGraphNode::MoveTo(CurrentPosition + BestOffset + AlignmentOffset, NodeFilter, bMarkDirty);
+	SGraphNode::MoveTo(CurrentPosition + UE::Slate::CastToVector2f(BestOffset) + AlignmentOffset, NodeFilter, bMarkDirty);
 
 	// Nodes that aren't being directly interacted with aren't allowed to update their positions themselves since that might cause
 	// issues with overlapping and alignment, so the node being directly dragged and moved needs to update the positions of all selected
@@ -326,7 +326,7 @@ const FSlateBrush* SDisplayClusterConfiguratorBaseNode::GetShadowBrush(bool bSel
 	return FAppStyle::GetNoBrush();
 }
 
-bool SDisplayClusterConfiguratorBaseNode::CanBeSelected(const FVector2D& MousePositionInNode) const
+bool SDisplayClusterConfiguratorBaseNode::CanBeSelected(const FVector2f& MousePositionInNode) const
 {
 	return IsNodeEnabled() && IsNodeUnlocked();
 }
@@ -341,7 +341,7 @@ int32 SDisplayClusterConfiguratorBaseNode::GetSortDepth() const
 	return GetNodeLogicalLayer();
 }
 
-TArray<FOverlayWidgetInfo> SDisplayClusterConfiguratorBaseNode::GetOverlayWidgets(bool bSelected, const FVector2D& WidgetSize) const
+TArray<FOverlayWidgetInfo> SDisplayClusterConfiguratorBaseNode::GetOverlayWidgets(bool bSelected, const FVector2f& WidgetSize) const
 {
 	TArray<FOverlayWidgetInfo> Widgets = SGraphNode::GetOverlayWidgets(bSelected, WidgetSize);
 
@@ -520,11 +520,11 @@ int32 SDisplayClusterConfiguratorBaseNode::GetNodeVisualLayer() const
 	return EdNode->GetNodeLayer(ObjectPtrDecay(GetOwnerPanel()->SelectionManager.SelectedNodes));
 }
 
-FVector2D SDisplayClusterConfiguratorBaseNode::GetReizeHandleOffset() const
+FVector2f SDisplayClusterConfiguratorBaseNode::GetReizeHandleOffset() const
 {
-	const FVector2D NodeSize = ComputeDesiredSize(FSlateApplication::Get().GetApplicationScale());
+	const FVector2f NodeSize = UE::Slate::CastToVector2f(ComputeDesiredSize(FSlateApplication::Get().GetApplicationScale()));
 	const float GraphZoom = GetOwnerPanel()->GetZoomAmount();
-	return FVector2D(NodeSize.X, NodeSize.Y) * GraphZoom;
+	return FVector2f(NodeSize.X, NodeSize.Y) * GraphZoom;
 }
 
 EVisibility SDisplayClusterConfiguratorBaseNode::GetResizeHandleVisibility() const
@@ -569,7 +569,7 @@ void SDisplayClusterConfiguratorBaseNode::UpdateAlignmentTarget(FAlignmentRulerT
 	}
 }
 
-void SDisplayClusterConfiguratorBaseNode::AddAlignmentRulerToOverlay(TArray<FOverlayWidgetInfo>& OverlayWidgets, TSharedPtr<SAlignmentRuler> RulerWidget, const FAlignmentRulerTarget& Target, const FVector2D& WidgetSize) const
+void SDisplayClusterConfiguratorBaseNode::AddAlignmentRulerToOverlay(TArray<FOverlayWidgetInfo>& OverlayWidgets, TSharedPtr<SAlignmentRuler> RulerWidget, const FAlignmentRulerTarget& Target, const FVector2f& WidgetSize) const
 {
 	if (!Target.TargetNode.IsValid())
 	{
@@ -585,7 +585,7 @@ void SDisplayClusterConfiguratorBaseNode::AddAlignmentRulerToOverlay(TArray<FOve
 	const int32 CrossAxis = 1 - Axis;
 
 	float RulerLength = 0;
-	FVector2D RulerOffset = FVector2D::ZeroVector;
+	FVector2f RulerOffset = FVector2f::ZeroVector;
 
 	if (Target.bIsAdjacent)
 	{
@@ -594,9 +594,9 @@ void SDisplayClusterConfiguratorBaseNode::AddAlignmentRulerToOverlay(TArray<FOve
 	}
 	else
 	{
-		const FVector2D ThisNodePosition = GetPosition();
-		const FVector2D TargetNodePosition = Target.TargetNode->GetNodePosition();
-		const FVector2D TargetNodeSize = Target.TargetNode->GetNodeSize();
+		const FVector2f ThisNodePosition = GetPosition2f();
+		const FVector2f TargetNodePosition = UE::Slate::CastToVector2f(Target.TargetNode->GetNodePosition());
+		const FVector2f TargetNodeSize = UE::Slate::CastToVector2f(Target.TargetNode->GetNodeSize());
 
 		if (Target.bIsTargetingParent)
 		{
diff --git a/Engine/Plugins/Runtime/nDisplay/Source/DisplayClusterConfigurator/Private/Views/OutputMapping/GraphNodes/SDisplayClusterConfiguratorBaseNode.h b/Engine/Plugins/Runtime/nDisplay/Source/DisplayClusterConfigurator/Private/Views/OutputMapping/GraphNodes/SDisplayClusterConfiguratorBaseNode.h
--- a/Engine/Plugins/Runtime/nDisplay/Source/DisplayClusterConfigurator/Private/Views/OutputMapping/GraphNodes/SDisplayClusterConfiguratorBaseNode.h
+++ b/Engine/Plugins/Runtime/nDisplay/Source/DisplayClusterConfigurator/Private/Views/OutputMapping/GraphNodes/SDisplayClusterConfiguratorBaseNode.h
@@ -72,13 +72,13 @@ public:
 	//~ Begin SGraphNode interface
 	virtual void UpdateGraphNode() override;
 	virtual FVector2D ComputeDesiredSize(float) const override;
-	virtual void MoveTo(const FVector2D& NewPosition, FNodeSet& NodeFilter, bool bMarkDirty = true) override;
+	virtual void MoveTo(const FVector2f& NewPosition, FNodeSet& NodeFilter, bool bMarkDirty = true) override;
 	virtual void EndUserInteraction() const override;
 	virtual const FSlateBrush* GetShadowBrush(bool bSelected) const override;
-	virtual bool CanBeSelected(const FVector2D& MousePositionInNode) const override;
+	virtual bool CanBeSelected(const FVector2f& MousePositionInNode) const override;
 	virtual bool ShouldAllowCulling() const override;
 	virtual int32 GetSortDepth() const override;
-	virtual TArray<FOverlayWidgetInfo> GetOverlayWidgets(bool bSelected, const FVector2D& WidgetSize) const override;
+	virtual TArray<FOverlayWidgetInfo> GetOverlayWidgets(bool bSelected, const FVector2f& WidgetSize) const override;
 	//~ End SGraphNode interface
 
 	/**
@@ -139,12 +139,12 @@ protected:
 	virtual int32 GetNodeLogicalLayer() const;
 	virtual int32 GetNodeVisualLayer() const;
 
-	FVector2D GetReizeHandleOffset() const;
+	FVector2f GetReizeHandleOffset() const;
 	EVisibility GetResizeHandleVisibility() const;
 
 	bool CanSnapAlign() const;
 	void UpdateAlignmentTarget(FAlignmentRulerTarget& OutTarget, const FNodeAlignment& Alignment, bool bIsTargetingParent);
-	void AddAlignmentRulerToOverlay(TArray<FOverlayWidgetInfo>& OverlayWidgets, TSharedPtr<SAlignmentRuler> RulerWidget, const FAlignmentRulerTarget& Target, const FVector2D& WidgetSize) const;
+	void AddAlignmentRulerToOverlay(TArray<FOverlayWidgetInfo>& OverlayWidgets, TSharedPtr<SAlignmentRuler> RulerWidget, const FAlignmentRulerTarget& Target, const FVector2f& WidgetSize) const;
 
 	template<class TObjectType>
 	TObjectType* GetGraphNodeChecked() const
diff --git a/Engine/Plugins/Runtime/nDisplay/Source/DisplayClusterConfigurator/Private/Views/OutputMapping/GraphNodes/SDisplayClusterConfiguratorCanvasNode.cpp b/Engine/Plugins/Runtime/nDisplay/Source/DisplayClusterConfigurator/Private/Views/OutputMapping/GraphNodes/SDisplayClusterConfiguratorCanvasNode.cpp
--- a/Engine/Plugins/Runtime/nDisplay/Source/DisplayClusterConfigurator/Private/Views/OutputMapping/GraphNodes/SDisplayClusterConfiguratorCanvasNode.cpp
+++ b/Engine/Plugins/Runtime/nDisplay/Source/DisplayClusterConfigurator/Private/Views/OutputMapping/GraphNodes/SDisplayClusterConfiguratorCanvasNode.cpp
@@ -75,7 +75,7 @@ void SDisplayClusterConfiguratorCanvasNode::UpdateGraphNode()
 	];
 }
 
-void SDisplayClusterConfiguratorCanvasNode::MoveTo(const FVector2D& NewPosition, FNodeSet& NodeFilter, bool bMarkDirty)
+void SDisplayClusterConfiguratorCanvasNode::MoveTo(const FVector2f& NewPosition, FNodeSet& NodeFilter, bool bMarkDirty)
 {
 	// Canvas node is not allowed to be moved in general, so add it to the node filter
 	NodeFilter.Add(SharedThis(this));
@@ -89,22 +89,22 @@ FVector2D SDisplayClusterConfiguratorCanvasNode::ComputeDesiredSize(float) const
 	return FVector2D(NodeSize.X + CanvasPadding.Left + CanvasPadding.Right, NodeSize.Y + CanvasPadding.Top + CanvasPadding.Bottom);
 }
 
-FVector2D SDisplayClusterConfiguratorCanvasNode::GetPosition() const
+FVector2f SDisplayClusterConfiguratorCanvasNode::GetPosition2f() const
 {
-	const FVector2D NodePosition = SDisplayClusterConfiguratorBaseNode::GetPosition();
+	const FVector2f NodePosition = SDisplayClusterConfiguratorBaseNode::GetPosition2f();
 
 	// Offset node position by the top and left margin of the canvas padding
-	return NodePosition - FVector2D(CanvasPadding.Left, CanvasPadding.Top);
+	return NodePosition - FVector2f(CanvasPadding.Left, CanvasPadding.Top);
 }
 
-TArray<FOverlayWidgetInfo> SDisplayClusterConfiguratorCanvasNode::GetOverlayWidgets(bool bSelected, const FVector2D& WidgetSize) const
+TArray<FOverlayWidgetInfo> SDisplayClusterConfiguratorCanvasNode::GetOverlayWidgets(bool bSelected, const FVector2f& WidgetSize) const
 {
 	TArray<FOverlayWidgetInfo> Widgets = SDisplayClusterConfiguratorBaseNode::GetOverlayWidgets(bSelected, WidgetSize);
 
-	const FVector2D TextSize = CanvasSizeTextWidget->GetDesiredSize();
+	const FVector2f TextSize = CanvasSizeTextWidget->GetDesiredSize();
 
 	FOverlayWidgetInfo Info;
-	Info.OverlayOffset = FVector2D(0.5f * (WidgetSize.X - TextSize.X), WidgetSize.Y);
+	Info.OverlayOffset = FVector2f(0.5f * (WidgetSize.X - TextSize.X), WidgetSize.Y);
 	Info.Widget = CanvasSizeTextWidget;
 
 	Widgets.Add(Info);
diff --git a/Engine/Plugins/Runtime/nDisplay/Source/DisplayClusterConfigurator/Private/Views/OutputMapping/GraphNodes/SDisplayClusterConfiguratorCanvasNode.h b/Engine/Plugins/Runtime/nDisplay/Source/DisplayClusterConfigurator/Private/Views/OutputMapping/GraphNodes/SDisplayClusterConfiguratorCanvasNode.h
--- a/Engine/Plugins/Runtime/nDisplay/Source/DisplayClusterConfigurator/Private/Views/OutputMapping/GraphNodes/SDisplayClusterConfiguratorCanvasNode.h
+++ b/Engine/Plugins/Runtime/nDisplay/Source/DisplayClusterConfigurator/Private/Views/OutputMapping/GraphNodes/SDisplayClusterConfiguratorCanvasNode.h
@@ -25,10 +25,10 @@ public:
 
 	//~ SGraphNode interface
 	virtual void UpdateGraphNode() override;
-	virtual void MoveTo(const FVector2D& NewPosition, FNodeSet& NodeFilter, bool bMarkDirty = true) override;
+	virtual void MoveTo(const FVector2f& NewPosition, FNodeSet& NodeFilter, bool bMarkDirty = true) override;
 	virtual FVector2D ComputeDesiredSize(float) const override;
-	virtual FVector2D GetPosition() const override;
-	virtual TArray<FOverlayWidgetInfo> GetOverlayWidgets(bool bSelected, const FVector2D& WidgetSize) const override;
+	virtual FVector2f GetPosition2f() const override;
+	virtual TArray<FOverlayWidgetInfo> GetOverlayWidgets(bool bSelected, const FVector2f& WidgetSize) const override;
 	//~ End SGraphNode interface
 
 	//~ Begin SDisplayClusterConfiguratorBaseNode interface
diff --git a/Engine/Plugins/Runtime/nDisplay/Source/DisplayClusterConfigurator/Private/Views/OutputMapping/GraphNodes/SDisplayClusterConfiguratorHostNode.cpp b/Engine/Plugins/Runtime/nDisplay/Source/DisplayClusterConfigurator/Private/Views/OutputMapping/GraphNodes/SDisplayClusterConfiguratorHostNode.cpp
--- a/Engine/Plugins/Runtime/nDisplay/Source/DisplayClusterConfigurator/Private/Views/OutputMapping/GraphNodes/SDisplayClusterConfiguratorHostNode.cpp
+++ b/Engine/Plugins/Runtime/nDisplay/Source/DisplayClusterConfigurator/Private/Views/OutputMapping/GraphNodes/SDisplayClusterConfiguratorHostNode.cpp
@@ -407,7 +407,7 @@ void SDisplayClusterConfiguratorHostNode::UpdateGraphNode()
 	];
 }
 
-void SDisplayClusterConfiguratorHostNode::MoveTo(const FVector2D& NewPosition, FNodeSet& NodeFilter, bool bMarkDirty)
+void SDisplayClusterConfiguratorHostNode::MoveTo(const FVector2f& NewPosition, FNodeSet& NodeFilter, bool bMarkDirty)
 {
 	UDisplayClusterConfiguratorHostNode* HostEdNode = GetGraphNodeChecked<UDisplayClusterConfiguratorHostNode>();
 	if (!HostEdNode->CanUserMoveNode())
@@ -423,9 +423,9 @@ FVector2D SDisplayClusterConfiguratorHostNode::ComputeDesiredSize(float) const
 	return GetSize() + FVector2D(BorderThickness.Left + BorderThickness.Right, BorderThickness.Top + BorderThickness.Bottom);
 }
 
-FVector2D SDisplayClusterConfiguratorHostNode::GetPosition() const
+FVector2f SDisplayClusterConfiguratorHostNode::GetPosition2f() const
 {
-	return SDisplayClusterConfiguratorBaseNode::GetPosition() - FVector2D(BorderThickness.Left, BorderThickness.Top);
+	return SDisplayClusterConfiguratorBaseNode::GetPosition2f() - FVector2f(BorderThickness.Left, BorderThickness.Top);
 }
 
 void SDisplayClusterConfiguratorHostNode::SetNodeSize(const FVector2D InLocalSize, bool bFixedAspectRatio)
diff --git a/Engine/Plugins/Runtime/nDisplay/Source/DisplayClusterConfigurator/Private/Views/OutputMapping/GraphNodes/SDisplayClusterConfiguratorHostNode.h b/Engine/Plugins/Runtime/nDisplay/Source/DisplayClusterConfigurator/Private/Views/OutputMapping/GraphNodes/SDisplayClusterConfiguratorHostNode.h
--- a/Engine/Plugins/Runtime/nDisplay/Source/DisplayClusterConfigurator/Private/Views/OutputMapping/GraphNodes/SDisplayClusterConfiguratorHostNode.h
+++ b/Engine/Plugins/Runtime/nDisplay/Source/DisplayClusterConfigurator/Private/Views/OutputMapping/GraphNodes/SDisplayClusterConfiguratorHostNode.h
@@ -23,9 +23,9 @@ public:
 
 	//~ Begin SGraphNode interface
 	virtual void UpdateGraphNode() override;
-	virtual void MoveTo(const FVector2D& NewPosition, FNodeSet& NodeFilter, bool bMarkDirty = true) override;
+	virtual void MoveTo(const FVector2f& NewPosition, FNodeSet& NodeFilter, bool bMarkDirty = true) override;
 	virtual FVector2D ComputeDesiredSize(float) const override;
-	virtual FVector2D GetPosition() const override;
+	virtual FVector2f GetPosition2f() const override;
 	//~ End SGraphNode interface
 
 	//~ Begin SDisplayClusterConfiguratorBaseNode interface
diff --git a/Engine/Plugins/Runtime/nDisplay/Source/DisplayClusterConfigurator/Private/Views/OutputMapping/GraphNodes/SDisplayClusterConfiguratorWindowNode.cpp b/Engine/Plugins/Runtime/nDisplay/Source/DisplayClusterConfigurator/Private/Views/OutputMapping/GraphNodes/SDisplayClusterConfiguratorWindowNode.cpp
--- a/Engine/Plugins/Runtime/nDisplay/Source/DisplayClusterConfigurator/Private/Views/OutputMapping/GraphNodes/SDisplayClusterConfiguratorWindowNode.cpp
+++ b/Engine/Plugins/Runtime/nDisplay/Source/DisplayClusterConfigurator/Private/Views/OutputMapping/GraphNodes/SDisplayClusterConfiguratorWindowNode.cpp
@@ -358,9 +358,9 @@ FVector2D SDisplayClusterConfiguratorWindowNode::ComputeDesiredSize(float) const
 	return GetSize() * WindowScaleFactor;
 }
 
-FVector2D SDisplayClusterConfiguratorWindowNode::GetPosition() const
+FVector2f SDisplayClusterConfiguratorWindowNode::GetPosition2f() const
 {
-	return SDisplayClusterConfiguratorBaseNode::GetPosition() * WindowScaleFactor;
+	return SDisplayClusterConfiguratorBaseNode::GetPosition2f() * UE::Slate::CastToVector2f(WindowScaleFactor);
 }
 
 bool SDisplayClusterConfiguratorWindowNode::IsAspectRatioFixed() const
diff --git a/Engine/Plugins/Runtime/nDisplay/Source/DisplayClusterConfigurator/Private/Views/OutputMapping/GraphNodes/SDisplayClusterConfiguratorWindowNode.h b/Engine/Plugins/Runtime/nDisplay/Source/DisplayClusterConfigurator/Private/Views/OutputMapping/GraphNodes/SDisplayClusterConfiguratorWindowNode.h
--- a/Engine/Plugins/Runtime/nDisplay/Source/DisplayClusterConfigurator/Private/Views/OutputMapping/GraphNodes/SDisplayClusterConfiguratorWindowNode.h
+++ b/Engine/Plugins/Runtime/nDisplay/Source/DisplayClusterConfigurator/Private/Views/OutputMapping/GraphNodes/SDisplayClusterConfiguratorWindowNode.h
@@ -34,7 +34,7 @@ public:
 	//~ Begin SGraphNode interface
 	virtual void UpdateGraphNode() override;
 	virtual FVector2D ComputeDesiredSize(float) const override;
-	virtual FVector2D GetPosition() const override;
+	virtual FVector2f GetPosition2f() const override;
 	//~ End SGraphNode interface
 
 	//~ Begin SDisplayClusterConfiguratorBaseNode interface
diff --git a/Engine/Plugins/TextureGraph/Source/TextureGraphEditor/Private/EdGraph/STG_EditorGraphNode.cpp b/Engine/Plugins/TextureGraph/Source/TextureGraphEditor/Private/EdGraph/STG_EditorGraphNode.cpp
--- a/Engine/Plugins/TextureGraph/Source/TextureGraphEditor/Private/EdGraph/STG_EditorGraphNode.cpp
+++ b/Engine/Plugins/TextureGraph/Source/TextureGraphEditor/Private/EdGraph/STG_EditorGraphNode.cpp
@@ -315,8 +315,8 @@ void STG_EditorGraphNode::UpdateGraphNode()
 			.IsGraphNodeHovered(this, &SGraphNode::IsHovered);
 
 		GetOrAddSlot(ENodeZone::TopCenter)
-			.SlotOffset(TAttribute<FVector2D>(CommentBubble.Get(), &SCommentBubble::GetOffset))
-			.SlotSize(TAttribute<FVector2D>(CommentBubble.Get(), &SCommentBubble::GetSize))
+			.SlotOffset2f(TAttribute<FVector2f>(CommentBubble.Get(), &SCommentBubble::GetOffset2f))
+			.SlotSize2f(TAttribute<FVector2f>(CommentBubble.Get(), &SCommentBubble::GetSize2f))
 			.AllowScaling(TAttribute<bool>(CommentBubble.Get(), &SCommentBubble::IsScalingAllowed))
 			.VAlign(VAlign_Top)
 			[
diff --git a/Engine/Source/Editor/AnimationBlueprintEditor/Private/AnimationNodes/SAnimationGraphNode.cpp b/Engine/Source/Editor/AnimationBlueprintEditor/Private/AnimationNodes/SAnimationGraphNode.cpp
--- a/Engine/Source/Editor/AnimationBlueprintEditor/Private/AnimationNodes/SAnimationGraphNode.cpp
+++ b/Engine/Source/Editor/AnimationBlueprintEditor/Private/AnimationNodes/SAnimationGraphNode.cpp
@@ -104,7 +104,7 @@ void SAnimationGraphNode::Tick(const FGeometry& AllottedGeometry, const double I
 	}
 }
 
-TArray<FOverlayWidgetInfo> SAnimationGraphNode::GetOverlayWidgets(bool bSelected, const FVector2D& WidgetSize) const
+TArray<FOverlayWidgetInfo> SAnimationGraphNode::GetOverlayWidgets(bool bSelected, const FVector2f& WidgetSize) const
 {
 	TArray<FOverlayWidgetInfo> Widgets;
 
@@ -115,7 +115,7 @@ TArray<FOverlayWidgetInfo> SAnimationGraphNode::GetOverlayWidgets(bool bSelected
 			const FSlateBrush* ImageBrush = FAppStyle::Get().GetBrush(TEXT("Graph.AnimationFastPathIndicator"));
 
 			FOverlayWidgetInfo Info;
-			Info.OverlayOffset = FVector2D(WidgetSize.X - (ImageBrush->ImageSize.X * 0.5f), -(ImageBrush->ImageSize.Y * 0.5f));
+			Info.OverlayOffset = FVector2f(WidgetSize.X - (ImageBrush->ImageSize.X * 0.5f), -(ImageBrush->ImageSize.Y * 0.5f));
 			Info.Widget = IndicatorWidget;
 
 			Widgets.Add(Info);
diff --git a/Engine/Source/Editor/AnimationBlueprintEditor/Private/AnimationNodes/SGraphNodeBlendSpaceGraph.cpp b/Engine/Source/Editor/AnimationBlueprintEditor/Private/AnimationNodes/SGraphNodeBlendSpaceGraph.cpp
--- a/Engine/Source/Editor/AnimationBlueprintEditor/Private/AnimationNodes/SGraphNodeBlendSpaceGraph.cpp
+++ b/Engine/Source/Editor/AnimationBlueprintEditor/Private/AnimationNodes/SGraphNodeBlendSpaceGraph.cpp
@@ -79,7 +79,7 @@ UEdGraph* SGraphNodeBlendSpaceGraph::GetInnerGraph() const
 	return BlendSpaceNode->GetBlendSpaceGraph();
 }
 
-TArray<FOverlayWidgetInfo> SGraphNodeBlendSpaceGraph::GetOverlayWidgets(bool bSelected, const FVector2D& WidgetSize) const
+TArray<FOverlayWidgetInfo> SGraphNodeBlendSpaceGraph::GetOverlayWidgets(bool bSelected, const FVector2f& WidgetSize) const
 {
 	TArray<FOverlayWidgetInfo> Widgets;
 
@@ -90,7 +90,7 @@ TArray<FOverlayWidgetInfo> SGraphNodeBlendSpaceGraph::GetOverlayWidgets(bool bSe
 			const FSlateBrush* ImageBrush = FAppStyle::Get().GetBrush(TEXT("Graph.AnimationFastPathIndicator"));
 
 			FOverlayWidgetInfo Info;
-			Info.OverlayOffset = FVector2D(WidgetSize.X - (ImageBrush->ImageSize.X * 0.5f), -(ImageBrush->ImageSize.Y * 0.5f));
+			Info.OverlayOffset = FVector2f(WidgetSize.X - (ImageBrush->ImageSize.X * 0.5f), -(ImageBrush->ImageSize.Y * 0.5f));
 			Info.Widget = FastPathWidget;
 
 			Widgets.Add(Info);
diff --git a/Engine/Source/Editor/AnimationBlueprintEditor/Private/AnimationNodes/SGraphNodeBlendSpaceGraph.h b/Engine/Source/Editor/AnimationBlueprintEditor/Private/AnimationNodes/SGraphNodeBlendSpaceGraph.h
--- a/Engine/Source/Editor/AnimationBlueprintEditor/Private/AnimationNodes/SGraphNodeBlendSpaceGraph.h
+++ b/Engine/Source/Editor/AnimationBlueprintEditor/Private/AnimationNodes/SGraphNodeBlendSpaceGraph.h
@@ -25,7 +25,7 @@ public:
 protected:
 	// SGraphNodeK2Composite interface
 	virtual UEdGraph* GetInnerGraph() const override;
-	virtual TArray<FOverlayWidgetInfo> GetOverlayWidgets(bool bSelected, const FVector2D& WidgetSize) const override;
+	virtual TArray<FOverlayWidgetInfo> GetOverlayWidgets(bool bSelected, const FVector2f& WidgetSize) const override;
 	virtual TSharedRef<SWidget> CreateNodeBody() override;
 	
 	// SGraphNode interface
diff --git a/Engine/Source/Editor/AnimationBlueprintEditor/Private/AnimationNodes/SGraphNodeStateMachineInstance.cpp b/Engine/Source/Editor/AnimationBlueprintEditor/Private/AnimationNodes/SGraphNodeStateMachineInstance.cpp
--- a/Engine/Source/Editor/AnimationBlueprintEditor/Private/AnimationNodes/SGraphNodeStateMachineInstance.cpp
+++ b/Engine/Source/Editor/AnimationBlueprintEditor/Private/AnimationNodes/SGraphNodeStateMachineInstance.cpp
@@ -38,7 +38,7 @@ UEdGraph* SGraphNodeStateMachineInstance::GetInnerGraph() const
 	return StateMachineInstance->EditorStateMachineGraph;
 }
 
-TArray<FOverlayWidgetInfo> SGraphNodeStateMachineInstance::GetOverlayWidgets(bool bSelected, const FVector2D& WidgetSize) const
+TArray<FOverlayWidgetInfo> SGraphNodeStateMachineInstance::GetOverlayWidgets(bool bSelected, const FVector2f& WidgetSize) const
 {
 	TArray<FOverlayWidgetInfo> Widgets;
 
diff --git a/Engine/Source/Editor/AnimationBlueprintEditor/Private/AnimationNodes/SGraphNodeStateMachineInstance.h b/Engine/Source/Editor/AnimationBlueprintEditor/Private/AnimationNodes/SGraphNodeStateMachineInstance.h
--- a/Engine/Source/Editor/AnimationBlueprintEditor/Private/AnimationNodes/SGraphNodeStateMachineInstance.h
+++ b/Engine/Source/Editor/AnimationBlueprintEditor/Private/AnimationNodes/SGraphNodeStateMachineInstance.h
@@ -23,7 +23,7 @@ public:
 protected:
 	// SGraphNodeK2Composite interface
 	virtual UEdGraph* GetInnerGraph() const override;
-	virtual TArray<FOverlayWidgetInfo> GetOverlayWidgets(bool bSelected, const FVector2D& WidgetSize) const override;
+	virtual TArray<FOverlayWidgetInfo> GetOverlayWidgets(bool bSelected, const FVector2f& WidgetSize) const override;
 	virtual TSharedRef<SWidget> CreateNodeBody() override;
 	// End of SGraphNodeK2Composite interface
 
diff --git a/Engine/Source/Editor/AnimationBlueprintEditor/Private/AnimationStateNodes/SGraphNodeAnimTransition.cpp b/Engine/Source/Editor/AnimationBlueprintEditor/Private/AnimationStateNodes/SGraphNodeAnimTransition.cpp
--- a/Engine/Source/Editor/AnimationBlueprintEditor/Private/AnimationStateNodes/SGraphNodeAnimTransition.cpp
+++ b/Engine/Source/Editor/AnimationBlueprintEditor/Private/AnimationStateNodes/SGraphNodeAnimTransition.cpp
@@ -65,7 +65,7 @@ void SGraphNodeAnimTransition::GetNodeInfoPopups(FNodeInfoContext* Context, TArr
 {
 }
 
-void SGraphNodeAnimTransition::MoveTo(const FVector2D& NewPosition, FNodeSet& NodeFilter, bool bMarkDirty)
+void SGraphNodeAnimTransition::MoveTo(const FVector2f& NewPosition, FNodeSet& NodeFilter, bool bMarkDirty)
 {
 	// Ignored; position is set by the location of the attached state nodes
 }
diff --git a/Engine/Source/Editor/AnimationBlueprintEditor/Private/AnimationStateNodes/SGraphNodeAnimTransition.h b/Engine/Source/Editor/AnimationBlueprintEditor/Private/AnimationStateNodes/SGraphNodeAnimTransition.h
--- a/Engine/Source/Editor/AnimationBlueprintEditor/Private/AnimationStateNodes/SGraphNodeAnimTransition.h
+++ b/Engine/Source/Editor/AnimationBlueprintEditor/Private/AnimationStateNodes/SGraphNodeAnimTransition.h
@@ -34,7 +34,7 @@ public:
 
 	// SNodePanel::SNode interface
 	virtual void GetNodeInfoPopups(FNodeInfoContext* Context, TArray<FGraphInformationPopupInfo>& Popups) const override;
-	virtual void MoveTo(const FVector2D& NewPosition, FNodeSet& NodeFilter, bool bMarkDirty = true) override;
+	virtual void MoveTo(const FVector2f& NewPosition, FNodeSet& NodeFilter, bool bMarkDirty = true) override;
 	virtual bool RequiresSecondPassLayout() const override;
 	virtual void PerformSecondPassLayout(const TMap< UObject*, TSharedRef<SNode> >& NodeToWidgetLookup) const override;
 	// End of SNodePanel::SNode interface
diff --git a/Engine/Source/Editor/AnimationBlueprintEditor/Public/AnimationNodes/SAnimationGraphNode.h b/Engine/Source/Editor/AnimationBlueprintEditor/Public/AnimationNodes/SAnimationGraphNode.h
--- a/Engine/Source/Editor/AnimationBlueprintEditor/Public/AnimationNodes/SAnimationGraphNode.h
+++ b/Engine/Source/Editor/AnimationBlueprintEditor/Public/AnimationNodes/SAnimationGraphNode.h
@@ -62,7 +62,7 @@ protected:
 	// End of SWidget interface
 
 	// SGraphNode interface
-	virtual TArray<FOverlayWidgetInfo> GetOverlayWidgets(bool bSelected, const FVector2D& WidgetSize) const override;
+	virtual TArray<FOverlayWidgetInfo> GetOverlayWidgets(bool bSelected, const FVector2f& WidgetSize) const override;
 	virtual TSharedRef<SWidget> CreateTitleWidget(TSharedPtr<SNodeTitle> InNodeTitle) override;
 	virtual void GetNodeInfoPopups(FNodeInfoContext* Context, TArray<FGraphInformationPopupInfo>& Popups) const override;
 	virtual void CreateBelowPinControls(TSharedPtr<SVerticalBox> MainBox) override;
diff --git a/Engine/Source/Editor/AnimationEditorWidgets/Private/SchematicGraphPanel/SSchematicGraphPanel.cpp b/Engine/Source/Editor/AnimationEditorWidgets/Private/SchematicGraphPanel/SSchematicGraphPanel.cpp
--- a/Engine/Source/Editor/AnimationEditorWidgets/Private/SchematicGraphPanel/SSchematicGraphPanel.cpp
+++ b/Engine/Source/Editor/AnimationEditorWidgets/Private/SchematicGraphPanel/SSchematicGraphPanel.cpp
@@ -527,9 +527,9 @@ FReply SSchematicGraphNode::OnMouseButtonDown(const FGeometry& MyGeometry, const
 	return FReply::Handled().DetectDrag(SharedThis(this), EKeys::LeftMouseButton);
 }
 
-FVector2d SSchematicGraphNode::GetPosition() const
+FVector2f SSchematicGraphNode::GetPosition2f() const
 {
-	return Position->Get() - (OriginalSize*0.5);
+	return UE::Slate::CastToVector2f(Position->Get() - (OriginalSize*0.5));
 }
 
 void SSchematicGraphNode::EnablePositionAnimation(bool bEnabled)
diff --git a/Engine/Source/Editor/AnimationEditorWidgets/Public/SchematicGraphPanel/SSchematicGraphPanel.h b/Engine/Source/Editor/AnimationEditorWidgets/Public/SchematicGraphPanel/SSchematicGraphPanel.h
--- a/Engine/Source/Editor/AnimationEditorWidgets/Public/SchematicGraphPanel/SSchematicGraphPanel.h
+++ b/Engine/Source/Editor/AnimationEditorWidgets/Public/SchematicGraphPanel/SSchematicGraphPanel.h
@@ -90,7 +90,7 @@ public:
 	virtual FReply OnDragDetected(const FGeometry& MyGeometry, const FPointerEvent& MouseEvent) override;
 
 	virtual EVisibility GetNodeVisibility() const;
-	virtual FVector2d GetPosition() const override;
+	virtual FVector2f GetPosition2f() const override;
 	void EnablePositionAnimation(bool bEnabled = true);
 
 	const FVector2d& GetOriginalSize() const { return OriginalSize; }
diff --git a/Engine/Source/Editor/BehaviorTreeEditor/Private/SGraphNode_BehaviorTree.cpp b/Engine/Source/Editor/BehaviorTreeEditor/Private/SGraphNode_BehaviorTree.cpp
--- a/Engine/Source/Editor/BehaviorTreeEditor/Private/SGraphNode_BehaviorTree.cpp
+++ b/Engine/Source/Editor/BehaviorTreeEditor/Private/SGraphNode_BehaviorTree.cpp
@@ -569,8 +569,8 @@ void SGraphNode_BehaviorTree::UpdateGraphNode()
 	.IsGraphNodeHovered( this, &SGraphNode::IsHovered );
 
 	GetOrAddSlot( ENodeZone::TopCenter )
-	.SlotOffset( TAttribute<FVector2D>( CommentBubble.Get(), &SCommentBubble::GetOffset ))
-	.SlotSize( TAttribute<FVector2D>( CommentBubble.Get(), &SCommentBubble::GetSize ))
+	.SlotOffset2f( TAttribute<FVector2f>( CommentBubble.Get(), &SCommentBubble::GetOffset2f ))
+	.SlotSize2f( TAttribute<FVector2f>( CommentBubble.Get(), &SCommentBubble::GetSize2f ))
 	.AllowScaling( TAttribute<bool>( CommentBubble.Get(), &SCommentBubble::IsScalingAllowed ))
 	.VAlign( VAlign_Top )
 	[
@@ -1012,7 +1012,7 @@ EVisibility SGraphNode_BehaviorTree::GetBlueprintIconVisibility() const
 	return (bCanShowIcon && (!MyOwnerPanel.IsValid() || MyOwnerPanel->GetCurrentLOD() > EGraphRenderingLOD::LowDetail)) ? EVisibility::Visible : EVisibility::Collapsed;
 }
 
-void SGraphNode_BehaviorTree::GetOverlayBrushes(bool bSelected, const FVector2D WidgetSize, TArray<FOverlayBrushInfo>& Brushes) const
+void SGraphNode_BehaviorTree::GetOverlayBrushes(bool bSelected, const FVector2f& WidgetSize, TArray<FOverlayBrushInfo>& Brushes) const
 {
 	UBehaviorTreeGraphNode* BTNode = Cast<UBehaviorTreeGraphNode>(GraphNode);
 	if (BTNode == NULL)
@@ -1050,7 +1050,7 @@ void SGraphNode_BehaviorTree::GetOverlayBrushes(bool bSelected, const FVector2D
 				IPOverlayInfo.OverlayOffset.Y = (Overlap - IPOverlayInfo.Brush->ImageSize.Y);
 			}
 
-			IPOverlayInfo.AnimationEnvelope = FVector2D(0.f, 10.f);
+			IPOverlayInfo.AnimationEnvelope = FVector2f(0.f, 10.f);
 			Brushes.Add(IPOverlayInfo);
 		}
 
@@ -1069,7 +1069,7 @@ void SGraphNode_BehaviorTree::GetOverlayBrushes(bool bSelected, const FVector2D
 					IPOverlayInfo.OverlayOffset.X = -IPOverlayInfo.Brush->ImageSize.X;
 					IPOverlayInfo.OverlayOffset.Y = TriggerOffsets[i].Position.Y + TriggerOffsets[i].Size.Y / 2.f - IPOverlayInfo.Brush->ImageSize.Y / 2.f;
 
-					IPOverlayInfo.AnimationEnvelope = FVector2D(10.f, 0.f);
+					IPOverlayInfo.AnimationEnvelope = FVector2f(10.f, 0.f);
 					Brushes.Add(IPOverlayInfo);
 				}
 			}
@@ -1077,14 +1077,14 @@ void SGraphNode_BehaviorTree::GetOverlayBrushes(bool bSelected, const FVector2D
 	}
 }
 
-TArray<FOverlayWidgetInfo> SGraphNode_BehaviorTree::GetOverlayWidgets(bool bSelected, const FVector2D& WidgetSize) const
+TArray<FOverlayWidgetInfo> SGraphNode_BehaviorTree::GetOverlayWidgets(bool bSelected, const FVector2f& WidgetSize) const
 {
 	TArray<FOverlayWidgetInfo> Widgets;
 
 	check(NodeBody.IsValid());
 	check(IndexOverlay.IsValid());
 
-	FVector2D Origin(0.0f, 0.0f);
+	FVector2f Origin(0.0f, 0.0f);
 
 	// build overlays for decorator sub-nodes
 	for(const auto& DecoratorWidget : DecoratorWidgets)
@@ -1099,7 +1099,7 @@ TArray<FOverlayWidgetInfo> SGraphNode_BehaviorTree::GetOverlayWidgets(bool bSele
 	}
 
 	FOverlayWidgetInfo Overlay(IndexOverlay);
-	Overlay.OverlayOffset = FVector2D(WidgetSize.X - (IndexOverlay->GetDesiredSize().X * 0.5f), Origin.Y);
+	Overlay.OverlayOffset = FVector2f(WidgetSize.X - (IndexOverlay->GetDesiredSize().X * 0.5f), Origin.Y);
 	Widgets.Add(Overlay);
 
 	Origin.Y += NodeBody->GetDesiredSize().Y;
@@ -1125,7 +1125,7 @@ TSharedRef<SGraphNode> SGraphNode_BehaviorTree::GetNodeUnderMouse(const FGeometr
 	return SubNode.IsValid() ? SubNode.ToSharedRef() : StaticCastSharedRef<SGraphNode>(AsShared());
 }
 
-void SGraphNode_BehaviorTree::MoveTo(const FVector2D& NewPosition, FNodeSet& NodeFilter, bool bMarkDirty)
+void SGraphNode_BehaviorTree::MoveTo(const FVector2f& NewPosition, FNodeSet& NodeFilter, bool bMarkDirty)
 {
 	SGraphNodeAI::MoveTo(NewPosition, NodeFilter, bMarkDirty);
 
diff --git a/Engine/Source/Editor/BehaviorTreeEditor/Private/SGraphNode_BehaviorTree.h b/Engine/Source/Editor/BehaviorTreeEditor/Private/SGraphNode_BehaviorTree.h
--- a/Engine/Source/Editor/BehaviorTreeEditor/Private/SGraphNode_BehaviorTree.h
+++ b/Engine/Source/Editor/BehaviorTreeEditor/Private/SGraphNode_BehaviorTree.h
@@ -45,10 +45,10 @@ public:
 	virtual void CreatePinWidgets() override;
 	virtual void AddPin(const TSharedRef<SGraphPin>& PinToAdd) override;
 	virtual TSharedPtr<SToolTip> GetComplexTooltip() override;
-	virtual void GetOverlayBrushes(bool bSelected, const FVector2D WidgetSize, TArray<FOverlayBrushInfo>& Brushes) const override;
-	virtual TArray<FOverlayWidgetInfo> GetOverlayWidgets(bool bSelected, const FVector2D& WidgetSize) const override;
+	virtual void GetOverlayBrushes(bool bSelected, const FVector2f& WidgetSize, TArray<FOverlayBrushInfo>& Brushes) const override;
+	virtual TArray<FOverlayWidgetInfo> GetOverlayWidgets(bool bSelected, const FVector2f& WidgetSize) const override;
 	virtual TSharedRef<SGraphNode> GetNodeUnderMouse(const FGeometry& MyGeometry, const FPointerEvent& MouseEvent) override;
-	virtual void MoveTo(const FVector2D& NewPosition, FNodeSet& NodeFilter, bool bMarkDirty = true) override;
+	virtual void MoveTo(const FVector2f& NewPosition, FNodeSet& NodeFilter, bool bMarkDirty = true) override;
 	// End of SGraphNode interface
 
 	/** handle double click */
diff --git a/Engine/Source/Editor/GraphEditor/Private/KismetNodes/SGraphNodeK2Base.cpp b/Engine/Source/Editor/GraphEditor/Private/KismetNodes/SGraphNodeK2Base.cpp
--- a/Engine/Source/Editor/GraphEditor/Private/KismetNodes/SGraphNodeK2Base.cpp
+++ b/Engine/Source/Editor/GraphEditor/Private/KismetNodes/SGraphNodeK2Base.cpp
@@ -298,8 +298,8 @@ void SGraphNodeK2Base::UpdateCompactNode()
 	.IsGraphNodeHovered( this, &SGraphNodeK2Base::IsHovered );
 
 	GetOrAddSlot( ENodeZone::TopCenter )
-	.SlotOffset( TAttribute<FVector2D>( CommentBubble.Get(), &SCommentBubble::GetOffset ))
-	.SlotSize( TAttribute<FVector2D>( CommentBubble.Get(), &SCommentBubble::GetSize ))
+	.SlotOffset2f( TAttribute<FVector2f>( CommentBubble.Get(), &SCommentBubble::GetOffset2f ))
+	.SlotSize2f( TAttribute<FVector2f>( CommentBubble.Get(), &SCommentBubble::GetSize2f ))
 	.AllowScaling( TAttribute<bool>( CommentBubble.Get(), &SCommentBubble::IsScalingAllowed ))
 	.VAlign( VAlign_Top )
 	[
@@ -447,7 +447,7 @@ FText SGraphNodeK2Base::GetNodeCompactTitle() const
 }
 
 /** Populate the brushes array with any overlay brushes to render */
-void SGraphNodeK2Base::GetOverlayBrushes(bool bSelected, const FVector2D WidgetSize, TArray<FOverlayBrushInfo>& Brushes) const
+void SGraphNodeK2Base::GetOverlayBrushes(bool bSelected, const FVector2f& WidgetSize, TArray<FOverlayBrushInfo>& Brushes) const
 {
 	UBlueprint* OwnerBlueprint = FBlueprintEditorUtils::FindBlueprintForNode(GraphNode);
 
@@ -491,7 +491,7 @@ void SGraphNodeK2Base::GetOverlayBrushes(bool bSelected, const FVector2D WidgetS
 			IPOverlayInfo.OverlayOffset.Y = (Overlap - IPOverlayInfo.Brush->ImageSize.Y);
 		}
 
-		IPOverlayInfo.AnimationEnvelope = FVector2D(0.f, 10.f);
+		IPOverlayInfo.AnimationEnvelope = FVector2f(0.f, 10.f);
 
 		Brushes.Add(IPOverlayInfo);
 	}
diff --git a/Engine/Source/Editor/GraphEditor/Private/KismetNodes/SGraphNodeK2Composite.cpp b/Engine/Source/Editor/GraphEditor/Private/KismetNodes/SGraphNodeK2Composite.cpp
--- a/Engine/Source/Editor/GraphEditor/Private/KismetNodes/SGraphNodeK2Composite.cpp
+++ b/Engine/Source/Editor/GraphEditor/Private/KismetNodes/SGraphNodeK2Composite.cpp
@@ -172,8 +172,8 @@ void SGraphNodeK2Composite::UpdateGraphNode()
 	.IsGraphNodeHovered( this, &SGraphNode::IsHovered );
 
 	GetOrAddSlot( ENodeZone::TopCenter )
-	.SlotOffset( TAttribute<FVector2D>( CommentBubble.Get(), &SCommentBubble::GetOffset ))
-	.SlotSize( TAttribute<FVector2D>( CommentBubble.Get(), &SCommentBubble::GetSize ))
+	.SlotOffset2f( TAttribute<FVector2f>( CommentBubble.Get(), &SCommentBubble::GetOffset2f ))
+	.SlotSize2f( TAttribute<FVector2f>( CommentBubble.Get(), &SCommentBubble::GetSize2f ))
 	.AllowScaling( TAttribute<bool>( CommentBubble.Get(), &SCommentBubble::IsScalingAllowed ))
 	.VAlign( VAlign_Top )
 	[
diff --git a/Engine/Source/Editor/GraphEditor/Private/KismetNodes/SGraphNodeK2Var.cpp b/Engine/Source/Editor/GraphEditor/Private/KismetNodes/SGraphNodeK2Var.cpp
--- a/Engine/Source/Editor/GraphEditor/Private/KismetNodes/SGraphNodeK2Var.cpp
+++ b/Engine/Source/Editor/GraphEditor/Private/KismetNodes/SGraphNodeK2Var.cpp
@@ -358,8 +358,8 @@ void SGraphNodeK2Var::UpdateGraphNode()
 	.IsGraphNodeHovered( this, &SGraphNode::IsHovered );
 
 	GetOrAddSlot( ENodeZone::TopCenter )
-	.SlotOffset( TAttribute<FVector2D>( CommentBubble.Get(), &SCommentBubble::GetOffset ))
-	.SlotSize( TAttribute<FVector2D>( CommentBubble.Get(), &SCommentBubble::GetSize ))
+	.SlotOffset2f( TAttribute<FVector2f>( CommentBubble.Get(), &SCommentBubble::GetOffset2f ))
+	.SlotSize2f( TAttribute<FVector2f>( CommentBubble.Get(), &SCommentBubble::GetSize2f ))
 	.AllowScaling( TAttribute<bool>( CommentBubble.Get(), &SCommentBubble::IsScalingAllowed ))
 	.VAlign( VAlign_Top )
 	[
diff --git a/Engine/Source/Editor/GraphEditor/Private/MaterialNodes/SGraphNodeMaterialBase.cpp b/Engine/Source/Editor/GraphEditor/Private/MaterialNodes/SGraphNodeMaterialBase.cpp
--- a/Engine/Source/Editor/GraphEditor/Private/MaterialNodes/SGraphNodeMaterialBase.cpp
+++ b/Engine/Source/Editor/GraphEditor/Private/MaterialNodes/SGraphNodeMaterialBase.cpp
@@ -370,7 +370,7 @@ void SGraphNodeMaterialBase::CreatePinWidgets()
 	}
 }
 
-void SGraphNodeMaterialBase::MoveTo(const FVector2D& NewPosition, FNodeSet& NodeFilter, bool bMarkDirty)
+void SGraphNodeMaterialBase::MoveTo(const FVector2f& NewPosition, FNodeSet& NodeFilter, bool bMarkDirty)
 {
 	SGraphNode::MoveTo(NewPosition, NodeFilter, bMarkDirty);
 
diff --git a/Engine/Source/Editor/GraphEditor/Private/MaterialNodes/SGraphNodeMaterialComment.cpp b/Engine/Source/Editor/GraphEditor/Private/MaterialNodes/SGraphNodeMaterialComment.cpp
--- a/Engine/Source/Editor/GraphEditor/Private/MaterialNodes/SGraphNodeMaterialComment.cpp
+++ b/Engine/Source/Editor/GraphEditor/Private/MaterialNodes/SGraphNodeMaterialComment.cpp
@@ -14,7 +14,7 @@ void SGraphNodeMaterialComment::Construct(const FArguments& InArgs, class UMater
 	this->CommentNode = InNode;
 }
 
-void SGraphNodeMaterialComment::MoveTo(const FVector2D& NewPosition, FNodeSet& NodeFilter, bool bMarkDirty)
+void SGraphNodeMaterialComment::MoveTo(const FVector2f& NewPosition, FNodeSet& NodeFilter, bool bMarkDirty)
 {
 	if( !NodeFilter.Find( SharedThis( this ) ))
 	{
diff --git a/Engine/Source/Editor/GraphEditor/Private/MaterialNodes/SGraphNodeMaterialComment.h b/Engine/Source/Editor/GraphEditor/Private/MaterialNodes/SGraphNodeMaterialComment.h
--- a/Engine/Source/Editor/GraphEditor/Private/MaterialNodes/SGraphNodeMaterialComment.h
+++ b/Engine/Source/Editor/GraphEditor/Private/MaterialNodes/SGraphNodeMaterialComment.h
@@ -18,7 +18,7 @@ public:
 	void Construct(const FArguments& InArgs, class UMaterialGraphNode_Comment* InNode);
 
 	// SNodePanel::SNode interface
-	virtual void MoveTo(const FVector2D& NewPosition, FNodeSet& NodeFilter, bool bMarkDirty = true) override;
+	virtual void MoveTo(const FVector2f& NewPosition, FNodeSet& NodeFilter, bool bMarkDirty = true) override;
 	// End of SNodePanel::SNode interface
 
 private:
diff --git a/Engine/Source/Editor/GraphEditor/Private/MaterialNodes/SGraphNodeMaterialComposite.cpp b/Engine/Source/Editor/GraphEditor/Private/MaterialNodes/SGraphNodeMaterialComposite.cpp
--- a/Engine/Source/Editor/GraphEditor/Private/MaterialNodes/SGraphNodeMaterialComposite.cpp
+++ b/Engine/Source/Editor/GraphEditor/Private/MaterialNodes/SGraphNodeMaterialComposite.cpp
@@ -55,7 +55,7 @@ void SGraphNodeMaterialComposite::Construct(const FArguments& InArgs, class UMat
 	this->UpdateGraphNode();
 }
 
-void SGraphNodeMaterialComposite::MoveTo(const FVector2D& NewPosition, FNodeSet& NodeFilter, bool bMarkDirty)
+void SGraphNodeMaterialComposite::MoveTo(const FVector2f& NewPosition, FNodeSet& NodeFilter, bool bMarkDirty)
 {
 	SGraphNode::MoveTo(NewPosition, NodeFilter, bMarkDirty);
 
@@ -181,8 +181,8 @@ void SGraphNodeMaterialComposite::UpdateGraphNode()
 	.IsGraphNodeHovered( this, &SGraphNode::IsHovered );
 
 	GetOrAddSlot( ENodeZone::TopCenter )
-	.SlotOffset( TAttribute<FVector2D>( CommentBubble.Get(), &SCommentBubble::GetOffset ))
-	.SlotSize( TAttribute<FVector2D>( CommentBubble.Get(), &SCommentBubble::GetSize ))
+	.SlotOffset2f( TAttribute<FVector2f>( CommentBubble.Get(), &SCommentBubble::GetOffset2f ))
+	.SlotSize2f( TAttribute<FVector2f>( CommentBubble.Get(), &SCommentBubble::GetSize2f ))
 	.AllowScaling( TAttribute<bool>( CommentBubble.Get(), &SCommentBubble::IsScalingAllowed ))
 	.VAlign( VAlign_Top )
 	[
diff --git a/Engine/Source/Editor/GraphEditor/Private/MaterialNodes/SGraphNodeMaterialResult.cpp b/Engine/Source/Editor/GraphEditor/Private/MaterialNodes/SGraphNodeMaterialResult.cpp
--- a/Engine/Source/Editor/GraphEditor/Private/MaterialNodes/SGraphNodeMaterialResult.cpp
+++ b/Engine/Source/Editor/GraphEditor/Private/MaterialNodes/SGraphNodeMaterialResult.cpp
@@ -85,7 +85,7 @@ void SGraphNodeMaterialResult::CreatePinWidgets()
 	}
 }
 
-void SGraphNodeMaterialResult::MoveTo(const FVector2D& NewPosition, FNodeSet& NodeFilter, bool bMarkDirty)
+void SGraphNodeMaterialResult::MoveTo(const FVector2f& NewPosition, FNodeSet& NodeFilter, bool bMarkDirty)
 {
 	SGraphNode::MoveTo(NewPosition, NodeFilter, bMarkDirty);
 
diff --git a/Engine/Source/Editor/GraphEditor/Private/SCommentBubble.cpp b/Engine/Source/Editor/GraphEditor/Private/SCommentBubble.cpp
--- a/Engine/Source/Editor/GraphEditor/Private/SCommentBubble.cpp
+++ b/Engine/Source/Editor/GraphEditor/Private/SCommentBubble.cpp
@@ -92,8 +92,8 @@ void SCommentBubble::Construct( const FArguments& InArgs )
 
 FCursorReply SCommentBubble::OnCursorQuery( const FGeometry& MyGeometry, const FPointerEvent& CursorEvent ) const
 {
-	const FVector2D Size( GetDesiredSize().X, GetDesiredSize().Y - SCommentBubbleDefs::BubbleArrowHeight );
-	const FSlateRect TestRect( FVector2D(MyGeometry.AbsolutePosition), FVector2D(MyGeometry.AbsolutePosition) + Size );
+	const FVector2f Size( GetDesiredSize().X, GetDesiredSize().Y - SCommentBubbleDefs::BubbleArrowHeight );
+	const FSlateRect TestRect( MyGeometry.AbsolutePosition, MyGeometry.AbsolutePosition + Size );
 
 	if( TestRect.ContainsPoint( CursorEvent.GetScreenSpacePosition() ))
 	{
@@ -352,7 +352,12 @@ bool SCommentBubble::TextBlockHasKeyboardFocus() const
 
 FVector2D SCommentBubble::GetOffset() const
 {
-	return FVector2D( 0.f, -GetDesiredSize().Y );
+	return FVector2D(GetOffset2f());
+}
+
+FVector2f SCommentBubble::GetOffset2f() const
+{
+	return FVector2f( 0.f, -GetDesiredSize().Y );
 }
 
 float SCommentBubble::GetArrowCenterOffset() const
@@ -368,6 +373,11 @@ float SCommentBubble::GetArrowCenterOffset() const
 }
 
 FVector2D SCommentBubble::GetSize() const
+{
+	return FVector2D(GetSize2f());
+}
+
+FVector2f SCommentBubble::GetSize2f() const
 {
 	return GetDesiredSize();
 }
diff --git a/Engine/Source/Editor/GraphEditor/Private/SGraphEditorImpl.cpp b/Engine/Source/Editor/GraphEditor/Private/SGraphEditorImpl.cpp
--- a/Engine/Source/Editor/GraphEditor/Private/SGraphEditorImpl.cpp
+++ b/Engine/Source/Editor/GraphEditor/Private/SGraphEditorImpl.cpp
@@ -593,7 +593,7 @@ bool SGraphEditorImpl::GetBoundsForSelectedNodes( class FSlateRect& Rect, float
 
 bool SGraphEditorImpl::GetBoundsForNode( const UEdGraphNode* InNode, class FSlateRect& Rect, float Padding) const
 {
-	FVector2D TopLeft, BottomRight;
+	FVector2f TopLeft, BottomRight;
 	if (GraphPanel->GetBoundsForNode(InNode, TopLeft, BottomRight, Padding))
 	{
 		Rect.Left = TopLeft.X;
@@ -1542,7 +1542,7 @@ EActiveTimerReturnType SGraphEditorImpl::HandleFocusEditorDeferred(double InCurr
 				return EActiveTimerReturnType::Continue;
 			}
 
-			FVector2D TopLeft, BottomRight;
+			FVector2f TopLeft, BottomRight;
 
 			// If the locked graph was instructed to pan to a destination, let it ignore the lock to reach that destination.
 			// this way we can support diffs of moved nodes.
diff --git a/Engine/Source/Editor/GraphEditor/Private/SGraphNode.cpp b/Engine/Source/Editor/GraphEditor/Private/SGraphNode.cpp
--- a/Engine/Source/Editor/GraphEditor/Private/SGraphNode.cpp
+++ b/Engine/Source/Editor/GraphEditor/Private/SGraphNode.cpp
@@ -356,18 +356,18 @@ FReply SGraphNode::OnDragOver( const FGeometry& MyGeometry, const FDragDropEvent
 }
 
 /** Given a coordinate in SGraphPanel space (i.e. panel widget space), return the same coordinate in graph space while taking zoom and panning into account */
-FVector2D SGraphNode::NodeCoordToGraphCoord( const FVector2D& NodeSpaceCoordinate ) const
+UE::Slate::FDeprecateVector2DResult SGraphNode::NodeCoordToGraphCoord( const UE::Slate::FDeprecateVector2DParameter& NodeSpaceCoordinate ) const
 {
 	TSharedPtr<SGraphPanel> OwnerCanvas = OwnerGraphPanelPtr.Pin();
 	if (OwnerCanvas.IsValid())
 	{
 		//@TODO: NodeSpaceCoordinate != PanelCoordinate
-		FVector2D PanelSpaceCoordinate = NodeSpaceCoordinate;
+		FVector2f PanelSpaceCoordinate = NodeSpaceCoordinate;
 		return OwnerCanvas->PanelCoordToGraphCoord( PanelSpaceCoordinate );
 	}
 	else
 	{
-		return FVector2D::ZeroVector;
+		return FVector2f::ZeroVector;
 	}
 }
 
@@ -571,7 +571,7 @@ void SGraphNode::SetOwner( const TSharedRef<SGraphPanel>& OwnerPanel )
 }
 
 /** @param NewPosition  The Node should be relocated to this position in the graph panel */
-void SGraphNode::MoveTo( const FVector2D& NewPosition, FNodeSet& NodeFilter, bool bMarkDirty )
+void SGraphNode::MoveTo( const FVector2f& NewPosition, FNodeSet& NodeFilter, bool bMarkDirty )
 {
 	if ( !NodeFilter.Find( SharedThis( this )))
 	{
@@ -586,9 +586,9 @@ void SGraphNode::MoveTo( const FVector2D& NewPosition, FNodeSet& NodeFilter, boo
 }
 
 /** @return the Node's position within the graph */
-FVector2D SGraphNode::GetPosition() const
+FVector2f SGraphNode::GetPosition2f() const
 {
-	return FVector2D( GraphNode->NodePosX, GraphNode->NodePosY );
+	return FVector2f( GraphNode->NodePosX, GraphNode->NodePosY );
 }
 
 FString SGraphNode::GetEditableNodeTitle() const
@@ -1057,8 +1057,8 @@ void SGraphNode::UpdateGraphNode()
 			.IsGraphNodeHovered(this, &SGraphNode::IsHovered);
 
 		GetOrAddSlot(ENodeZone::TopCenter)
-			.SlotOffset(TAttribute<FVector2D>(CommentBubble.Get(), &SCommentBubble::GetOffset))
-			.SlotSize(TAttribute<FVector2D>(CommentBubble.Get(), &SCommentBubble::GetSize))
+			.SlotOffset2f(TAttribute<FVector2f>(CommentBubble.Get(), &SCommentBubble::GetOffset2f))
+			.SlotSize2f(TAttribute<FVector2f>(CommentBubble.Get(), &SCommentBubble::GetSize2f))
 			.AllowScaling(TAttribute<bool>(CommentBubble.Get(), &SCommentBubble::IsScalingAllowed))
 			.VAlign(VAlign_Top)
 			[
@@ -1382,10 +1382,10 @@ void SGraphNode::PlaySpawnEffect()
 	SpawnAnim.Play( this->AsShared() );
 }
 
-FVector2D SGraphNode::GetContentScale() const
+UE::Slate::FDeprecateVector2DResult SGraphNode::GetContentScale() const
 {
 	const float CurZoomValue = ZoomCurve.GetLerp();
-	return FVector2D( CurZoomValue, CurZoomValue );
+	return FVector2f( CurZoomValue, CurZoomValue );
 }
 
 FLinearColor SGraphNode::GetColorAndOpacity() const
@@ -1578,8 +1578,8 @@ void SGraphNode::ApplyRename()
 
 FSlateRect SGraphNode::GetTitleRect() const
 {
-	const FVector2D NodePosition = GetPosition();
-	const FVector2D NodeSize = GraphNode ? InlineEditableText->GetDesiredSize() : GetDesiredSize();
+	const FVector2f NodePosition = GetPosition2f();
+	const FVector2f NodeSize = GraphNode ? InlineEditableText->GetDesiredSize() : GetDesiredSize();
 
 	return FSlateRect( NodePosition.X, NodePosition.Y + NodeSize.Y, NodePosition.X + NodeSize.X, NodePosition.Y );
 }
diff --git a/Engine/Source/Editor/GraphEditor/Private/SGraphNodeComment.cpp b/Engine/Source/Editor/GraphEditor/Private/SGraphNodeComment.cpp
--- a/Engine/Source/Editor/GraphEditor/Private/SGraphNodeComment.cpp
+++ b/Engine/Source/Editor/GraphEditor/Private/SGraphNodeComment.cpp
@@ -240,8 +240,8 @@ void SGraphNodeComment::UpdateGraphNode()
 	.IsGraphNodeHovered(this, &SGraphNode::IsHovered);
 
 	GetOrAddSlot(ENodeZone::TopCenter)
-	.SlotOffset(TAttribute<FVector2D>(CommentBubble.Get(), &SCommentBubble::GetOffset))
-	.SlotSize(TAttribute<FVector2D>(CommentBubble.Get(), &SCommentBubble::GetSize))
+	.SlotOffset2f(TAttribute<FVector2f>(CommentBubble.Get(), &SCommentBubble::GetOffset2f))
+	.SlotSize2f(TAttribute<FVector2f>(CommentBubble.Get(), &SCommentBubble::GetSize2f))
 	.AllowScaling(TAttribute<bool>(CommentBubble.Get(), &SCommentBubble::IsScalingAllowed))
 	.VAlign(VAlign_Top)
 	[
@@ -251,7 +251,7 @@ void SGraphNodeComment::UpdateGraphNode()
 
 FVector2D SGraphNodeComment::ComputeDesiredSize( float ) const
 {
-	return UserSize;
+	return FVector2D(UserSize);
 }
 
 FString SGraphNodeComment::GetNodeComment() const
@@ -293,7 +293,7 @@ FReply SGraphNodeComment::OnMouseButtonUp( const FGeometry& MyGeometry, const FP
 		UserSize.X = FMath::RoundToFloat(UserSize.X);
 		UserSize.Y = FMath::RoundToFloat(UserSize.Y);
 
-		GetNodeObj()->ResizeNode(UserSize);
+		GetNodeObj()->ResizeNode(FVector2D(UserSize));
 
 		// End resize transaction
 		ResizeTransactionPtr.Reset();
@@ -314,7 +314,7 @@ int32 SGraphNodeComment::GetSortDepth() const
 
 void SGraphNodeComment::HandleSelection(bool bSelected, bool bUpdateNodesUnderComment) const
 {
-	const FVector2D NodeSize = GetDesiredSize();
+	const FVector2f NodeSize = GetDesiredSize();
 	// we only want to do this after the comment has a valid desired size
 	if( !NodeSize.IsZero() )
 	{
@@ -349,12 +349,12 @@ void SGraphNodeComment::HandleSelection(bool bSelected, bool bUpdateNodesUnderCo
 
 bool SGraphNodeComment::IsNodeUnderComment(UEdGraphNode_Comment* InCommentNode, const TSharedRef<SGraphNode> InNodeWidget) const
 {
-	const FVector2D NodePosition = GetPosition();
-	const FVector2D NodeSize = GetDesiredSize();
+	const FVector2f NodePosition = GetPosition2f();
+	const FVector2f NodeSize = GetDesiredSize();
 	const FSlateRect CommentRect(NodePosition.X, NodePosition.Y, NodePosition.X + NodeSize.X, NodePosition.Y + NodeSize.Y);
 
-	const FVector2D InNodePosition = InNodeWidget->GetPosition();
-	const FVector2D InNodeSize = InNodeWidget->GetDesiredSize();
+	const FVector2f InNodePosition = InNodeWidget->GetPosition2f();
+	const FVector2f InNodeSize = InNodeWidget->GetDesiredSize();
 
 	const FSlateRect NodeGeometryGraphSpace(InNodePosition.X, InNodePosition.Y, InNodePosition.X + InNodeSize.X, InNodePosition.Y + InNodeSize.Y);
 	return FSlateRect::IsRectangleContained(CommentRect, NodeGeometryGraphSpace);
@@ -366,7 +366,7 @@ const FSlateBrush* SGraphNodeComment::GetShadowBrush(bool bSelected) const
 	return SGraphNode::GetShadowBrush(bSelected);
 }
 
-void SGraphNodeComment::GetOverlayBrushes(bool bSelected, const FVector2D WidgetSize, TArray<FOverlayBrushInfo>& Brushes) const
+void SGraphNodeComment::GetOverlayBrushes(bool bSelected, const FVector2f& WidgetSize, TArray<FOverlayBrushInfo>& Brushes) const
 {
 	const float Fudge = 3.0f;
 
@@ -378,12 +378,12 @@ void SGraphNodeComment::GetOverlayBrushes(bool bSelected, const FVector2D Widget
 	HandleBrush.OverlayOffset.Y = WidgetSize.Y - HandleBrush.Brush->ImageSize.Y - Fudge;
 
 	Brushes.Add(HandleBrush);
-	return SGraphNode::GetOverlayBrushes(bSelected, WidgetSize, Brushes);
+	return SGraphNode::GetOverlayBrushes(bSelected, UE::Slate::CastToVector2f(WidgetSize), Brushes);
 }
 
-void SGraphNodeComment::MoveTo( const FVector2D& NewPosition, FNodeSet& NodeFilter, bool bMarkDirty)
+void SGraphNodeComment::MoveTo( const FVector2f& NewPosition, FNodeSet& NodeFilter, bool bMarkDirty)
 {
-	FVector2D PositionDelta = NewPosition - GetPosition();
+	FVector2f PositionDelta = NewPosition - GetPosition2f();
 	SGraphNode::MoveTo(NewPosition, NodeFilter, bMarkDirty);
 	// Don't drag note content if either of the shift keys are down.
 	FModifierKeysState KeysState = FSlateApplication::Get().GetModifierKeys();
@@ -416,10 +416,10 @@ void SGraphNodeComment::MoveTo( const FVector2D& NewPosition, FNodeSet& NodeFilt
 void SGraphNodeComment::EndUserInteraction() const
 {
 	// Find any parent comments and their list of child nodes
-	const FVector2D NodeSize = GetDesiredSize();
+	const FVector2f NodeSize = GetDesiredSize();
 	if( !NodeSize.IsZero() )
 	{
-		const FVector2D NodePosition = GetPosition();
+		const FVector2f NodePosition = GetPosition2f();
 		const FSlateRect CommentRect( NodePosition.X, NodePosition.Y, NodePosition.X + NodeSize.X, NodePosition.Y + NodeSize.Y );
 
 		TSharedPtr<SGraphPanel> Panel = GetOwnerPanel();
@@ -437,8 +437,8 @@ void SGraphNodeComment::EndUserInteraction() const
 				continue;
 			}
 
-			const FVector2D SomeNodePosition = SomeNodeWidget->GetPosition();
-			const FVector2D SomeNodeSize = SomeNodeWidget->GetDesiredSize();
+			const FVector2f SomeNodePosition = SomeNodeWidget->GetPosition2f();
+			const FVector2f SomeNodeSize = SomeNodeWidget->GetDesiredSize();
 
 			const FSlateRect NodeGeometryGraphSpace(SomeNodePosition.X, SomeNodePosition.Y, SomeNodePosition.X + SomeNodeSize.X, SomeNodePosition.Y + SomeNodeSize.Y);
 			if (FSlateRect::DoRectanglesIntersect(CommentRect, NodeGeometryGraphSpace))
@@ -461,9 +461,9 @@ FSlateRect SGraphNodeComment::GetHitTestingBorder() const
 	return SCommentNodeDefs::HitResultBorderSize;
 }
 
-FVector2D SGraphNodeComment::GetNodeMaximumSize() const
+FVector2f SGraphNodeComment::GetNodeMaximumSize2f() const
 {
-	return FVector2D( UserSize.X + 100, UserSize.Y + 100 );
+	return FVector2f( UserSize.X + 100.0f, UserSize.Y + 100.0f );
 }
 
 FSlateColor SGraphNodeComment::GetCommentBodyColor() const
@@ -511,22 +511,22 @@ FSlateColor SGraphNodeComment::GetCommentBubbleColor() const
 	}
 }
 
-bool SGraphNodeComment::CanBeSelected(const FVector2D& MousePositionInNode) const
+bool SGraphNodeComment::CanBeSelected(const FVector2f& MousePositionInNode) const
 {
 	const EResizableWindowZone InMouseZone = FindMouseZone(MousePositionInNode);
 	return CRWZ_TitleBar == InMouseZone;
 }
 
-FVector2D SGraphNodeComment::GetDesiredSizeForMarquee() const
+FVector2f SGraphNodeComment::GetDesiredSizeForMarquee2f() const
 {
 	const float TitleBarHeight = TitleBar.IsValid() ? TitleBar->GetDesiredSize().Y : 0.0f;
-	return FVector2D(UserSize.X, TitleBarHeight);
+	return FVector2f(UserSize.X, TitleBarHeight);
 }
 
 FSlateRect SGraphNodeComment::GetTitleRect() const
 {
-	const FVector2D NodePosition = GetPosition();
-	FVector2D NodeSize  = TitleBar.IsValid() ? TitleBar->GetDesiredSize() : GetDesiredSize();
+	const FVector2f NodePosition = GetPosition2f();
+	FVector2f NodeSize  = TitleBar.IsValid() ? TitleBar->GetDesiredSize() : GetDesiredSize();
 	return FSlateRect( NodePosition.X, NodePosition.Y, NodePosition.X + NodeSize.X, NodePosition.Y + NodeSize.Y ) + SCommentNodeDefs::TitleBarOffset;
 }
 
diff --git a/Engine/Source/Editor/GraphEditor/Private/SGraphNodeDocumentation.cpp b/Engine/Source/Editor/GraphEditor/Private/SGraphNodeDocumentation.cpp
--- a/Engine/Source/Editor/GraphEditor/Private/SGraphNodeDocumentation.cpp
+++ b/Engine/Source/Editor/GraphEditor/Private/SGraphNodeDocumentation.cpp
@@ -49,16 +49,16 @@ namespace GraphNodeDocumentationDefs
 	static const FSlateRect HitTestBorderSize( 10, 10, 8, 14 );
 
 	/** Minimum size for node */
-	static const FVector2D MinNodeSize( 200.0f, 10.0f );
+	static const FVector2f MinNodeSize( 200.0f, 10.0f );
 
 	/** Maximum size for node */
-	static const FVector2D MaximumNodeSize( 4000.0f, 10.0f );
+	static const FVector2f MaximumNodeSize( 4000.0f, 10.0f );
 
 	/** Default documentation content size */
-	static const FVector2D DefaultContentSize( 600.0f, 400.0f );
+	static const FVector2f DefaultContentSize( 600.0f, 400.0f );
 
 	/** Placeholder documentation content size */
-	static const FVector2D PlaceholderContentSize( 380.0f, 45.0f );
+	static const FVector2f PlaceholderContentSize( 380.0f, 45.0f );
 
 	/** Default content border */
 	static const FMargin DefaultContentBorder( 4.0f, 2.0f, 4.0f, 10.0f );
@@ -327,14 +327,14 @@ FVector2D SGraphNodeDocumentation::ComputeDesiredSize( float ) const
 	return FVector2D( UserSize.X, UserSize.Y + GetTitleBarHeight() );
 }
 
-FVector2D SGraphNodeDocumentation::GetNodeMinimumSize() const
+FVector2f SGraphNodeDocumentation::GetNodeMinimumSize2f() const
 {
 	return GraphNodeDocumentationDefs::MinNodeSize;
 }
 
-FVector2D SGraphNodeDocumentation::GetNodeMaximumSize() const
+FVector2f SGraphNodeDocumentation::GetNodeMaximumSize2f() const
 {
-	return FVector2D( DocumentationSize.X, ContentWidget->GetDesiredSize().Y );
+	return FVector2f( DocumentationSize.X, ContentWidget->GetDesiredSize().Y );
 }
 
 FReply SGraphNodeDocumentation::OnMouseWheel( const FGeometry& MyGeometry, const FPointerEvent& MouseEvent )
@@ -362,7 +362,7 @@ void SGraphNodeDocumentation::Tick( const FGeometry& AllottedGeometry, const dou
 	if( !bUserIsDragging )
 	{
 		ChildWidgetVisibility = EVisibility::HitTestInvisible;
-		FVector2D LocalMouseCoordinates = AllottedGeometry.AbsoluteToLocal( FSlateApplication::Get().GetCursorPos() );
+		FVector2f LocalMouseCoordinates = AllottedGeometry.AbsoluteToLocal( FSlateApplication::Get().GetCursorPos() );
 
 		EResizableWindowZone CurrMouseZone = FindMouseZone( LocalMouseCoordinates );
 		
@@ -371,7 +371,7 @@ void SGraphNodeDocumentation::Tick( const FGeometry& AllottedGeometry, const dou
 			ChildWidgetVisibility = EVisibility::Visible;
 		}
 	}
-	// Check Cached Links to determine if we need to update the documention link/excerpt
+	// Check Cached Links to determine if we need to update the documentation link/excerpt
 	if( CachedDocumentationLink != GraphNode->GetDocumentationLink() ||
 		CachedDocumentationExcerpt != GraphNode->GetDocumentationExcerptName())
 	{
diff --git a/Engine/Source/Editor/GraphEditor/Private/SGraphNodeKnot.cpp b/Engine/Source/Editor/GraphEditor/Private/SGraphNodeKnot.cpp
--- a/Engine/Source/Editor/GraphEditor/Private/SGraphNodeKnot.cpp
+++ b/Engine/Source/Editor/GraphEditor/Private/SGraphNodeKnot.cpp
@@ -417,8 +417,8 @@ void SGraphNodeKnot::UpdateGraphNode()
 	.OnToggled(this, &SGraphNode::OnCommentBubbleToggled);
 
 	GetOrAddSlot(ENodeZone::TopCenter)
-	.SlotOffset(TAttribute<FVector2D>(this, &SGraphNodeKnot::GetCommentOffset))
-	.SlotSize( TAttribute<FVector2D>( CommentBubble.Get(), &SCommentBubble::GetSize))
+	.SlotOffset2f(TAttribute<FVector2f>(this, &SGraphNodeKnot::GetCommentOffset))
+	.SlotSize2f( TAttribute<FVector2f>( CommentBubble.Get(), &SCommentBubble::GetSize2f))
 	.AllowScaling( TAttribute<bool>(CommentBubble.Get(), &SCommentBubble::IsScalingAllowed))
 	.VAlign(VAlign_Top)
 	[
@@ -470,14 +470,14 @@ void SGraphNodeKnot::AddPin(const TSharedRef<SGraphPin>& PinToAdd)
 	}
 }
 
-FVector2D SGraphNodeKnot::GetCommentOffset() const
+FVector2f SGraphNodeKnot::GetCommentOffset() const
 {
 	const bool bBubbleVisible = GraphNode->bCommentBubbleVisible || bAlwaysShowCommentBubble;
 	const float ZoomAmount = GraphNode->bCommentBubblePinned && OwnerGraphPanelPtr.IsValid() ? OwnerGraphPanelPtr.Pin()->GetZoomAmount() : 1.f;
 	const float NodeWidthOffset = bBubbleVisible ?	SKnotNodeDefinitions::KnotCenterBubbleAdjust * ZoomAmount :
 													SKnotNodeDefinitions::KnotCenterButtonAdjust * ZoomAmount;
 
-	return FVector2D(NodeWidthOffset - CommentBubble->GetArrowCenterOffset(), -CommentBubble->GetDesiredSize().Y);
+	return FVector2f(NodeWidthOffset - CommentBubble->GetArrowCenterOffset(), -CommentBubble->GetDesiredSize().Y);
 }
 
 void SGraphNodeKnot::OnCommentBubbleToggled(bool bInCommentBubbleVisible)
diff --git a/Engine/Source/Editor/GraphEditor/Private/SGraphNodeResizable.cpp b/Engine/Source/Editor/GraphEditor/Private/SGraphNodeResizable.cpp
--- a/Engine/Source/Editor/GraphEditor/Private/SGraphNodeResizable.cpp
+++ b/Engine/Source/Editor/GraphEditor/Private/SGraphNodeResizable.cpp
@@ -26,10 +26,10 @@ namespace GraphNodeResizableDefs
 	static const float DefaultTitleBarHeight = 12.f;
 
 	/** Minimum size for node */
-	static const FVector2D MinNodeSize( 30.0f, 30.0f );
+	static const FDeprecateSlateVector2D MinNodeSize( 30.0f, 30.0f );
 
 	/** Maximum size for node */
-	static const FVector2D MaxNodeSize( 400.0f, 400.0f );
+	static const FDeprecateSlateVector2D MaxNodeSize( 400.0f, 400.0f );
 }
 
 bool SGraphNodeResizable::InSelectionArea(EResizableWindowZone InMouseZone) const
@@ -44,7 +44,7 @@ void SGraphNodeResizable::OnMouseEnter( const FGeometry& MyGeometry, const FPoin
 	// Determine the zone the mouse is in
 	if( !bUserIsDragging )
 	{
-		FVector2D LocalMouseCoordinates = MyGeometry.AbsoluteToLocal( MouseEvent.GetScreenSpacePosition() );
+		FVector2f LocalMouseCoordinates = MyGeometry.AbsoluteToLocal( MouseEvent.GetScreenSpacePosition() );
 		MouseZone = FindMouseZone(LocalMouseCoordinates);
 		SNodePanel::SNode::OnMouseEnter( MyGeometry, MouseEvent );
 	}
@@ -133,10 +133,10 @@ FReply SGraphNodeResizable::OnMouseMove(const FGeometry& MyGeometry, const FPoin
 {
 	if (bUserIsDragging)
 	{
-		FVector2D GraphSpaceCoordinates = NodeCoordToGraphCoord( MouseEvent.GetScreenSpacePosition() );
-		FVector2D OldGraphSpaceCoordinates = NodeCoordToGraphCoord( MouseEvent.GetLastScreenSpacePosition() );
+		FVector2f GraphSpaceCoordinates = NodeCoordToGraphCoord( MouseEvent.GetScreenSpacePosition() );
+		FVector2f OldGraphSpaceCoordinates = NodeCoordToGraphCoord( MouseEvent.GetLastScreenSpacePosition() );
 		TSharedPtr<SWindow> OwnerWindow = FSlateApplication::Get().FindWidgetWindow(AsShared());
-		FVector2D Delta = (GraphSpaceCoordinates - OldGraphSpaceCoordinates) / (OwnerWindow.IsValid() ? OwnerWindow->GetDPIScaleFactor() : 1.0f);
+		FVector2f Delta = (GraphSpaceCoordinates - OldGraphSpaceCoordinates) / (OwnerWindow.IsValid() ? OwnerWindow->GetDPIScaleFactor() : 1.0f);
 		
 
 		//Clamp delta value based on resizing direction
@@ -150,7 +150,7 @@ FReply SGraphNodeResizable::OnMouseMove(const FGeometry& MyGeometry, const FPoin
 		}
 
 		//Resize node delta value
-		FVector2D DeltaNodeSize = Delta;
+		FVector2f DeltaNodeSize = Delta;
 
 		//Modify node size delta value based on resizing direction
 		if(	(MouseZone == CRWZ_LeftBorder) || (MouseZone == CRWZ_TopBorder) || (MouseZone == CRWZ_TopLeftBorder) )
@@ -170,20 +170,20 @@ FReply SGraphNodeResizable::OnMouseMove(const FGeometry& MyGeometry, const FPoin
 		DragSize.Y += DeltaNodeSize.Y;
 		// apply snap
 		const uint32 SnapSize = SNodePanel::GetSnapGridSize();
-		FVector2D SnappedSize;
+		FVector2f SnappedSize;
 		SnappedSize.X = SnapSize * FMath::RoundToFloat( DragSize.X/SnapSize );
 		SnappedSize.Y = SnapSize * FMath::RoundToFloat( DragSize.Y/SnapSize );
 
 		// Enforce min/max sizing
-		const FVector2D MinSize = GetNodeMinimumSize();
+		const FVector2f MinSize = GetNodeMinimumSize2f();
 		SnappedSize.X = FMath::Max( SnappedSize.X, MinSize.X );
 		SnappedSize.Y = FMath::Max( SnappedSize.Y, MinSize.Y );
 
-		const FVector2D MaxSize = GetNodeMaximumSize();
+		const FVector2f MaxSize = GetNodeMaximumSize2f();
 		SnappedSize.X = FMath::Min( SnappedSize.X, MaxSize.X );
 		SnappedSize.Y = FMath::Min( SnappedSize.Y, MaxSize.Y );
 
-		FVector2D DeltaNodePos(0,0);
+		FVector2f DeltaNodePos(0.0f,0.0f);
 
 		if( UserSize != SnappedSize )
 		{
@@ -205,7 +205,7 @@ FReply SGraphNodeResizable::OnMouseMove(const FGeometry& MyGeometry, const FPoin
 			}
 			UserSize = SnappedSize;
 			GraphNode->ResizeNode( UserSize );
-			DeltaNodePos = GetCorrectedNodePosition() - GetPosition();
+			DeltaNodePos = GetCorrectedNodePosition() - GetPosition2f();
 		}
 
 		if (!ResizeTransactionPtr.IsValid() && UserSize != StoredUserSize)
@@ -216,11 +216,11 @@ FReply SGraphNodeResizable::OnMouseMove(const FGeometry& MyGeometry, const FPoin
 		}
 
 		SGraphNode::FNodeSet NodeFilter;
-		SGraphNode::MoveTo( GetPosition() + DeltaNodePos, NodeFilter );
+		SGraphNode::MoveTo( GetPosition2f() + DeltaNodePos, NodeFilter );
 	}
 	else
 	{
-		const FVector2D LocalMouseCoordinates = MyGeometry.AbsoluteToLocal( MouseEvent.GetScreenSpacePosition() );
+		const FVector2f LocalMouseCoordinates = MyGeometry.AbsoluteToLocal( MouseEvent.GetScreenSpacePosition() );
 		MouseZone = FindMouseZone(LocalMouseCoordinates);
 	}
 	return SGraphNode::OnMouseMove( MyGeometry, MouseEvent );
@@ -228,7 +228,7 @@ FReply SGraphNodeResizable::OnMouseMove(const FGeometry& MyGeometry, const FPoin
 
 void SGraphNodeResizable::InitNodeAnchorPoint()
 {
-	NodeAnchorPoint = GetPosition();
+	NodeAnchorPoint = GetPosition2f();
 	
 	if( (MouseZone == CRWZ_LeftBorder) || (MouseZone == CRWZ_TopBorder) || (MouseZone == CRWZ_TopLeftBorder) )
 	{
@@ -245,9 +245,9 @@ void SGraphNodeResizable::InitNodeAnchorPoint()
 }
 
 
-FVector2D SGraphNodeResizable::GetCorrectedNodePosition() const
+UE::Slate::FDeprecateVector2DResult SGraphNodeResizable::GetCorrectedNodePosition() const
 {
-	FVector2D CorrectedPos = NodeAnchorPoint;
+	FVector2f CorrectedPos = NodeAnchorPoint;
 	
 	if( (MouseZone == CRWZ_LeftBorder) || (MouseZone == CRWZ_TopBorder) || (MouseZone == CRWZ_TopLeftBorder) )
 	{
@@ -264,12 +264,16 @@ FVector2D SGraphNodeResizable::GetCorrectedNodePosition() const
 
 	return CorrectedPos;
 }
-
 SGraphNodeResizable::EResizableWindowZone SGraphNodeResizable::FindMouseZone(const FVector2D& LocalMouseCoordinates) const
+{
+	return FindMouseZone(UE::Slate::CastToVector2f(LocalMouseCoordinates));
+}
+
+SGraphNodeResizable::EResizableWindowZone SGraphNodeResizable::FindMouseZone(const FVector2f& LocalMouseCoordinates) const
 {
 	EResizableWindowZone OutMouseZone = CRWZ_NotInWindow;
 	FSlateRect HitResultBorderSize = GetHitTestingBorder();
-	const FVector2D NodeSize = GetDesiredSize();
+	const FVector2f NodeSize = GetDesiredSize();
 	const float TitleBarHeight = GetTitleBarHeight();
 
 	const float Zoom = GetOwnerPanel()->GetZoomAmount();
@@ -364,11 +368,21 @@ float SGraphNodeResizable::GetTitleBarHeight() const
 }
 
 FVector2D SGraphNodeResizable::GetNodeMinimumSize() const
+{
+	return FVector2D(GetNodeMinimumSize2f());
+}
+
+FVector2f SGraphNodeResizable::GetNodeMinimumSize2f() const
 {
 	return GraphNodeResizableDefs::MinNodeSize;
 }
 
 FVector2D SGraphNodeResizable::GetNodeMaximumSize() const
+{
+	return FVector2D(GetNodeMaximumSize2f());
+}
+
+FVector2f SGraphNodeResizable::GetNodeMaximumSize2f() const
 {
 	return GraphNodeResizableDefs::MaxNodeSize;
 }
diff --git a/Engine/Source/Editor/GraphEditor/Private/SGraphPanel.cpp b/Engine/Source/Editor/GraphEditor/Private/SGraphPanel.cpp
--- a/Engine/Source/Editor/GraphEditor/Private/SGraphPanel.cpp
+++ b/Engine/Source/Editor/GraphEditor/Private/SGraphPanel.cpp
@@ -113,7 +113,7 @@ void SGraphPanel::Construct( const SGraphPanel::FArguments& InArgs )
 	this->bPreservePinPreviewConnection = false;
 	this->PinVisibility = SGraphEditor::Pin_Show;
 
-	CachedAllottedGeometryScaledSize = FVector2D(160, 120);
+	CachedAllottedGeometryScaledSize = FVector2f(160, 120);
 	if (InArgs._InitialZoomToFit)
 	{
 		ZoomToFit(/*bOnlySelection=*/ false);
@@ -131,8 +131,8 @@ void SGraphPanel::Construct( const SGraphPanel::FArguments& InArgs )
 	
 	ShowGraphStateOverlay = InArgs._ShowGraphStateOverlay;
 
-	SavedMousePosForOnPaintEventLocalSpace = FVector2D::ZeroVector;
-	PreviousFrameSavedMousePosForSplineOverlap = FVector2D::ZeroVector;
+	SavedMousePosForOnPaintEventLocalSpace = FVector2f::ZeroVector;
+	PreviousFrameSavedMousePosForSplineOverlap = FVector2f::ZeroVector;
 
 	TimeLeftToInvalidatePerTick = 0.0f;
 	
@@ -193,7 +193,7 @@ int32 SGraphPanel::OnPaint( const FPaintArgs& Args, const FGeometry& AllottedGeo
 
 	const FPaintArgs NewArgs = Args.WithNewParent(this);
 
-	const FVector2D NodeShadowSize = GetDefault<UGraphEditorSettings>()->GetShadowDeltaSize();
+	const FVector2f NodeShadowSize = UE::Slate::CastToVector2f(GetDefault<UGraphEditorSettings>()->GetShadowDeltaSize());
 	const UEdGraphSchema* Schema = GraphObj->GetSchema();
 
 	
@@ -408,7 +408,7 @@ int32 SGraphPanel::OnPaint( const FPaintArgs& Args, const FGeometry& AllottedGeo
 				// Draw the node's overlay, if it has one.
 				{
 					// Get its size
-					const FVector2D WidgetSize = CurWidget.Geometry.Size;
+					const FVector2f WidgetSize = CurWidget.Geometry.Size;
 
 					{
 						TArray<FOverlayBrushInfo> OverlayBrushes;
@@ -520,17 +520,17 @@ int32 SGraphPanel::OnPaint( const FPaintArgs& Args, const FGeometry& AllottedGeo
 				TArray< TSharedRef<SWidget> > NodePins;
 				ChildNode->GetPins(NodePins);
 
-				const FVector2D NodeLoc = ChildNode->GetPosition();
-				const FGeometry SynthesizedNodeGeometry(GraphCoordToPanelCoord(NodeLoc) * AllottedGeometry.Scale, FVector2D(AllottedGeometry.AbsolutePosition), FVector2D::ZeroVector, 1.f);
+				const FVector2f NodeLoc = ChildNode->GetPosition2f();
+				const FGeometry SynthesizedNodeGeometry(GraphCoordToPanelCoord(NodeLoc) * AllottedGeometry.Scale, AllottedGeometry.AbsolutePosition, FVector2f::ZeroVector, 1.f);
 
 				for (TArray< TSharedRef<SWidget> >::TConstIterator NodePinIterator(NodePins); NodePinIterator; ++NodePinIterator)
 				{
 					const SGraphPin& PinWidget = static_cast<const SGraphPin&>((*NodePinIterator).Get());
 					if (PinWidget.GetPinObj())
 					{
-						FVector2D PinLoc = NodeLoc + PinWidget.GetNodeOffset();
+						FVector2f PinLoc = NodeLoc + UE::Slate::CastToVector2f(PinWidget.GetNodeOffset());
 
-						const FGeometry SynthesizedPinGeometry(GraphCoordToPanelCoord(PinLoc) * AllottedGeometry.Scale, FVector2D(AllottedGeometry.AbsolutePosition), FVector2D::ZeroVector, 1.f);
+						const FGeometry SynthesizedPinGeometry(GraphCoordToPanelCoord(PinLoc) * AllottedGeometry.Scale, AllottedGeometry.AbsolutePosition, FVector2f::ZeroVector, 1.f);
 						PinGeometries.Add(*NodePinIterator, FArrangedWidget(*NodePinIterator, SynthesizedPinGeometry));
 					}
 				}
@@ -608,7 +608,7 @@ int32 SGraphPanel::OnPaint( const FPaintArgs& Args, const FGeometry& AllottedGeo
 				OverlapData.ComputeBestPin();
 
 				// Only allow spline overlaps when there is no node under the cursor (unless it is a comment box)
-				const FVector2D PaintAbsoluteSpaceMousePos = AllottedGeometry.LocalToAbsolute(SavedMousePosForOnPaintEventLocalSpace);
+				const FVector2f PaintAbsoluteSpaceMousePos = AllottedGeometry.LocalToAbsolute(SavedMousePosForOnPaintEventLocalSpace);
 				const int32 HoveredNodeIndex = SWidget::FindChildUnderPosition(ArrangedChildren, PaintAbsoluteSpaceMousePos);
 				if (HoveredNodeIndex != INDEX_NONE)
 				{
@@ -747,7 +747,7 @@ void SGraphPanel::OnArrangeChildren( const FGeometry& AllottedGeometry, FArrange
 		{
 			FOverlayWidgetInfo& OverlayInfo = OverlayWidgets[WidgetIndex];
 
-			MyArrangedChildren.AddWidget(AllottedGeometry.MakeChild( OverlayInfo.Widget.ToSharedRef(), FVector2D(CurWidget.Geometry.Position) + OverlayInfo.OverlayOffset, OverlayInfo.Widget->GetDesiredSize(), GetZoomAmount() ));
+			MyArrangedChildren.AddWidget(AllottedGeometry.MakeChild( OverlayInfo.Widget.ToSharedRef(), CurWidget.Geometry.Position + OverlayInfo.OverlayOffset, OverlayInfo.Widget->GetDesiredSize(), GetZoomAmount() ));
 		}
 	}
 
@@ -764,7 +764,7 @@ bool SGraphPanel::CustomPrepass(float LayoutScaleMultiplier)
 	return true; // still run prepass on everything
 }
 
-void SGraphPanel::UpdateSelectedNodesPositions(FVector2D PositionIncrement)
+void SGraphPanel::UpdateSelectedNodesPositions(const FVector2f& PositionIncrement)
 {
 	FScopedTransaction Transaction(NSLOCTEXT("GraphEditor", "NudgeNodeAction", "Nudge Node"));
 	for (auto& NodeIt : SelectionManager.SelectedNodes)
@@ -774,7 +774,7 @@ void SGraphPanel::UpdateSelectedNodesPositions(FVector2D PositionIncrement)
 		{
 			SNode& Widget = pWidget->Get();
 			SNode::FNodeSet NodeFilter;
-			Widget.MoveTo(Widget.GetPosition() + PositionIncrement, NodeFilter);
+			Widget.MoveTo(Widget.GetPosition2f() + PositionIncrement, NodeFilter);
 		}
 	}
 }
@@ -788,22 +788,22 @@ FReply SGraphPanel::OnKeyDown( const FGeometry& MyGeometry, const FKeyEvent& InK
 		{
 			if( InKeyEvent.GetKey() == EKeys::Up || InKeyEvent.GetKey() == EKeys::NumPadEight )
 			{
-				UpdateSelectedNodesPositions(FVector2D(0.0f,-1.0f * GetSnapGridSize()));
+				UpdateSelectedNodesPositions(FVector2f(0.0f,-1.0f * GetSnapGridSize()));
 				return FReply::Handled();
 			}
 			if( InKeyEvent.GetKey() == EKeys::Down || InKeyEvent.GetKey() == EKeys::NumPadTwo )
 			{
-				UpdateSelectedNodesPositions(FVector2D(0.0f,GetSnapGridSize()));
+				UpdateSelectedNodesPositions(FVector2f(0.0f,GetSnapGridSize()));
 				return FReply::Handled();
 			}
 			if( InKeyEvent.GetKey() == EKeys::Right || InKeyEvent.GetKey() == EKeys::NumPadSix )
 			{
-				UpdateSelectedNodesPositions(FVector2D(GetSnapGridSize(),0.0f));
+				UpdateSelectedNodesPositions(FVector2f(GetSnapGridSize(),0.0f));
 				return FReply::Handled();
 			}
 			if( InKeyEvent.GetKey() == EKeys::Left || InKeyEvent.GetKey() == EKeys::NumPadFour )
 			{
-				UpdateSelectedNodesPositions(FVector2D(-1.0f * GetSnapGridSize(),0.0f));
+				UpdateSelectedNodesPositions(FVector2f(-1.0f * GetSnapGridSize(),0.0f));
 				return FReply::Handled();
 			}
 		}
@@ -838,11 +838,11 @@ FReply SGraphPanel::OnKeyDown( const FGeometry& MyGeometry, const FKeyEvent& InK
 		{
 			// Note: We can't use SavedMousePosForOnPaintEventLocalSpace since it isn't updated while dragging,
 			// and so would just be the mouse position of the connection origin. So instead we'll just use the current cursor pos
-			FVector2D NewNodePosition = PanelCoordToGraphCoord(MyGeometry.AbsoluteToLocal(FSlateApplication::Get().GetCursorPos()));
+			FVector2f NewNodePosition = PanelCoordToGraphCoord(MyGeometry.AbsoluteToLocal(FSlateApplication::Get().GetCursorPos()));
 			FInputChord KeyChord = FInputChord(InKeyEvent.GetKey(), EModifierKey::FromBools(InKeyEvent.IsControlDown(), InKeyEvent.IsAltDown(), InKeyEvent.IsShiftDown(), InKeyEvent.IsCommandDown()));
 
 			int32 NodeCountBefore = GraphObj->Nodes.Num();
-			FReply SpawnNodeReply = OnSpawnNodeByShortcut.Execute(KeyChord, NewNodePosition);
+			FReply SpawnNodeReply = OnSpawnNodeByShortcut.Execute(KeyChord, FVector2D(NewNodePosition));
 			int32 NodeCountAfter = GraphObj->Nodes.Num();
 
 			// If we spawned a node then we won't call down into super and instead do some extra handling
@@ -1067,10 +1067,10 @@ FReply SGraphPanel::OnMouseButtonDoubleClick(const FGeometry& MyGeometry, const
 	if (PreviousFrameSplineOverlap.GetPins(*this, /*out*/ Pin1, /*out*/ Pin2))
 	{
 		// Give the schema a chance to do something interesting with a double click on a proper spline (both ends are attached to a pin, i.e., not a preview/drag one)
-		const FVector2D DoubleClickPositionInGraphSpace = PanelCoordToGraphCoord(MyGeometry.AbsoluteToLocal(MouseEvent.GetScreenSpacePosition()));
+		const FVector2f DoubleClickPositionInGraphSpace = PanelCoordToGraphCoord(MyGeometry.AbsoluteToLocal(MouseEvent.GetScreenSpacePosition()));
 
 		const UEdGraphSchema* Schema = GraphObj->GetSchema();
-		Schema->OnPinConnectionDoubleCicked(Pin1, Pin2, DoubleClickPositionInGraphSpace);
+		Schema->OnPinConnectionDoubleCicked(Pin1, Pin2, FVector2D(DoubleClickPositionInGraphSpace));
 	}
 	else if (!PreviousFrameSplineOverlap.GetCloseToSpline())
 	{
@@ -1125,7 +1125,7 @@ void SGraphPanel::ArrangeChildrenForContextMenuSummon(const FGeometry& AllottedG
 		const TSharedRef<SNode>& SomeChild = VisibleChildren[ChildIndex];
 		if (!SomeChild->RequiresSecondPassLayout())
 		{
-			ArrangedChildren.AddWidget(AllottedGeometry.MakeChild(SomeChild, SomeChild->GetPosition() - ViewOffset, SomeChild->GetDesiredSizeForMarquee(), GetZoomAmount()));
+			ArrangedChildren.AddWidget(AllottedGeometry.MakeChild(SomeChild, SomeChild->GetPosition2f() - ViewOffset, SomeChild->GetDesiredSizeForMarquee2f(), GetZoomAmount()));
 		}
 	}
 
@@ -1136,7 +1136,7 @@ void SGraphPanel::ArrangeChildrenForContextMenuSummon(const FGeometry& AllottedG
 		if (SomeChild->RequiresSecondPassLayout())
 		{
 			SomeChild->PerformSecondPassLayout(NodeToWidgetLookup);
-			ArrangedChildren.AddWidget(AllottedGeometry.MakeChild(SomeChild, SomeChild->GetPosition() - ViewOffset, SomeChild->GetDesiredSizeForMarquee(), GetZoomAmount()));
+			ArrangedChildren.AddWidget(AllottedGeometry.MakeChild(SomeChild, SomeChild->GetPosition2f() - ViewOffset, SomeChild->GetDesiredSizeForMarquee2f(), GetZoomAmount()));
 		}
 	}
 }
@@ -1148,7 +1148,7 @@ TSharedPtr<SWidget> SGraphPanel::OnSummonContextMenu(const FGeometry& MyGeometry
 	UEdGraphNode* EdNodeUnderMouse = NodeUnderMouse.IsValid() ? NodeUnderMouse->GetNodeObj() : nullptr;
 	TArray<UEdGraphPin*> NoSourcePins;
 
-	const FVector2D NodeAddPosition = PanelCoordToGraphCoord(MyGeometry.AbsoluteToLocal(MouseEvent.GetScreenSpacePosition()));
+	const FVector2f NodeAddPosition = PanelCoordToGraphCoord(MyGeometry.AbsoluteToLocal(MouseEvent.GetScreenSpacePosition()));
 	return SummonContextMenu(MouseEvent.GetScreenSpacePosition(), NodeAddPosition, EdNodeUnderMouse, PinUnderCursor, NoSourcePins);
 }
 
@@ -1213,7 +1213,7 @@ FReply SGraphPanel::OnMouseMove(const FGeometry& MyGeometry, const FPointerEvent
 	SavedMousePosForOnPaintEventLocalSpace = MyGeometry.AbsoluteToLocal(MouseEvent.GetScreenSpacePosition());
 
 	// Invalidate the spline results if we moved very far
-	const FVector2D MouseDelta = SavedMousePosForOnPaintEventLocalSpace - PreviousFrameSavedMousePosForSplineOverlap;
+	const FVector2f MouseDelta = SavedMousePosForOnPaintEventLocalSpace - PreviousFrameSavedMousePosForSplineOverlap;
 	const float MouseDeltaLengthSquared = MouseDelta.SizeSquared();
 	const bool bCursorInDeadZone = MouseDeltaLengthSquared <= FMath::Square(FSlateApplication::Get().GetDragTriggerDistance());
 
@@ -1331,7 +1331,7 @@ FReply SGraphPanel::OnDragOver( const FGeometry& MyGeometry, const FDragDropEven
 
 FReply SGraphPanel::OnDrop( const FGeometry& MyGeometry, const FDragDropEvent& DragDropEvent )
 {
-	const FVector2D NodeAddPosition = PanelCoordToGraphCoord( MyGeometry.AbsoluteToLocal( DragDropEvent.GetScreenSpacePosition() ) );
+	const FVector2f NodeAddPosition = PanelCoordToGraphCoord( MyGeometry.AbsoluteToLocal( DragDropEvent.GetScreenSpacePosition() ) );
 
 	FSlateApplication::Get().SetKeyboardFocus(AsShared(), EFocusCause::SetDirectly);
 
@@ -1347,7 +1347,7 @@ FReply SGraphPanel::OnDrop( const FGeometry& MyGeometry, const FDragDropEvent& D
 		TSharedPtr<FGraphEditorDragDropAction> DragConn = StaticCastSharedPtr<FGraphEditorDragDropAction>(Operation);
 		if (DragConn.IsValid() && DragConn->IsSupportedBySchema(GraphObj->GetSchema()))
 		{
-			return DragConn->DroppedOnPanel(SharedThis(this), DragDropEvent.GetScreenSpacePosition(), NodeAddPosition, *GraphObj);
+			return DragConn->DroppedOnPanel(SharedThis(this), FVector2D(DragDropEvent.GetScreenSpacePosition()), FVector2D(NodeAddPosition), *GraphObj);
 		}
 
 		return FReply::Unhandled();
@@ -1355,21 +1355,21 @@ FReply SGraphPanel::OnDrop( const FGeometry& MyGeometry, const FDragDropEvent& D
 	else if (Operation->IsOfType<FActorDragDropGraphEdOp>())
 	{
 		TSharedPtr<FActorDragDropGraphEdOp> ActorOp = StaticCastSharedPtr<FActorDragDropGraphEdOp>(Operation);
-		OnDropActor.ExecuteIfBound(ActorOp->Actors, GraphObj, NodeAddPosition);
+		OnDropActor.ExecuteIfBound(ActorOp->Actors, GraphObj, FVector2D(NodeAddPosition));
 		return FReply::Handled();
 	}
 
 	else if (Operation->IsOfType<FLevelDragDropOp>())
 	{
 		TSharedPtr<FLevelDragDropOp> LevelOp = StaticCastSharedPtr<FLevelDragDropOp>(Operation);
-		OnDropStreamingLevel.ExecuteIfBound(LevelOp->StreamingLevelsToDrop, GraphObj, NodeAddPosition);
+		OnDropStreamingLevel.ExecuteIfBound(LevelOp->StreamingLevelsToDrop, GraphObj, FVector2D(NodeAddPosition));
 		return FReply::Handled();
 	}
 
 	else if (Operation->IsOfType<FGraphNodeDragDropOp>())
 	{
 		TSharedPtr<FGraphNodeDragDropOp> NodeDropOp = StaticCastSharedPtr<FGraphNodeDragDropOp>(Operation);
-		NodeDropOp->OnPerformDropToGraph.ExecuteIfBound(NodeDropOp, GraphObj, NodeAddPosition, DragDropEvent.GetScreenSpacePosition());
+		NodeDropOp->OnPerformDropToGraph.ExecuteIfBound(NodeDropOp, GraphObj, FVector2D(NodeAddPosition), FVector2D(DragDropEvent.GetScreenSpacePosition()));
 		return FReply::Handled();
 	}
 	else
@@ -1382,7 +1382,7 @@ FReply SGraphPanel::OnDrop( const FGeometry& MyGeometry, const FDragDropEvent& D
 			{
 				if (PassesAssetReferenceFilter(DroppedAssetData))
 				{
-					GraphObj->GetSchema()->DroppedAssetsOnGraph( DroppedAssetData, NodeAddPosition, GraphObj );
+					GraphObj->GetSchema()->DroppedAssetsOnGraph( DroppedAssetData, FVector2D(NodeAddPosition), GraphObj );
 				}
 				return FReply::Handled();
 			}
@@ -1458,7 +1458,7 @@ UEdGraphPin* SGraphPanel::GetPinUnderMouse(const FGeometry& MyGeometry, const FP
 	return PinUnderCursor;
 }
 
-void SGraphPanel::AdjustNewlySpawnedNodePositions(TArrayView<UEdGraphNode* const> SpawnedNodes, TArrayView<UEdGraphPin*> DraggedFromPins, FVector2D AnchorPosition)
+void SGraphPanel::AdjustNewlySpawnedNodePositions(TArrayView<UEdGraphNode* const> SpawnedNodes, TArrayView<UEdGraphPin*> DraggedFromPins, FVector2f AnchorPosition)
 {
 	static auto FindFirstLinkedAutoWiredPin = [](TArrayView<UEdGraphNode* const> SpawnedNodes, TArrayView<UEdGraphPin*> DraggedFromPins) -> UEdGraphPin*
 	{
@@ -1482,7 +1482,7 @@ void SGraphPanel::AdjustNewlySpawnedNodePositions(TArrayView<UEdGraphNode* const
 	}
 }
 
-void SGraphPanel::MoveNodesToAnchorPinAtGraphPosition(TArrayView<UEdGraphNode* const> NodesToMove, FGraphPinHandle PinToAnchor, FVector2D DesiredPinGraphPosition)
+void SGraphPanel::MoveNodesToAnchorPinAtGraphPosition(TArrayView<UEdGraphNode* const> NodesToMove, FGraphPinHandle PinToAnchor, FVector2f DesiredPinGraphPosition)
 {
 	struct FAnchorUtils
 	{
@@ -1498,13 +1498,13 @@ void SGraphPanel::MoveNodesToAnchorPinAtGraphPosition(TArrayView<UEdGraphNode* c
 			Node->NodePosY = RoundToGrid(Node->NodePosY, GridSnapSize);
 		}
 
-		static EActiveTimerReturnType AlignPinToPositionDelayed(double, float, TSharedRef<SGraphPanel> Panel, FGraphPinHandle DragFromPinHandle, FVector2D DesiredPinImageCenterGraph, TArray<UEdGraphNode*> SpawnedNodes)
+		static EActiveTimerReturnType AlignPinToPositionDelayed(double, float, TSharedRef<SGraphPanel> Panel, FGraphPinHandle DragFromPinHandle, FVector2f DesiredPinImageCenterGraph, TArray<UEdGraphNode*> SpawnedNodes)
 		{
 			AlignPinToPosition(Panel, DragFromPinHandle, DesiredPinImageCenterGraph, SpawnedNodes);
 			return EActiveTimerReturnType::Stop;
 		}
 
-		static void AlignPinToPosition(TSharedRef<SGraphPanel> Panel, FGraphPinHandle DragFromPinHandle, FVector2D DesiredPinImageCenterGraph, TArrayView<UEdGraphNode* const> SpawnedNodes)
+		static void AlignPinToPosition(TSharedRef<SGraphPanel> Panel, FGraphPinHandle DragFromPinHandle, FVector2f DesiredPinImageCenterGraph, TArrayView<UEdGraphNode* const> SpawnedNodes)
 		{
 			TSharedPtr<SGraphPin> DragFromPinWidget = DragFromPinHandle.FindInGraphPanel(*Panel);
 			if (!DragFromPinWidget.IsValid())
@@ -1578,18 +1578,18 @@ void SGraphPanel::MoveNodesToAnchorPinAtGraphPosition(TArrayView<UEdGraphNode* c
 			}
 
 			// Default to a similar calculation to Paint(), though we don't have access to the drawing policy so can't add the arrow offset
-			FVector2D PinImageCenterAbsolute = DragFromPinWidget->GetDirection() == EGPD_Input ? FGeometryHelper::VerticalMiddleLeftOf(ArrangedDragFromPinWidget->Geometry) : FGeometryHelper::VerticalMiddleRightOf(ArrangedDragFromPinWidget->Geometry);
+			FVector2f PinImageCenterAbsolute = DragFromPinWidget->GetDirection() == EGPD_Input ? FGeometryHelper::VerticalMiddleLeftOf(ArrangedDragFromPinWidget->Geometry) : FGeometryHelper::VerticalMiddleRightOf(ArrangedDragFromPinWidget->Geometry);
 
 			// Though if we found the actual pin image widget, then we can use its exact center instead, which will probably be more accurate
 			// Either should be close enough given we'll be snapping to the grid afterwards anyway
 			if (ArrangedPinImageWidget)
 			{
-				PinImageCenterAbsolute = ArrangedPinImageWidget->Geometry.GetAbsolutePositionAtCoordinates(FVector2D(0.5f, 0.5f));
+				PinImageCenterAbsolute = ArrangedPinImageWidget->Geometry.GetAbsolutePositionAtCoordinates(FVector2f(0.5f, 0.5f));
 			}
 
-			FVector2D PinImageCenterPanel = Panel->GetTickSpaceGeometry().AbsoluteToLocal(PinImageCenterAbsolute);
-			FVector2D PinImageCenterGraph = Panel->PanelCoordToGraphCoord(PinImageCenterPanel);
-			FVector2D Delta = DesiredPinImageCenterGraph - PinImageCenterGraph;
+			FVector2f PinImageCenterPanel = Panel->GetTickSpaceGeometry().AbsoluteToLocal(PinImageCenterAbsolute);
+			FVector2f PinImageCenterGraph = Panel->PanelCoordToGraphCoord(PinImageCenterPanel);
+			FVector2f Delta = DesiredPinImageCenterGraph - PinImageCenterGraph;
 
 			// Offset all nodes that were spawned by this same delta so that their relative
 			// positioning is maintained (re-snapping them each individually though)
@@ -1718,7 +1718,7 @@ void SGraphPanel::RemoveAllNodes()
 	SNodePanel::RemoveAllNodes();
 }
 
-TSharedPtr<SWidget> SGraphPanel::SummonContextMenu(const FVector2D& WhereToSummon, const FVector2D& WhereToAddNode, UEdGraphNode* ForNode, UEdGraphPin* ForPin, const TArray<UEdGraphPin*>& DragFromPins)
+TSharedPtr<SWidget> SGraphPanel::SummonContextMenu(const UE::Slate::FDeprecateVector2DParameter& WhereToSummon, const UE::Slate::FDeprecateVector2DParameter& WhereToAddNode, UEdGraphNode* ForNode, UEdGraphPin* ForPin, const TArray<UEdGraphPin*>& DragFromPins)
 {
 	TRACE_CPUPROFILER_EVENT_SCOPE(SGraphPanel::SummonContextMenu);
 
@@ -1766,9 +1766,9 @@ TSharedPtr<SWidget> SGraphPanel::SummonContextMenu(const FVector2D& WhereToSummo
 
 void SGraphPanel::SummonCreateNodeMenuFromUICommand(uint32 NumNodesAdded)
 {
-	FVector2D WhereToSummonMenu = LastPointerEvent.GetScreenSpacePosition();
+	FVector2f WhereToSummonMenu = LastPointerEvent.GetScreenSpacePosition();
 	const float AdditionalOffset = 1.0f + (NumNodesAdded * GetDefault<UGraphEditorSettings>()->PaddingAutoCollateIncrement);
-	FVector2D WhereToAddNode = PastePosition + AdditionalOffset;
+	FVector2f WhereToAddNode = PastePosition + AdditionalOffset;
 
 	TSharedPtr<SGraphNode> NodeUnderMouse = GetGraphNodeUnderMouse(LastPointerGeometry, LastPointerEvent);
 
@@ -1808,6 +1808,15 @@ void SGraphPanel::AttachGraphEvents(TSharedPtr<SGraphNode> CreatedSubNode)
 }
 
 bool SGraphPanel::GetBoundsForNode(const UObject* InNode, FVector2D& MinCorner, FVector2D& MaxCorner, float Padding) const
+{
+	FVector2f TempMinCorner, TempMaxCorner;
+	bool bRetVal = GetBoundsForNode(InNode, TempMinCorner, TempMaxCorner, Padding);
+	MinCorner = FVector2D(TempMinCorner);
+	MaxCorner = FVector2D(TempMaxCorner);
+
+	return bRetVal;
+}
+bool SGraphPanel::GetBoundsForNode(const UObject* InNode, FVector2f& MinCorner, FVector2f& MaxCorner, float Padding) const
 {
 	return SNodePanel::GetBoundsForNode(InNode, MinCorner, MaxCorner, Padding);
 }
@@ -2362,8 +2371,8 @@ bool SGraphPanel::IsNodeTitleVisible(const class UEdGraphNode* Node, bool bReque
 		if(GraphNode.IsValid() && !HasMouseCapture())
 		{
 			FSlateRect TitleRect = GraphNode.Pin()->GetTitleRect();
-			const FVector2D TopLeft = FVector2D( TitleRect.Left, TitleRect.Top );
-			const FVector2D BottomRight = FVector2D( TitleRect.Right, TitleRect.Bottom );
+			const FVector2f TopLeft = FVector2f( TitleRect.Left, TitleRect.Top );
+			const FVector2f BottomRight = FVector2f( TitleRect.Right, TitleRect.Bottom );
 
 			if( IsRectVisible( TopLeft, BottomRight ))
 			{
@@ -2384,13 +2393,13 @@ bool SGraphPanel::IsNodeTitleVisible(const class UEdGraphNode* Node, bool bReque
 	return bTitleVisible;
 }
 
-bool SGraphPanel::IsRectVisible(const FVector2D &TopLeft, const FVector2D &BottomRight)
+bool SGraphPanel::IsRectVisible(const UE::Slate::FDeprecateVector2DParameter &TopLeft, const UE::Slate::FDeprecateVector2DParameter &BottomRight)
 {
-	return TopLeft.ComponentwiseAllGreaterOrEqual( PanelCoordToGraphCoord( FVector2D::ZeroVector )) && 
+	return TopLeft.ComponentwiseAllGreaterOrEqual( PanelCoordToGraphCoord( FVector2f::ZeroVector )) && 
 		BottomRight.ComponentwiseAllLessOrEqual( PanelCoordToGraphCoord( CachedAllottedGeometryScaledSize ) );
 }
 
-bool SGraphPanel::JumpToRect(const FVector2D &TopLeft, const FVector2D &BottomRight)
+bool SGraphPanel::JumpToRect(const UE::Slate::FDeprecateVector2DParameter &TopLeft, const UE::Slate::FDeprecateVector2DParameter &BottomRight)
 {
 	ZoomToTarget(TopLeft, BottomRight);
 
diff --git a/Engine/Source/Editor/GraphEditor/Private/SGraphPin.cpp b/Engine/Source/Editor/GraphEditor/Private/SGraphPin.cpp
--- a/Engine/Source/Editor/GraphEditor/Private/SGraphPin.cpp
+++ b/Engine/Source/Editor/GraphEditor/Private/SGraphPin.cpp
@@ -842,11 +842,11 @@ FReply SGraphPin::OnDrop( const FGeometry& MyGeometry, const FDragDropEvent& Dra
 	{
 		TSharedPtr<FGraphEditorDragDropAction> DragConnectionOp = StaticCastSharedPtr<FGraphEditorDragDropAction>(Operation);
 
-		FVector2D NodeAddPosition = FVector2D::ZeroVector;
+		FVector2f NodeAddPosition = FVector2f::ZeroVector;
 		TSharedPtr<SGraphNode> OwnerNode = OwnerNodePtr.Pin();
 		if (OwnerNode.IsValid())
 		{
-			NodeAddPosition	= OwnerNode->GetPosition() + FVector2D(MyGeometry.Position);
+			NodeAddPosition	= OwnerNode->GetPosition2f() + MyGeometry.Position;
 
 			//Don't have access to bounding information for node, using fixed offset that should work for most cases.
 			const float FixedOffset = 200.0f;
@@ -857,8 +857,8 @@ FReply SGraphPin::OnDrop( const FGeometry& MyGeometry, const FDragDropEvent& Dra
 			// if the pin widget is nested into another compound
 			if (MyGeometry.Position == FVector2f::ZeroVector)
 			{
-				FVector2D PinOffsetPosition = FVector2D(MyGeometry.AbsolutePosition) - FVector2D(NodeWidget->GetTickSpaceGeometry().AbsolutePosition);
-				NodeAddPosition = OwnerNode->GetPosition() + PinOffsetPosition;
+				FVector2f PinOffsetPosition = MyGeometry.AbsolutePosition - NodeWidget->GetTickSpaceGeometry().AbsolutePosition;
+				NodeAddPosition = OwnerNode->GetPosition2f() + PinOffsetPosition;
 			}
 
 			if(GetDirection() == EEdGraphPinDirection::EGPD_Input)
@@ -876,7 +876,7 @@ FReply SGraphPin::OnDrop( const FGeometry& MyGeometry, const FDragDropEvent& Dra
 			
 		}
 
-		return DragConnectionOp->DroppedOnPin(DragDropEvent.GetScreenSpacePosition(), NodeAddPosition);
+		return DragConnectionOp->DroppedOnPin(DragDropEvent.GetScreenSpacePosition(), FVector2D(NodeAddPosition));
 	}
 	// handle dropping an asset on the pin
 	else if (Operation->IsOfType<FAssetDragDropOp>() && NodeWidget.IsValid())
@@ -887,7 +887,7 @@ FReply SGraphPin::OnDrop( const FGeometry& MyGeometry, const FDragDropEvent& Dra
 			TSharedPtr<FAssetDragDropOp> AssetOp = StaticCastSharedPtr<FAssetDragDropOp>(Operation);
 			if (AssetOp->HasAssets())
 			{
-				Node->GetSchema()->DroppedAssetsOnPin(AssetOp->GetAssets(), NodeWidget->GetPosition() + FVector2D(MyGeometry.Position), GraphPinObj);
+				Node->GetSchema()->DroppedAssetsOnPin(AssetOp->GetAssets(), FVector2D(NodeWidget->GetPosition2f() + MyGeometry.Position), GraphPinObj);
 			}
 		}
 		return FReply::Handled();
@@ -1406,7 +1406,7 @@ void SGraphPin::GetInteractiveTooltipLocation(FVector2D& InOutDesiredLocation) c
 		if (GraphPanel.IsValid())
 		{
 			// Reset to the pin's location in graph space.
-			InOutDesiredLocation = OwnerNode->GetPosition() + CachedNodeOffset;
+			InOutDesiredLocation = FVector2D(OwnerNode->GetPosition2f()) + CachedNodeOffset;
 			
 			// Shift the desired location to the right edge of the pin's geometry.
 			InOutDesiredLocation.X += GetTickSpaceGeometry().Size.X;
diff --git a/Engine/Source/Editor/GraphEditor/Private/SNodePanel.cpp b/Engine/Source/Editor/GraphEditor/Private/SNodePanel.cpp
--- a/Engine/Source/Editor/GraphEditor/Private/SNodePanel.cpp
+++ b/Engine/Source/Editor/GraphEditor/Private/SNodePanel.cpp
@@ -236,13 +236,29 @@ void SNodePanel::SNode::FNodeSlot::Construct(const FChildren& SlotOwner, FSlotAr
 	{
 		SlotPadding = MoveTemp(InArgs._Padding);
 	}
+
+PRAGMA_DISABLE_DEPRECATION_WARNINGS
 	if (InArgs._SlotOffset.IsSet())
 	{
-		Offset = MoveTemp(InArgs._SlotOffset);
+		SetSlotOffset(InArgs._SlotOffset);
+	}
+PRAGMA_ENABLE_DEPRECATION_WARNINGS
+
+	if (InArgs._SlotOffset2f.IsSet())
+	{
+		Offset = MoveTemp(InArgs._SlotOffset2f);
 	}
+
+PRAGMA_DISABLE_DEPRECATION_WARNINGS
 	if (InArgs._SlotSize.IsSet())
 	{
-		Size = MoveTemp(InArgs._SlotSize);
+		SetSlotSize(InArgs._SlotSize);
+	}
+PRAGMA_ENABLE_DEPRECATION_WARNINGS
+
+	if (InArgs._SlotSize2f.IsSet())
+	{
+		Size = MoveTemp(InArgs._SlotSize2f);
 	}
 	if (InArgs._AllowScaling.IsSet())
 	{
@@ -301,7 +317,7 @@ void SNodePanel::ArrangeChildNodes(const FGeometry& AllottedGeometry, FArrangedC
 		const TSharedRef<SNode>& SomeChild = ChildrenToArrange[ChildIndex];
 		if (!SomeChild->RequiresSecondPassLayout())
 		{
-			ArrangedChildren.AddWidget(AllottedGeometry.MakeChild(SomeChild, SomeChild->GetPosition() - ViewOffset, SomeChild->GetDesiredSize(), GetZoomAmount()));
+			ArrangedChildren.AddWidget(AllottedGeometry.MakeChild(SomeChild, SomeChild->GetPosition2f() - ViewOffset, SomeChild->GetDesiredSize(), GetZoomAmount()));
 		}
 	}
 
@@ -312,7 +328,7 @@ void SNodePanel::ArrangeChildNodes(const FGeometry& AllottedGeometry, FArrangedC
 		if (SomeChild->RequiresSecondPassLayout())
 		{
 			SomeChild->PerformSecondPassLayout(NodeToWidgetLookup);
-			ArrangedChildren.AddWidget(AllottedGeometry.MakeChild(SomeChild, SomeChild->GetPosition() - ViewOffset, SomeChild->GetDesiredSize(), GetZoomAmount()));
+			ArrangedChildren.AddWidget(AllottedGeometry.MakeChild(SomeChild, SomeChild->GetPosition2f() - ViewOffset, SomeChild->GetDesiredSize(), GetZoomAmount()));
 		}
 	}
 }
@@ -356,18 +372,27 @@ FSlateColor SNodePanel::GetZoomTextColorAndOpacity() const
 	return FLinearColor( 1, 1, 1, 1.25f - ZoomLevelFade.GetLerp() );
 }
 
-FVector2D SNodePanel::GetViewOffset() const
+UE::Slate::FDeprecateVector2DResult SNodePanel::GetViewOffset() const
 {
 	return ViewOffset;
 }
 
 bool SNodePanel::GetZoomTargetRect(FVector2D& TopLeft, FVector2D& BottomRight) const
+{
+	FVector2f TempTopLeft, TempBottomRight;
+	bool bRetVal = GetZoomTargetRect(TempTopLeft, TempBottomRight);
+	TopLeft = FVector2D(TempTopLeft);
+	BottomRight = FVector2D(TempBottomRight);
+	return bRetVal;
+}
+
+bool SNodePanel::GetZoomTargetRect(FVector2f& TopLeft, FVector2f& BottomRight) const
 {
 	TopLeft = ZoomTargetTopLeft;
 	BottomRight = ZoomTargetBottomRight;
 	
 	// if the zoom target rect is all zeroed out, then notify caller that there is no target
-	return !(ZoomTargetTopLeft == FVector2D::ZeroVector && ZoomTargetBottomRight == FVector2D::ZeroVector);
+	return !(ZoomTargetTopLeft == FVector2f::ZeroVector && ZoomTargetBottomRight == FVector2f::ZeroVector);
 }
 
 void SNodePanel::Construct()
@@ -380,14 +405,14 @@ void SNodePanel::Construct()
 	PreviousZoomLevel = ZoomLevels->GetDefaultZoomLevel();
 	PostChangedZoom();
 
-	ViewOffset = FVector2D::ZeroVector;
+	ViewOffset = FVector2f::ZeroVector;
 	TotalMouseDelta = 0;
 	TotalMouseDeltaXY = 0;
 	bDeferredZoomToSelection = false;
 	bDeferredZoomToNodeExtents = false;
 
-	ZoomTargetTopLeft = FVector2D::ZeroVector;
-	ZoomTargetBottomRight = FVector2D::ZeroVector;
+	ZoomTargetTopLeft = FVector2f::ZeroVector;
+	ZoomTargetBottomRight = FVector2f::ZeroVector;
 	ZoomPadding = NodePanelDefs::DefaultZoomPadding;
 
 	bAllowContinousZoomInterpolation = false;
@@ -406,14 +431,14 @@ void SNodePanel::Construct()
 	ZoomLevelGraphFade = FCurveSequence( 0.0f, 0.5f );
 	ZoomLevelGraphFade.Play( this->AsShared() );
 
-	PastePosition = FVector2D::ZeroVector;
+	PastePosition = FVector2f::ZeroVector;
 
-	DeferredPanPosition = FVector2D::ZeroVector;
+	DeferredPanPosition = FVector2f::ZeroVector;
 	bRequestDeferredPan = false;
 
 	OldViewOffset = ViewOffset;
 	OldZoomAmount = GetZoomAmount();
-	ZoomStartOffset = FVector2D::ZeroVector;
+	ZoomStartOffset = FVector2f::ZeroVector;
 	TotalGestureMagnify = 0.0f;
 
 	ScopedTransactionPtr.Reset();
@@ -421,7 +446,7 @@ void SNodePanel::Construct()
 	bVisualUpdatePending = false;
 }
 
-FVector2D SNodePanel::ComputeEdgePanAmount(const FGeometry& MyGeometry, const FVector2D& TargetPosition)
+FVector2f SNodePanel::ComputeEdgePanAmount(const FGeometry& MyGeometry, const FVector2f& TargetPosition)
 {
 	// How quickly to ramp up the pan speed as the user moves the mouse further past the edge of the graph panel.
 	static const float EdgePanSpeedCoefficient = 2.f, EdgePanSpeedPower = 0.6f;
@@ -456,35 +481,35 @@ FVector2D SNodePanel::ComputeEdgePanAmount(const FGeometry& MyGeometry, const FV
 		EdgePanThisTick.Y = FMath::Min( MaxPanSpeed, EdgePanSpeedCoefficient * FMath::Pow(LocalCursorPos.Y - float(MyGeometry.GetLocalSize().Y) + EdgePanForgivenessZone, EdgePanSpeedPower) );
 	}
 
-	return FVector2D(EdgePanThisTick);
+	return EdgePanThisTick;
 }
 
-void SNodePanel::UpdateViewOffset (const FGeometry& MyGeometry, const FVector2D& TargetPosition)
+void SNodePanel::UpdateViewOffset(const FGeometry& MyGeometry, const FVector2f& TargetPosition)
 {
-	const FVector2D PanAmount = ComputeEdgePanAmount( MyGeometry, TargetPosition ) / GetZoomAmount();
+	const FVector2f PanAmount = ComputeEdgePanAmount( MyGeometry, TargetPosition ) / GetZoomAmount();
 	ViewOffset += PanAmount;
 }
 
-void SNodePanel::RequestDeferredPan(const FVector2D& UpdatePosition)
+void SNodePanel::RequestDeferredPan(const UE::Slate::FDeprecateVector2DParameter& UpdatePosition)
 {
 	bRequestDeferredPan = true;
 	DeferredPanPosition = UpdatePosition;
 }
 
-FVector2D SNodePanel::GraphCoordToPanelCoord( const FVector2D& GraphSpaceCoordinate ) const
+FVector2f SNodePanel::GraphCoordToPanelCoord( const UE::Slate::FDeprecateVector2DParameter& GraphSpaceCoordinate ) const
 {
 	return (GraphSpaceCoordinate - GetViewOffset()) * GetZoomAmount();
 }
 
-FVector2D SNodePanel::PanelCoordToGraphCoord( const FVector2D& PanelSpaceCoordinate ) const
+UE::Slate::FDeprecateVector2DResult SNodePanel::PanelCoordToGraphCoord( const UE::Slate::FDeprecateVector2DParameter& PanelSpaceCoordinate ) const
 {
 	return PanelSpaceCoordinate / GetZoomAmount() + GetViewOffset();
 }
 
 FSlateRect SNodePanel::PanelRectToGraphRect( const FSlateRect& PanelSpaceRect ) const
 {
-	FVector2D UpperLeft = PanelCoordToGraphCoord( FVector2D(PanelSpaceRect.Left, PanelSpaceRect.Top) );
-	FVector2D LowerRight = PanelCoordToGraphCoord(  FVector2D(PanelSpaceRect.Right, PanelSpaceRect.Bottom) );
+	FVector2f UpperLeft = PanelCoordToGraphCoord( FVector2f(PanelSpaceRect.Left, PanelSpaceRect.Top) );
+	FVector2f LowerRight = PanelCoordToGraphCoord(  FVector2f(PanelSpaceRect.Right, PanelSpaceRect.Bottom) );
 
 	return FSlateRect(
 		UpperLeft.X, UpperLeft.Y,
@@ -492,6 +517,11 @@ FSlateRect SNodePanel::PanelRectToGraphRect( const FSlateRect& PanelSpaceRect )
 }
 
 void SNodePanel::OnBeginNodeInteraction(const TSharedRef<SNode>& InNodeToDrag, const FVector2D& GrabOffset)
+{
+	OnBeginNodeInteraction(InNodeToDrag, UE::Slate::CastToVector2f(GrabOffset));
+}
+
+void SNodePanel::OnBeginNodeInteraction(const TSharedRef<SNode>& InNodeToDrag, const FVector2f& GrabOffset)
 {
 	NodeUnderMousePtr = InNodeToDrag;
 	NodeGrabOffset = GrabOffset;
@@ -504,7 +534,7 @@ void SNodePanel::OnEndNodeInteraction(const TSharedRef<SNode>& InNodeToDrag)
 
 EActiveTimerReturnType SNodePanel::HandleZoomToFit(double InCurrentTime, float InDeltaTime)
 {
-	const FVector2D DesiredViewCenter = ( ZoomTargetTopLeft + ZoomTargetBottomRight ) * 0.5f;
+	const FVector2f DesiredViewCenter = ( ZoomTargetTopLeft + ZoomTargetBottomRight ) * 0.5f;
 	const bool bDoneScrolling = ScrollToLocation(CachedGeometry, DesiredViewCenter, bTeleportInsteadOfScrollingWhenZoomingToFit ? 1000.0f : InDeltaTime);
 	bool bDoneZooming = ZoomToLocation(CachedGeometry.GetLocalSize(), ZoomTargetBottomRight - ZoomTargetTopLeft, bDoneScrolling);
 
@@ -515,8 +545,8 @@ EActiveTimerReturnType SNodePanel::HandleZoomToFit(double InCurrentTime, float I
 		
 		// Reset ZoomPadding
 		ZoomPadding = NodePanelDefs::DefaultZoomPadding;
-		ZoomTargetTopLeft = FVector2D::ZeroVector;
-		ZoomTargetBottomRight = FVector2D::ZeroVector;
+		ZoomTargetTopLeft = FVector2f::ZeroVector;
+		ZoomTargetBottomRight = FVector2f::ZeroVector;
 
 		DeferredMovementTargetObject = nullptr;
 
@@ -661,7 +691,7 @@ FReply SNodePanel::OnMouseButtonDown( const FGeometry& MyGeometry, const FPointe
 		if (bIsLeftMouseButtonEffecting)
 		{
 			// Got here from panning mode (with RMB held) - clear panning mode, but use cached software cursor position
-			const FVector2D WidgetSpaceCursorPos = GraphCoordToPanelCoord( SoftwareCursorPosition );
+			const FVector2f WidgetSpaceCursorPos = GraphCoordToPanelCoord( SoftwareCursorPosition );
 			ZoomStartOffset = WidgetSpaceCursorPos;
 			this->bIsPanning = false;
 		}
@@ -716,7 +746,7 @@ FReply SNodePanel::OnMouseButtonDown( const FGeometry& MyGeometry, const FPointe
 
 			// This changes selection and starts dragging it.
 			const FArrangedWidget& NodeGeometry = ArrangedChildren[NodeUnderMouseIndex];
-			const FVector2D MousePositionInNode = NodeGeometry.Geometry.AbsoluteToLocal(MouseEvent.GetScreenSpacePosition());
+			const FVector2f MousePositionInNode = NodeGeometry.Geometry.AbsoluteToLocal(MouseEvent.GetScreenSpacePosition());
 			TSharedRef<SNode> NodeWidgetUnderMouse = StaticCastSharedRef<SNode>( NodeGeometry.Widget );
 
 			if( NodeWidgetUnderMouse->CanBeSelected(MousePositionInNode) )
@@ -728,7 +758,7 @@ FReply SNodePanel::OnMouseButtonDown( const FGeometry& MyGeometry, const FPointe
 		}
 
 		// START MARQUEE SELECTION.
-		const FVector2D GraphMousePos = PanelCoordToGraphCoord( MyGeometry.AbsoluteToLocal( MouseEvent.GetScreenSpacePosition() ) );
+		const FVector2f GraphMousePos = PanelCoordToGraphCoord( MyGeometry.AbsoluteToLocal( MouseEvent.GetScreenSpacePosition() ) );
 		Marquee.Start( GraphMousePos, FMarqueeOperation::OperationTypeFromMouseEvent(MouseEvent) );
 
 		// If we're marquee selecting, then we're not clicking on a node!
@@ -754,7 +784,7 @@ FReply SNodePanel::OnMouseMove(const FGeometry& MyGeometry, const FPointerEvent&
 
 	if ( this->HasMouseCapture() )
 	{
-		const FVector2D CursorDelta = MouseEvent.GetCursorDelta();
+		const FVector2f CursorDelta = MouseEvent.GetCursorDelta();
 		// Track how much the mouse moved since the mouse down.
 		TotalMouseDelta += CursorDelta.Size();
 
@@ -842,10 +872,10 @@ FReply SNodePanel::OnMouseMove(const FGeometry& MyGeometry, const FPointerEvent&
 					if ( !bCursorInDeadZone )
 					{
 						// Note, NodeGrabOffset() comes from the node itself, so it's already scaled correctly.
-						FVector2D AnchorNodeNewPos = PanelCoordToGraphCoord( MyGeometry.AbsoluteToLocal( MouseEvent.GetScreenSpacePosition() ) ) - NodeGrabOffset;
+						FVector2f AnchorNodeNewPos = PanelCoordToGraphCoord( MyGeometry.AbsoluteToLocal( MouseEvent.GetScreenSpacePosition() ) ) - NodeGrabOffset;
 
 						// Snap to grid
-						const float SnapSize = GetSnapGridSize();
+						const float SnapSize = static_cast<float>(GetSnapGridSize());
 						AnchorNodeNewPos.X = SnapSize * FMath::RoundToFloat( AnchorNodeNewPos.X / SnapSize );
 						AnchorNodeNewPos.Y = SnapSize * FMath::RoundToFloat( AnchorNodeNewPos.Y / SnapSize );
 
@@ -854,8 +884,8 @@ FReply SNodePanel::OnMouseMove(const FGeometry& MyGeometry, const FPointerEvent&
 
 						// Move all the selected nodes.
 						{
-							const FVector2D AnchorNodeOldPos = NodeBeingDragged->GetPosition();
-							const FVector2D DeltaPos = AnchorNodeNewPos - AnchorNodeOldPos;
+							const FVector2f AnchorNodeOldPos = NodeBeingDragged->GetPosition2f();
+							const FVector2f DeltaPos = AnchorNodeNewPos - AnchorNodeOldPos;
 
 							// Perform movement in 3 passes:
 							
@@ -863,7 +893,7 @@ FReply SNodePanel::OnMouseMove(const FGeometry& MyGeometry, const FPointerEvent&
 							struct FDefferedNodePosition 
 							{ 
 								SNode*		Node; 
-								FVector2D	NewPosition; 
+								FVector2f	NewPosition; 
 							};
 							TArray<FDefferedNodePosition> DefferedNodesToMove;
 
@@ -874,12 +904,12 @@ FReply SNodePanel::OnMouseMove(const FGeometry& MyGeometry, const FPointerEvent&
 								if (TSharedRef<SNode>* pWidget = NodeToWidgetLookup.Find(*NodeIt))
 								{
 									SNode& Widget = pWidget->Get();
-									FDefferedNodePosition NodePosition = { &Widget, Widget.GetPosition() + DeltaPos };
+									FDefferedNodePosition NodePosition = { &Widget, Widget.GetPosition2f() + DeltaPos };
 									DefferedNodesToMove.Add(NodePosition);
 
 									if (bStoreOriginalNodePositions)
 									{
-										OriginalNodePositions.FindOrAdd(*pWidget) = Widget.GetPosition();
+										OriginalNodePositions.FindOrAdd(*pWidget) = Widget.GetPosition2f();
 									}
 								}
 							}
@@ -901,7 +931,7 @@ FReply SNodePanel::OnMouseMove(const FGeometry& MyGeometry, const FPointerEvent&
 			if ( !NodeBeingDragged.IsValid() )
 			{
 				// We are marquee selecting
-				const FVector2D GraphMousePos = PanelCoordToGraphCoord( MyGeometry.AbsoluteToLocal( MouseEvent.GetScreenSpacePosition() ) );
+				const FVector2f GraphMousePos = PanelCoordToGraphCoord( MyGeometry.AbsoluteToLocal( MouseEvent.GetScreenSpacePosition() ) );
 				Marquee.Rect.UpdateEndPoint(GraphMousePos);
 
 				FindNodesAffectedByMarquee( /*out*/ Marquee.AffectedNodes );
@@ -1046,7 +1076,7 @@ FReply SNodePanel::OnMouseButtonUp( const FGeometry& MyGeometry, const FPointerE
 		if ( this->HasMouseCapture() )
 		{
 			FSlateRect ThisPanelScreenSpaceRect = MyGeometry.GetLayoutBoundingRect();
-			const FVector2D ScreenSpaceCursorPos = MyGeometry.LocalToAbsolute( GraphCoordToPanelCoord( SoftwareCursorPosition ) );
+			const FVector2f ScreenSpaceCursorPos = MyGeometry.LocalToAbsolute( GraphCoordToPanelCoord( SoftwareCursorPosition ) );
 
 			FIntPoint BestPositionInViewport(
 				FMath::RoundToInt( FMath::Clamp( ScreenSpaceCursorPos.X, ThisPanelScreenSpaceRect.Left, ThisPanelScreenSpaceRect.Right ) ),
@@ -1078,7 +1108,7 @@ void SNodePanel::OnMouseCaptureLost(const FCaptureLostEvent& CaptureLostEvent)
 FReply SNodePanel::OnMouseWheel(const FGeometry& MyGeometry, const FPointerEvent& MouseEvent)
 {
 	// We want to zoom into this point; i.e. keep it the same fraction offset into the panel
-	const FVector2D WidgetSpaceCursorPos = MyGeometry.AbsoluteToLocal( MouseEvent.GetScreenSpacePosition() );
+	const FVector2f WidgetSpaceCursorPos = MyGeometry.AbsoluteToLocal( MouseEvent.GetScreenSpacePosition() );
 	const int32 ZoomLevelDelta = FMath::TruncToInt( FMath::RoundFromZero( MouseEvent.GetWheelDelta() ) );
 	ChangeZoomLevel( ZoomLevelDelta, WidgetSpaceCursorPos, MouseEvent.IsControlDown() );
 
@@ -1123,14 +1153,14 @@ void SNodePanel::OnFocusLost( const FFocusEvent& InFocusEvent )
 FReply SNodePanel::OnTouchGesture( const FGeometry& MyGeometry, const FPointerEvent& GestureEvent )
 {
 	const EGestureEvent GestureType = GestureEvent.GetGestureType();
-	const FVector2D& GestureDelta = GestureEvent.GetGestureDelta();
+	const FVector2f& GestureDelta = GestureEvent.GetGestureDelta();
 	if (GestureType == EGestureEvent::Magnify)
 	{
 		TotalGestureMagnify += GestureDelta.X;
 		if (FMath::Abs(TotalGestureMagnify) > 0.07f)
 		{
 			// We want to zoom into this point; i.e. keep it the same fraction offset into the panel
-			const FVector2D WidgetSpaceCursorPos = MyGeometry.AbsoluteToLocal(GestureEvent.GetScreenSpacePosition());
+			const FVector2f WidgetSpaceCursorPos = MyGeometry.AbsoluteToLocal(GestureEvent.GetScreenSpacePosition());
 			const int32 ZoomLevelDelta = TotalGestureMagnify > 0.0f ? 1 : -1;
 			ChangeZoomLevel(ZoomLevelDelta, WidgetSpaceCursorPos, GestureEvent.IsControlDown());
 			TotalGestureMagnify = 0.0f;
@@ -1177,8 +1207,8 @@ void SNodePanel::FindNodesAffectedByMarquee( FGraphPanelSelectionSet& OutAffecte
 	for ( int32 NodeIndex=0; NodeIndex < Children.Num(); ++NodeIndex )
 	{
 		const TSharedRef<SNode>& SomeNodeWidget = Children[NodeIndex];
-		const FVector2D NodePosition = SomeNodeWidget->GetPosition();
-		const FVector2D NodeSize = SomeNodeWidget->GetDesiredSizeForMarquee();
+		const FVector2f NodePosition = SomeNodeWidget->GetPosition2f();
+		const FVector2f NodeSize = SomeNodeWidget->GetDesiredSizeForMarquee2f();
 
 		if (NodeSize.X > 0.f && NodeSize.Y > 0.f)
 		{
@@ -1304,7 +1334,7 @@ bool SNodePanel::Contains(UObject* Node) const
 	return NodeToWidgetLookup.Find(Node) != nullptr;
 }
 
-void SNodePanel::RestoreViewSettings(const FVector2D& InViewOffset, float InZoomAmount, const FGuid& InBookmarkGuid)
+void SNodePanel::RestoreViewSettings(const UE::Slate::FDeprecateVector2DParameter& InViewOffset, float InZoomAmount, const FGuid& InBookmarkGuid)
 {
 	ViewOffset = InViewOffset;
 
@@ -1356,7 +1386,7 @@ void SNodePanel::PaintBackgroundAsLines(const FSlateBrush* BackgroundImage, cons
 	const FLinearColor CenterColor(GetDefault<UEditorStyleSettings>()->CenterColor);
 	const float GraphSmallestGridSize = 8.0f;
 	const float RawZoomFactor = GetZoomAmount();
-	const float NominalGridSize = GetSnapGridSize();
+	const float NominalGridSize = static_cast<float>(GetSnapGridSize());
 
 	float ZoomFactor = RawZoomFactor;
 	float Inflation = 1.0f;
@@ -1373,7 +1403,7 @@ void SNodePanel::PaintBackgroundAsLines(const FSlateBrush* BackgroundImage, cons
 	float ImageOffsetX = GraphSpaceGridX0 * -ZoomFactor;
 	float ImageOffsetY = GraphSpaceGridY0 * -ZoomFactor;
 
-	const FVector2D ZeroSpace = GraphCoordToPanelCoord(FVector2D::ZeroVector);
+	const FVector2f ZeroSpace = GraphCoordToPanelCoord(FVector2f::ZeroVector);
 	
 	// Fill the background
 	FSlateDrawElement::MakeBox(
@@ -1385,9 +1415,9 @@ void SNodePanel::PaintBackgroundAsLines(const FSlateBrush* BackgroundImage, cons
 		GraphBackGroundImageColor
 	);
 
-	TArray<FVector2D> LinePoints;
-	new (LinePoints) FVector2D(0.0f, 0.0f);
-	new (LinePoints) FVector2D(0.0f, 0.0f);
+	TArray<FVector2f> LinePoints;
+	new (LinePoints) FVector2f(0.0f, 0.0f);
+	new (LinePoints) FVector2f(0.0f, 0.0f);
 
 	//If we want to use grid then show grid, otherwise don't render the grid
 	if (GetDefault<UEditorStyleSettings>()->bUseGrid == true){
@@ -1406,8 +1436,8 @@ void SNodePanel::PaintBackgroundAsLines(const FSlateBrush* BackgroundImage, cons
 					Color = &CenterColor;
 				}
 
-				LinePoints[0] = FVector2D(0.0f, ImageOffsetY);
-				LinePoints[1] = FVector2D(AllottedGeometry.GetLocalSize().X, ImageOffsetY);
+				LinePoints[0] = FVector2f(0.0f, ImageOffsetY);
+				LinePoints[1] = FVector2f(AllottedGeometry.GetLocalSize().X, ImageOffsetY);
 
 				FSlateDrawElement::MakeLines(
 					OutDrawElements,
@@ -1434,8 +1464,8 @@ void SNodePanel::PaintBackgroundAsLines(const FSlateBrush* BackgroundImage, cons
 					Color = &CenterColor;
 				}
 
-				LinePoints[0] = FVector2D(ImageOffsetX, 0.0f);
-				LinePoints[1] = FVector2D(ImageOffsetX, AllottedGeometry.GetLocalSize().Y);
+				LinePoints[0] = FVector2f(ImageOffsetX, 0.0f);
+				LinePoints[1] = FVector2f(ImageOffsetX, AllottedGeometry.GetLocalSize().Y);
 
 				FSlateDrawElement::MakeLines(
 					OutDrawElements,
@@ -1500,16 +1530,16 @@ void SNodePanel::PaintComment(const FString& CommentText, const FGeometry& Allot
 	const FSlateBrush* CommentCalloutBubble = FAppStyle::GetBrush(TEXT("Graph.Node.CommentBubble"));
 	const FSlateFontInfo CommentFont = FAppStyle::GetFontStyle( TEXT("Graph.Node.CommentFont") );
 	const FSlateColor CommentTextColor = FAppStyle::GetColor( TEXT("Graph.Node.Comment.TextColor") );
-	const FVector2D CommentBubblePadding = FAppStyle::GetVector( TEXT("Graph.Node.Comment.BubblePadding") );
+	const FVector2f CommentBubblePadding = UE::Slate::CastToVector2f(FAppStyle::GetVector( TEXT("Graph.Node.Comment.BubblePadding") ));
 
 	const TSharedRef< FSlateFontMeasure > FontMeasureService = FSlateApplication::Get().GetRenderer()->GetFontMeasureService();
-	FVector2D CommentTextSize = FontMeasureService->Measure( CommentText, CommentFont ) + (CommentBubblePadding * 2);
+	FVector2f CommentTextSize = FontMeasureService->Measure( CommentText, CommentFont ) + (CommentBubblePadding * 2);
 
 	const float PositionBias = HeightAboveNode;
 	HeightAboveNode += CommentTextSize.Y + 8.0f;
 
-	const FVector2D CommentBubbleOffset = FVector2D(0, -(CommentTextSize.Y + CommentCalloutArrow->ImageSize.Y) - PositionBias);
-	const FVector2D CommentBubbleArrowOffset = FVector2D( CommentCalloutArrow->ImageSize.X, -CommentCalloutArrow->ImageSize.Y - PositionBias);
+	const FVector2f CommentBubbleOffset = FVector2f(0, -(CommentTextSize.Y + CommentCalloutArrow->ImageSize.Y) - PositionBias);
+	const FVector2f CommentBubbleArrowOffset = FVector2f( CommentCalloutArrow->ImageSize.X, -CommentCalloutArrow->ImageSize.Y - PositionBias);
 
 	// We want the watch window comment bubbles to draw on top of the node so that they are not obscured
 	// and on top of the big blueprint debugging arrow
@@ -1553,10 +1583,10 @@ bool SNodePanel::IsNodeCulled(const TSharedRef<SNode>& Node, const FGeometry& Al
 {
 	if ( Node->ShouldAllowCulling() )
 	{
-		const FVector2D MinClipArea = AllottedGeometry.GetDrawSize() * -NodePanelDefs::GuardBandArea;
-		const FVector2D MaxClipArea = AllottedGeometry.GetDrawSize() * ( 1.f + NodePanelDefs::GuardBandArea );
-		const FVector2D NodeTopLeft = GraphCoordToPanelCoord( Node->GetPosition() );
-		const FVector2D NodeBottomRight = GraphCoordToPanelCoord( Node->GetPosition() + Node->GetDesiredSize() );
+		const FVector2f MinClipArea = AllottedGeometry.GetDrawSize() * -NodePanelDefs::GuardBandArea;
+		const FVector2f MaxClipArea = AllottedGeometry.GetDrawSize() * ( 1.f + NodePanelDefs::GuardBandArea );
+		const FVector2f NodeTopLeft = GraphCoordToPanelCoord( Node->GetPosition2f() );
+		const FVector2f NodeBottomRight = GraphCoordToPanelCoord( Node->GetPosition2f() + Node->GetDesiredSize() );
 
 		return 
 			NodeBottomRight.X < MinClipArea.X ||
@@ -1570,11 +1600,20 @@ bool SNodePanel::IsNodeCulled(const TSharedRef<SNode>& Node, const FGeometry& Al
 	}
 
 }
-
 bool SNodePanel::GetBoundsForNode(const UObject* InNode, /*out*/ FVector2D& MinCorner, /*out*/ FVector2D& MaxCorner, float Padding) const
 {
-	MinCorner = FVector2D(MAX_FLT, MAX_FLT);
-	MaxCorner = FVector2D(-MAX_FLT, -MAX_FLT);
+	FVector2f TempMinCorner, TempMaxCorner;
+	bool bRetVal = SNodePanel::GetBoundsForNode(InNode, TempMinCorner, TempMaxCorner, Padding);
+	MinCorner = FVector2D(TempMinCorner);
+	MaxCorner = FVector2D(TempMaxCorner);
+
+	return bRetVal;
+}
+
+bool SNodePanel::GetBoundsForNode(const UObject* InNode, /*out*/ FVector2f& MinCorner, /*out*/ FVector2f& MaxCorner, float Padding) const
+{
+	MinCorner = FVector2f(MAX_FLT, MAX_FLT);
+	MaxCorner = FVector2f(-MAX_FLT, -MAX_FLT);
 
 	bool bValid = false;
 
@@ -1582,8 +1621,8 @@ bool SNodePanel::GetBoundsForNode(const UObject* InNode, /*out*/ FVector2D& MinC
 	if (pWidget)
 	{
 		const SNode& Widget = pWidget->Get();
-		const FVector2D Lower = Widget.GetPosition();
-		const FVector2D Upper = Lower + Widget.GetDesiredSize();
+		const FVector2f Lower = Widget.GetPosition2f();
+		const FVector2f Upper = Lower + Widget.GetDesiredSize();
 
 		MinCorner.X = FMath::Min(MinCorner.X, Lower.X);
 		MinCorner.Y = FMath::Min(MinCorner.Y, Lower.Y);
@@ -1603,11 +1642,20 @@ bool SNodePanel::GetBoundsForNode(const UObject* InNode, /*out*/ FVector2D& MinC
 
 	return bValid;
 }
+bool SNodePanel::GetBoundsForNodes(bool bSelectionSetOnly, FVector2D& MinCorner, FVector2D& MaxCorner, float Padding) const
+{
+	FVector2f TempMinCorner, TempMaxCorner;
+	bool bRetVal = SNodePanel::GetBoundsForNodes(bSelectionSetOnly, TempMinCorner, TempMaxCorner, Padding);
+	MinCorner = FVector2D(TempMinCorner);
+	MaxCorner = FVector2D(TempMaxCorner);
 
-bool SNodePanel::GetBoundsForNodes(bool bSelectionSetOnly, FVector2D& MinCorner, FVector2D& MaxCorner, float Padding)
+	return bRetVal;
+}
+
+bool SNodePanel::GetBoundsForNodes(bool bSelectionSetOnly, FVector2f& MinCorner, FVector2f& MaxCorner, float Padding) const
 {
-	MinCorner = FVector2D(MAX_FLT, MAX_FLT);
-	MaxCorner = FVector2D(-MAX_FLT, -MAX_FLT);
+	MinCorner = FVector2f(MAX_FLT, MAX_FLT);
+	MaxCorner = FVector2f(-MAX_FLT, -MAX_FLT);
 
 	bool bValid = false;
 
@@ -1615,12 +1663,12 @@ bool SNodePanel::GetBoundsForNodes(bool bSelectionSetOnly, FVector2D& MinCorner,
 	{
 		for (FGraphPanelSelectionSet::TConstIterator NodeIt(SelectionManager.GetSelectedNodes()); NodeIt; ++NodeIt)
 		{
-			TSharedRef<SNode>* pWidget = NodeToWidgetLookup.Find(*NodeIt);
+			const TSharedRef<SNode>* pWidget = NodeToWidgetLookup.Find(*NodeIt);
 			if (pWidget != nullptr)
 			{
 				SNode& Widget = pWidget->Get();
-				const FVector2D Lower = Widget.GetPosition();
-				const FVector2D Upper = Lower + Widget.GetDesiredSize();
+				const FVector2f Lower = Widget.GetPosition2f();
+				const FVector2f Upper = Lower + Widget.GetDesiredSize();
 
 				MinCorner.X = FMath::Min(MinCorner.X, Lower.X);
 				MinCorner.Y = FMath::Min(MinCorner.Y, Lower.Y);
@@ -1637,8 +1685,8 @@ bool SNodePanel::GetBoundsForNodes(bool bSelectionSetOnly, FVector2D& MinCorner,
 		{
 			SNode& Widget = NodeIt.Value().Get();
 
-			const FVector2D Lower = Widget.GetPosition();
-			const FVector2D Upper = Lower + Widget.GetDesiredSize();
+			const FVector2f Lower = Widget.GetPosition2f();
+			const FVector2f Upper = Lower + Widget.GetDesiredSize();
 
 			MinCorner.X = FMath::Min(MinCorner.X, Lower.X);
 			MinCorner.Y = FMath::Min(MinCorner.Y, Lower.Y);
@@ -1658,19 +1706,19 @@ bool SNodePanel::GetBoundsForNodes(bool bSelectionSetOnly, FVector2D& MinCorner,
 	return bValid;
 }
 
-bool SNodePanel::ScrollToLocation(const FGeometry& MyGeometry, FVector2D DesiredCenterPosition, const float InDeltaTime)
+bool SNodePanel::ScrollToLocation(const FGeometry& MyGeometry, FVector2f DesiredCenterPosition, const float InDeltaTime)
 {
-	const FVector2D HalfOFScreenInGraphSpace = 0.5f * MyGeometry.GetLocalSize() / GetZoomAmount();
-	FVector2D CurrentPosition = ViewOffset + HalfOFScreenInGraphSpace;
+	const FVector2f HalfOFScreenInGraphSpace = 0.5f * MyGeometry.GetLocalSize() / GetZoomAmount();
+	FVector2f CurrentPosition = ViewOffset + HalfOFScreenInGraphSpace;
 
-	FVector2D NewPosition = FMath::Vector2DInterpTo(CurrentPosition, DesiredCenterPosition, InDeltaTime, 10.f);
+	FVector2f NewPosition = FMath::Vector2DInterpTo(CurrentPosition, DesiredCenterPosition, InDeltaTime, 10.f);
 	ViewOffset = NewPosition - HalfOFScreenInGraphSpace;
 
 	// If within 1 pixel of target, stop interpolating
 	return ((NewPosition - DesiredCenterPosition).SizeSquared() < 1.f);
 }
 
-bool SNodePanel::ZoomToLocation(const FVector2D& CurrentSizeWithoutZoom, const FVector2D& InDesiredSize, bool bDoneScrolling)
+bool SNodePanel::ZoomToLocation(const FVector2f& CurrentSizeWithoutZoom, const FVector2f& InDesiredSize, bool bDoneScrolling)
 {
 	if (bAllowContinousZoomInterpolation && ZoomLevelGraphFade.IsPlaying())
 	{
@@ -1684,13 +1732,13 @@ bool SNodePanel::ZoomToLocation(const FVector2D& CurrentSizeWithoutZoom, const F
 	// Find lowest zoom level that will display all nodes
 	for (int32 Zoom = 0; Zoom < DefaultZoomLevel; ++Zoom)
 	{
-		const FVector2D SizeWithZoom = CurrentSizeWithoutZoom / ZoomLevels->GetZoomAmount(Zoom);
-		const FVector2D LeftOverSize = SizeWithZoom - InDesiredSize;
+		const FVector2f SizeWithZoom = CurrentSizeWithoutZoom / ZoomLevels->GetZoomAmount(Zoom);
+		const FVector2f LeftOverSize = SizeWithZoom - InDesiredSize;
 		
 		if ((InDesiredSize.X > SizeWithZoom.X) || (InDesiredSize.Y > SizeWithZoom.Y))
 		{
 			// Use the previous zoom level, this one is too tight
-			DesiredZoom = FMath::Max<int32>(0, Zoom - 1);
+			DesiredZoom = FMath::Max(0, Zoom - 1);
 			break; 
 		}
 	}
@@ -1738,7 +1786,7 @@ void SNodePanel::ZoomToFit(bool bOnlySelection)
 	ZoomPadding = NodePanelDefs::DefaultZoomPadding;
 }
 
-void SNodePanel::ZoomToTarget(const FVector2D& TopLeft, const FVector2D& BottomRight)
+void SNodePanel::ZoomToTarget(const FVector2f& TopLeft, const FVector2f& BottomRight)
 {
 	bDeferredZoomToNodeExtents = false;
 
@@ -1748,10 +1796,10 @@ void SNodePanel::ZoomToTarget(const FVector2D& TopLeft, const FVector2D& BottomR
 	RequestZoomToFit();
 }
 
-void SNodePanel::ChangeZoomLevel(int32 ZoomLevelDelta, const FVector2D& WidgetSpaceZoomOrigin, bool bOverrideZoomLimiting)
+void SNodePanel::ChangeZoomLevel(int32 ZoomLevelDelta, const FVector2f& WidgetSpaceZoomOrigin, bool bOverrideZoomLimiting)
 {
 	// We want to zoom into this point; i.e. keep it the same fraction offset into the panel
-	const FVector2D PointToMaintainGraphSpace = PanelCoordToGraphCoord( WidgetSpaceZoomOrigin );
+	const FVector2f PointToMaintainGraphSpace = PanelCoordToGraphCoord( WidgetSpaceZoomOrigin );
 
 	const int32 DefaultZoomLevel = ZoomLevels->GetDefaultZoomLevel();
 	const int32 NumZoomLevels = ZoomLevels->GetNumZoomLevels();
@@ -1765,7 +1813,7 @@ void SNodePanel::ChangeZoomLevel(int32 ZoomLevelDelta, const FVector2D& WidgetSp
 		// If they are already zoomed in past 1:1, user may zoom freely
 		(ZoomLevel > DefaultZoomLevel);
 
-	const float OldZoomLevel = ZoomLevel;
+	const int32 OldZoomLevel = ZoomLevel;
 
 	if ( bAllowFullZoomRange )
 	{
@@ -1794,7 +1842,7 @@ bool SNodePanel::GetBoundsForSelectedNodes( class FSlateRect& Rect, float Paddin
 	bool Result = false;
 	if(SelectionManager.GetSelectedNodes().Num()  > 0)
 	{
-		FVector2D MinCorner, MaxCorner;
+		FVector2f MinCorner, MaxCorner;
 		Result =  GetBoundsForNodes(true, MinCorner, MaxCorner,Padding);
 
 		Rect = FSlateRect(MinCorner.X, MinCorner.Y, MaxCorner.X, MaxCorner.Y);
@@ -1802,7 +1850,7 @@ bool SNodePanel::GetBoundsForSelectedNodes( class FSlateRect& Rect, float Paddin
 	return Result;
 }
 
-FVector2D SNodePanel::GetPastePosition() const
+UE::Slate::FDeprecateVector2DResult SNodePanel::GetPastePosition() const
 {
 	return PastePosition;
 }
@@ -1823,7 +1871,7 @@ void SNodePanel::FinalizeNodeMovements()
 	if (OriginalNodePositions.Num() > 0)
 	{
 		// Build up all the current positions
-		TMap<SNode*, FVector2D> CurrentNodePositions;
+		TMap<SNode*, FVector2f> CurrentNodePositions;
 
 		for (decltype(SelectionManager.SelectedNodes)::TIterator NodeIt(SelectionManager.SelectedNodes); NodeIt; ++NodeIt)
 		{
@@ -1831,14 +1879,14 @@ void SNodePanel::FinalizeNodeMovements()
 			if (pWidget != nullptr)
 			{
 				SNode& Widget = pWidget->Get();
-				CurrentNodePositions.FindOrAdd(&Widget) = Widget.GetPosition();
+				CurrentNodePositions.FindOrAdd(&Widget) = Widget.GetPosition2f();
 			}
 		}
 
 		// Move all the nodes back to their original position before we start the transaction
 		SNode::FNodeSet OriginalNodeFilter;
 
-		for (TPair<TWeakPtr<SNode>, FVector2D>& OriginalNodePosition : OriginalNodePositions)
+		for (TPair<TWeakPtr<SNode>, FVector2f>& OriginalNodePosition : OriginalNodePositions)
 		{
 			if (OriginalNodePosition.Key.IsValid())
 			{
@@ -1858,7 +1906,7 @@ void SNodePanel::FinalizeNodeMovements()
 			// Move all the nodes back to their current position but on on the undo stack
 			SNode::FNodeSet CurrentNodeFilter;
 
-			for (TPair<SNode*, FVector2D>& CurrentNodePosition : CurrentNodePositions)
+			for (TPair<SNode*, FVector2f>& CurrentNodePosition : CurrentNodePositions)
 			{
 				CurrentNodePosition.Key->MoveTo(CurrentNodePosition.Value, CurrentNodeFilter, true);
 			}
@@ -1888,8 +1936,8 @@ void SNodePanel::CancelZoomToFit()
 	{
 		// Reset Zoom destination
 		ZoomPadding = NodePanelDefs::DefaultZoomPadding;
-		ZoomTargetTopLeft = FVector2D::ZeroVector;
-		ZoomTargetBottomRight = FVector2D::ZeroVector;
+		ZoomTargetTopLeft = FVector2f::ZeroVector;
+		ZoomTargetBottomRight = FVector2f::ZeroVector;
 		UnRegisterActiveTimer(ActiveTimerHandle.Pin().ToSharedRef());
 	}
 }
diff --git a/Engine/Source/Editor/GraphEditor/Public/KismetNodes/SGraphNodeK2Base.h b/Engine/Source/Editor/GraphEditor/Public/KismetNodes/SGraphNodeK2Base.h
--- a/Engine/Source/Editor/GraphEditor/Public/KismetNodes/SGraphNodeK2Base.h
+++ b/Engine/Source/Editor/GraphEditor/Public/KismetNodes/SGraphNodeK2Base.h
@@ -32,7 +32,7 @@ public:
 	virtual void UpdateGraphNode() override;
 
 	//~ Begin SNodePanel::SNode Interface
-	virtual void GetOverlayBrushes(bool bSelected, const FVector2D WidgetSize, TArray<FOverlayBrushInfo>& Brushes) const override;
+	virtual void GetOverlayBrushes(bool bSelected, const FVector2f& WidgetSize, TArray<FOverlayBrushInfo>& Brushes) const override;
 	virtual void GetNodeInfoPopups(FNodeInfoContext* Context, TArray<FGraphInformationPopupInfo>& Popups) const override;
 	virtual const FSlateBrush* GetShadowBrush(bool bSelected) const override;
 	virtual void GetDiffHighlightBrushes(const FSlateBrush*& BackgroundOut, const FSlateBrush*& ForegroundOut) const override;
diff --git a/Engine/Source/Editor/GraphEditor/Public/MaterialNodes/SGraphNodeMaterialBase.h b/Engine/Source/Editor/GraphEditor/Public/MaterialNodes/SGraphNodeMaterialBase.h
--- a/Engine/Source/Editor/GraphEditor/Public/MaterialNodes/SGraphNodeMaterialBase.h
+++ b/Engine/Source/Editor/GraphEditor/Public/MaterialNodes/SGraphNodeMaterialBase.h
@@ -101,7 +101,7 @@ public:
 	// End of SGraphNode interface
 
 	// SNodePanel::SNode interface
-	virtual void MoveTo(const FVector2D& NewPosition, FNodeSet& NodeFilter, bool bMarkDirty = true) override;
+	virtual void MoveTo(const FVector2f& NewPosition, FNodeSet& NodeFilter, bool bMarkDirty = true) override;
 	// End of SNodePanel::SNode interface
 
 	UMaterialGraphNode* GetMaterialGraphNode() const {return MaterialNode;}
diff --git a/Engine/Source/Editor/GraphEditor/Public/MaterialNodes/SGraphNodeMaterialComposite.h b/Engine/Source/Editor/GraphEditor/Public/MaterialNodes/SGraphNodeMaterialComposite.h
--- a/Engine/Source/Editor/GraphEditor/Public/MaterialNodes/SGraphNodeMaterialComposite.h
+++ b/Engine/Source/Editor/GraphEditor/Public/MaterialNodes/SGraphNodeMaterialComposite.h
@@ -29,7 +29,7 @@ public:
 	// End of SGraphNode interface
 
 	// SNodePanel::SNode interface
-	virtual void MoveTo(const FVector2D& NewPosition, FNodeSet& NodeFilter, bool bMarkDirty = true) override;
+	virtual void MoveTo(const FVector2f& NewPosition, FNodeSet& NodeFilter, bool bMarkDirty = true) override;
 	// End of SNodePanel::SNode interface
 protected:
 	virtual UEdGraph* GetInnerGraph() const;
diff --git a/Engine/Source/Editor/GraphEditor/Public/MaterialNodes/SGraphNodeMaterialResult.h b/Engine/Source/Editor/GraphEditor/Public/MaterialNodes/SGraphNodeMaterialResult.h
--- a/Engine/Source/Editor/GraphEditor/Public/MaterialNodes/SGraphNodeMaterialResult.h
+++ b/Engine/Source/Editor/GraphEditor/Public/MaterialNodes/SGraphNodeMaterialResult.h
@@ -24,7 +24,7 @@ public:
 	// End of SGraphNode interface
 
 	// SNodePanel::SNode interface
-	virtual void MoveTo(const FVector2D& NewPosition, FNodeSet& NodeFilter, bool bMarkDirty = true) override;
+	virtual void MoveTo(const FVector2f& NewPosition, FNodeSet& NodeFilter, bool bMarkDirty = true) override;
 	// End of SNodePanel::SNode interface
 
 private:
diff --git a/Engine/Source/Editor/GraphEditor/Public/SCommentBubble.h b/Engine/Source/Editor/GraphEditor/Public/SCommentBubble.h
--- a/Engine/Source/Editor/GraphEditor/Public/SCommentBubble.h
+++ b/Engine/Source/Editor/GraphEditor/Public/SCommentBubble.h
@@ -99,13 +99,17 @@ class GRAPHEDITOR_API SCommentBubble : public SCompoundWidget
 	void TickVisibility(const double InCurrentTime, const float InDeltaTime);
 
 	/** Returns the offset from the SNode center slot */
+	UE_DEPRECATED(5.6, "Use the version of the function accepting FVector2f; this Slate API no longer interfaces directly with double-precision scalars and vectors.")
 	FVector2D GetOffset() const;
+	FVector2f GetOffset2f() const;
 
 	/** Returns the offset to the arrow center accounting for zoom on either the comment bubble or the title bar button based on current state */
 	float GetArrowCenterOffset() const;
 
 	/** Returns the bubble size */
+	UE_DEPRECATED(5.6, "Use the version of the function accepting FVector2f; this Slate API no longer interfaces directly with double-precision scalars and vectors.")
 	FVector2D GetSize() const;
+	FVector2f GetSize2f() const;
 
 	/** Returns if comment bubble is visible */
 	bool IsBubbleVisible() const;
diff --git a/Engine/Source/Editor/GraphEditor/Public/SGraphNode.h b/Engine/Source/Editor/GraphEditor/Public/SGraphNode.h
--- a/Engine/Source/Editor/GraphEditor/Public/SGraphNode.h
+++ b/Engine/Source/Editor/GraphEditor/Public/SGraphNode.h
@@ -140,8 +140,8 @@ public:
 	// End of SWidget interface
 
 	// SNodePanel::SNode interface
-	virtual void MoveTo( const FVector2D& NewPosition, FNodeSet& NodeFilter, bool bMarkDirty = true ) override;
-	virtual FVector2D GetPosition() const override;
+	virtual void MoveTo( const FVector2f& NewPosition, FNodeSet& NodeFilter, bool bMarkDirty = true ) override;
+	virtual FVector2f GetPosition2f() const override;
 	virtual FString GetNodeComment() const override;
 	virtual UObject* GetObjectBeingDisplayed() const override;
 	// End of SNodePanel::SNode interface
@@ -227,9 +227,9 @@ public:
 	void PlaySpawnEffect();
 
 	/** Given a coordinate in SGraphNode space, return the same coordinate in graph space while taking zoom and panning of the parent graph into account */
-	FVector2D NodeCoordToGraphCoord( const FVector2D& PanelSpaceCoordinate ) const;
+	UE::Slate::FDeprecateVector2DResult NodeCoordToGraphCoord( const UE::Slate::FDeprecateVector2DParameter& PanelSpaceCoordinate ) const;
 
-	FVector2D GetContentScale() const;
+	UE::Slate::FDeprecateVector2DResult GetContentScale() const;
 	FLinearColor GetColorAndOpacity() const;
 	FLinearColor GetPinLabelColorAndOpacity() const;	
 
diff --git a/Engine/Source/Editor/GraphEditor/Public/SGraphNodeComment.h b/Engine/Source/Editor/GraphEditor/Public/SGraphNodeComment.h
--- a/Engine/Source/Editor/GraphEditor/Public/SGraphNodeComment.h
+++ b/Engine/Source/Editor/GraphEditor/Public/SGraphNodeComment.h
@@ -41,7 +41,7 @@ public:
 
 	//~ Begin SNodePanel::SNode Interface
 	virtual const FSlateBrush* GetShadowBrush(bool bSelected) const override;
-	virtual void GetOverlayBrushes(bool bSelected, const FVector2D WidgetSize, TArray<FOverlayBrushInfo>& Brushes) const override;
+	virtual void GetOverlayBrushes(bool bSelected, const FVector2f& WidgetSize, TArray<FOverlayBrushInfo>& Brushes) const override;
 	virtual bool ShouldAllowCulling() const override { return true; }
 	virtual int32 GetSortDepth() const override;
 	virtual void EndUserInteraction() const override;
@@ -66,10 +66,10 @@ public:
 	void Construct( const FArguments& InArgs, UEdGraphNode_Comment* InNode );
 
 	/** return if the node can be selected, by pointing given location */
-	virtual bool CanBeSelected( const FVector2D& MousePositionInNode ) const override;
+	virtual bool CanBeSelected( const FVector2f& MousePositionInNode ) const override;
 
 	/** return size of the title bar */
-	virtual FVector2D GetDesiredSizeForMarquee() const override;
+	virtual FVector2f GetDesiredSizeForMarquee2f() const override;
 
 	/** return rect of the title bar */
 	virtual FSlateRect GetTitleRect() const override;
@@ -90,12 +90,12 @@ protected:
 	virtual bool IsNodeUnderComment(UEdGraphNode_Comment* InCommentNode, const TSharedRef<SGraphNode> InNodeWidget) const;
 
 	/** called when user is moving the comment node */
-	virtual void MoveTo(const FVector2D& NewPosition, FNodeSet& NodeFilter, bool bMarkDirty = true) override;
+	virtual void MoveTo(const FVector2f& NewPosition, FNodeSet& NodeFilter, bool bMarkDirty = true) override;
 
 	//~ Begin SGraphNodeResizable Interface
 	virtual float GetTitleBarHeight() const override;
 	virtual FSlateRect GetHitTestingBorder() const override;
-	virtual FVector2D GetNodeMaximumSize() const override;
+	virtual FVector2f GetNodeMaximumSize2f() const override;
 	//~ Begin SGraphNodeResizable Interface
 
 	/** @return the color to tint the comment body */
diff --git a/Engine/Source/Editor/GraphEditor/Public/SGraphNodeDocumentation.h b/Engine/Source/Editor/GraphEditor/Public/SGraphNodeDocumentation.h
--- a/Engine/Source/Editor/GraphEditor/Public/SGraphNodeDocumentation.h
+++ b/Engine/Source/Editor/GraphEditor/Public/SGraphNodeDocumentation.h
@@ -37,8 +37,8 @@ public:
 	//~ End SWidget Interface
 
 	//~ Begin SGraphNodeResizable Interface
-	virtual FVector2D GetNodeMinimumSize() const override;
-	virtual FVector2D GetNodeMaximumSize() const override;
+	virtual FVector2f GetNodeMinimumSize2f() const override;
+	virtual FVector2f GetNodeMaximumSize2f() const override;
 	virtual float GetTitleBarHeight() const override;
 	virtual FSlateRect GetHitTestingBorder() const override;
 	//~ End SGraphNodeResizable Interface
@@ -50,7 +50,7 @@ protected:
 	virtual bool IsNameReadOnly () const override { return false; }
 	//~ End SGraphNode Interface
 
-	/** Retrives the current documentation title based on the chosen excerpt */
+	/** Retrieves the current documentation title based on the chosen excerpt */
 	FText GetDocumentationTitle() const;
 	/** Create documentation page from link and excerpt */
 	TSharedPtr<SWidget> CreateDocumentationPage();
@@ -72,7 +72,7 @@ private:
 	/** Node Title Bar */
 	TSharedPtr<SWidget> TitleBar;
 	/** Documentation excerpt maximum page size */
-	FVector2D DocumentationSize;
+	FVector2f DocumentationSize;
 	/** Content Widget with the desired size being managed */
 	TSharedPtr<SWidget> ContentWidget;
 	/** Tracks if child widgets are availble to hit test against */
diff --git a/Engine/Source/Editor/GraphEditor/Public/SGraphNodeKnot.h b/Engine/Source/Editor/GraphEditor/Public/SGraphNodeKnot.h
--- a/Engine/Source/Editor/GraphEditor/Public/SGraphNodeKnot.h
+++ b/Engine/Source/Editor/GraphEditor/Public/SGraphNodeKnot.h
@@ -37,7 +37,7 @@ public:
 
 protected:
 	/** Returns Offset to center comment on the node's only pin */
-	FVector2D GetCommentOffset() const;
+	FVector2f GetCommentOffset() const;
 protected:
 
 	/** Toggles the hovered visibility state */
diff --git a/Engine/Source/Editor/GraphEditor/Public/SGraphNodeResizable.h b/Engine/Source/Editor/GraphEditor/Public/SGraphNodeResizable.h
--- a/Engine/Source/Editor/GraphEditor/Public/SGraphNodeResizable.h
+++ b/Engine/Source/Editor/GraphEditor/Public/SGraphNodeResizable.h
@@ -48,7 +48,9 @@ public:
 protected:
 
 	/** Find the current window zone the mouse is in */
-	virtual SGraphNodeResizable::EResizableWindowZone FindMouseZone(const FVector2D& LocalMouseCoordinates) const;
+	UE_DEPRECATED(5.6, "Use the version of the function accepting FVector2f; this Slate API no longer interfaces directly with double-precision scalars and vectors.")
+	virtual SGraphNodeResizable::EResizableWindowZone FindMouseZone(const FVector2D& LocalMouseCoordinates) const UE_SLATE_DEPRECATED_VECTOR_VIRTUAL_FUNCTION;
+	virtual SGraphNodeResizable::EResizableWindowZone FindMouseZone(const FVector2f& LocalMouseCoordinates) const;
 
 	/** @return true if the current window zone is considered a selection area */
 	bool InSelectionArea() const { return InSelectionArea(MouseZone); }
@@ -60,16 +62,20 @@ protected:
 	void InitNodeAnchorPoint();
 
 	/** Function to fetch the corrected node position based on anchor point*/
-	FVector2D GetCorrectedNodePosition() const;
+	UE::Slate::FDeprecateVector2DResult GetCorrectedNodePosition() const;
 
 	/** Get the current titlebar size */
 	virtual float GetTitleBarHeight() const;
 
 	/** Return smallest desired node size */
-	virtual FVector2D GetNodeMinimumSize() const;
+	UE_DEPRECATED(5.6, "Use the version of the function accepting FVector2f; this Slate API no longer interfaces directly with double-precision scalars and vectors.")
+	virtual FVector2D GetNodeMinimumSize() const UE_SLATE_DEPRECATED_VECTOR_VIRTUAL_FUNCTION;
+	virtual FVector2f GetNodeMinimumSize2f() const;
 
 	/** Return largest desired node size */
-	virtual FVector2D GetNodeMaximumSize() const;
+	UE_DEPRECATED(5.6, "Use the version of the function accepting FVector2f; this Slate API no longer interfaces directly with double-precision scalars and vectors.")
+	virtual FVector2D GetNodeMaximumSize() const UE_SLATE_DEPRECATED_VECTOR_VIRTUAL_FUNCTION;
+	virtual FVector2f GetNodeMaximumSize2f() const;
 
 	//** Return slate rect border for hit testing */
 	virtual FSlateRect GetHitTestingBorder() const;
@@ -77,19 +83,19 @@ protected:
 protected:
 
 	/** The non snapped size of the node for fluid resizing */
-	FVector2D DragSize;
+	FDeprecateSlateVector2D DragSize;
 
 	/** The desired size of the node set during a drag */
-	FVector2D UserSize;
+	FDeprecateSlateVector2D UserSize;
 
 	/** The original size of the node while resizing */
-	FVector2D StoredUserSize;
+	FDeprecateSlateVector2D StoredUserSize;
 
 	/** The resize transaction */
 	TSharedPtr<FScopedTransaction> ResizeTransactionPtr;
 
 	/** Anchor point used to correct node position on resizing the node*/
-	FVector2D NodeAnchorPoint;
+	FDeprecateSlateVector2D NodeAnchorPoint;
 
 	/** The current window zone the mouse is in */
 	EResizableWindowZone MouseZone;
diff --git a/Engine/Source/Editor/GraphEditor/Public/SGraphPanel.h b/Engine/Source/Editor/GraphEditor/Public/SGraphPanel.h
--- a/Engine/Source/Editor/GraphEditor/Public/SGraphPanel.h
+++ b/Engine/Source/Editor/GraphEditor/Public/SGraphPanel.h
@@ -59,7 +59,7 @@ DECLARE_DELEGATE( FOnUpdateGraphPanel )
 struct FGraphContextMenuArguments
 {
 	// The endpoint of the drag or the location of the right-click
-	FVector2D NodeAddPosition;
+	FDeprecateSlateVector2D NodeAddPosition;
 
 	// The source node if there are any
 	UEdGraphNode* GraphNode;
@@ -154,7 +154,7 @@ public:
 	// End of FGCObject interface.
 
 	void ArrangeChildrenForContextMenuSummon(const FGeometry& AllottedGeometry, FArrangedChildren& ArrangedChildren) const;
-	TSharedPtr<SWidget> SummonContextMenu(const FVector2D& WhereToSummon, const FVector2D& WhereToAddNode, UEdGraphNode* ForNode, UEdGraphPin* ForPin, const TArray<UEdGraphPin*>& DragFromPins);
+	TSharedPtr<SWidget> SummonContextMenu(const UE::Slate::FDeprecateVector2DParameter& WhereToSummon, const UE::Slate::FDeprecateVector2DParameter& WhereToAddNode, UEdGraphNode* ForNode, UEdGraphPin* ForPin, const TArray<UEdGraphPin*>& DragFromPins);
 	void SummonCreateNodeMenuFromUICommand(uint32 NumNodesAdded);
 	void DismissContextMenu();
 
@@ -182,10 +182,10 @@ public:
 	bool IsNodeTitleVisible(const class UEdGraphNode* Node, bool bRequestRename);
 
 	/** Use to determine if a rectangle is currently visible */
-	bool IsRectVisible(const FVector2D &TopLeft, const FVector2D &BottomRight);
+	bool IsRectVisible(const UE::Slate::FDeprecateVector2DParameter &TopLeft, const UE::Slate::FDeprecateVector2DParameter &BottomRight);
 
 	/** Focuses the view on rectangle, zooming if neccesary */
-	bool JumpToRect(const FVector2D &BottomLeft, const FVector2D &TopRight);
+	bool JumpToRect(const UE::Slate::FDeprecateVector2DParameter &BottomLeft, const UE::Slate::FDeprecateVector2DParameter &TopRight);
 
 	void JumpToNode(const class UEdGraphNode* JumpToMe, bool bRequestRename, bool bSelectNode);
 
@@ -208,7 +208,9 @@ public:
 	bool IsGraphEditable() const { return IsEditable.Get(); }
 
 	/** Attempt to retrieve the bounds for the specified node */
+	UE_DEPRECATED(5.6, "Use the version of the function accepting FVector2f; this Slate API no longer interfaces directly with double-precision scalars and vectors.")
 	bool GetBoundsForNode(const UObject* InNode, FVector2D& MinCorner, FVector2D& MaxCorner, float Padding = 0.0f) const;
+	bool GetBoundsForNode(const UObject* InNode, FVector2f& MinCorner, FVector2f& MaxCorner, float Padding = 0.0f) const;
 
 	/** Straighten all connections between the selected nodes */
 	void StraightenConnections();
@@ -290,14 +292,14 @@ protected:
 
 	/** Sometimes the panel draws a preview connector; e.g. when the user is connecting pins */
 	TArray< FGraphPinHandle > PreviewConnectorFromPins;
-	FVector2D PreviewConnectorEndpoint;
+	FDeprecateSlateVector2D PreviewConnectorEndpoint;
 	mutable bool bIsDrawStateCached = false;
 
 	/** Last mouse position seen, used for paint-centric highlighting */
-	FVector2D SavedMousePosForOnPaintEventLocalSpace;
+	FDeprecateSlateVector2D SavedMousePosForOnPaintEventLocalSpace;
 	
 	/** The overlap results from the previous OnPaint call */
-	FVector2D PreviousFrameSavedMousePosForSplineOverlap;
+	FDeprecateSlateVector2D PreviousFrameSavedMousePosForSplineOverlap;
 	FGraphSplineOverlapResult PreviousFrameSplineOverlap;
 
 	/** The mouse state from the last mouse move event, used to synthesize pin actions when hovering over a spline on the panel */
@@ -320,7 +322,7 @@ protected:
 	FCurveSequence BounceCurve;
 
 	/** Geometry cache */
-	mutable FVector2D CachedAllottedGeometryScaledSize;
+	mutable FDeprecateSlateVector2D CachedAllottedGeometryScaledSize;
 
 	/** Invoked when text is being committed on panel to verify it */
 	FOnNodeVerifyTextCommit OnVerifyTextCommit;
@@ -369,7 +371,7 @@ private:
 	void OnGraphChanged( const FEdGraphEditAction& InAction );
 
 	/** Update all selected nodes position by provided vector2d */
-	void UpdateSelectedNodesPositions(FVector2D PositionIncrement);
+	void UpdateSelectedNodesPositions(const FVector2f& PositionIncrement);
 
 	/** Handle updating the spline hover state */
 	bool OnSplineHoverStateChanged(const FGraphSplineOverlapResult& NewSplineHoverState);
@@ -387,10 +389,10 @@ private:
 	UEdGraphPin* GetPinUnderMouse(const FGeometry& MyGeometry, const FPointerEvent& MouseEvent, TSharedPtr<SGraphNode> GraphNode) const;
 
 	/** If the spawned nodes were auto-wired from any of the dragged pins, then this will try to make the newly connected pin end up at SpawnGraphPosition */
-	void AdjustNewlySpawnedNodePositions(TArrayView<UEdGraphNode* const> SpawnedNodes, TArrayView<UEdGraphPin*> DraggedFromPins, FVector2D SpawnGraphPosition);
+	void AdjustNewlySpawnedNodePositions(TArrayView<UEdGraphNode* const> SpawnedNodes, TArrayView<UEdGraphPin*> DraggedFromPins, FVector2f SpawnGraphPosition);
 
 	/** Will move a group of nodes by the amount needed for an anchor pin to be at a certain position */
-	void MoveNodesToAnchorPinAtGraphPosition(TArrayView<UEdGraphNode* const> NodesToMove, FGraphPinHandle PinToAnchor, FVector2D DesiredPinGraphPosition);
+	void MoveNodesToAnchorPinAtGraphPosition(TArrayView<UEdGraphNode* const> NodesToMove, FGraphPinHandle PinToAnchor, FVector2f DesiredPinGraphPosition);
 
 	/** Begins a drag operation that disconnects the pin at ClosestPin, but only from the ConnectedPin, rather than all connections */
 	FReply OnMouseDownDisconnectClosestPinFromWire(TSharedRef<class SGraphPin> ClosestPin, TSharedRef<class SGraphPin> ConnectedPin);
diff --git a/Engine/Source/Editor/GraphEditor/Public/SNodePanel.h b/Engine/Source/Editor/GraphEditor/Public/SNodePanel.h
--- a/Engine/Source/Editor/GraphEditor/Public/SNodePanel.h
+++ b/Engine/Source/Editor/GraphEditor/Public/SNodePanel.h
@@ -41,6 +41,7 @@
 #include "Templates/UnrealTemplate.h"
 #include "Types/PaintArgs.h"
 #include "Types/SlateEnums.h"
+#include "Types/SlateVector2.h"
 #include "Types/WidgetMouseEventsDelegate.h"
 #include "UObject/GCObject.h"
 #include "Widgets/DeclarativeSyntaxSupport.h"
@@ -113,9 +114,9 @@ public:
 	/** Brush to draw */
 	const FSlateBrush* Brush;
 	/** Scale of animation to apply */
-	FVector2D AnimationEnvelope;
+	FDeprecateSlateVector2D AnimationEnvelope;
 	/** Offset origin of the overlay from the widget */
-	FVector2D OverlayOffset;
+	FDeprecateSlateVector2D OverlayOffset;
 
 public:
 	FOverlayBrushInfo()
@@ -148,7 +149,7 @@ public:
 	TSharedPtr<SWidget> Widget;
 
 	/** Offset origin of the overlay from the widget */
-	FVector2D OverlayOffset;
+	FDeprecateSlateVector2D OverlayOffset;
 
 public:
 	FOverlayWidgetInfo()
@@ -331,48 +332,21 @@ public:
 				, TAlignmentWidgetSlotMixin<FNodeSlot>(HAlign_Fill, VAlign_Fill)
 				, Zone(InZone)
 				, SlotPadding(0.0f)
-				, Offset(FVector2D::ZeroVector)
+				, Offset(FVector2f::ZeroVector)
 				, AllowScale(true)
 			{ }
 
 			SLATE_SLOT_BEGIN_ARGS_OneMixin(FNodeSlot, TSlotBase<FNodeSlot>, TAlignmentWidgetSlotMixin<FNodeSlot>)
 				SLATE_ATTRIBUTE(FMargin, Padding)
-				SLATE_ATTRIBUTE(FVector2D, SlotOffset)
-				SLATE_ATTRIBUTE(FVector2D, SlotSize)
+				SLATE_ATTRIBUTE_DEPRECATED(FVector2D, SlotOffset, 5.6, "Use SlotOffset2f; this Slate API no longer interfaces directly with double-precision scalars and vectors.")
+				SLATE_ATTRIBUTE(FVector2f, SlotOffset2f)
+				SLATE_ATTRIBUTE_DEPRECATED(FVector2D, SlotSize, 5.6, "UseSlotSize2f; this Slate API no longer interfaces directly with double-precision scalars and vectors.")
+				SLATE_ATTRIBUTE(FVector2f, SlotSize2f)
 				SLATE_ATTRIBUTE(bool, AllowScaling)
 			SLATE_SLOT_END_ARGS()
 
 			void Construct(const FChildren& SlotOwner, FSlotArguments&& InArgs);
 
-		public:
-			UE_DEPRECATED(5.0, "Padding is now deprecated. Use the FSlotArgument or the SetPadding function.")
-			FNodeSlot& Padding( const TAttribute<FMargin> InPadding )
-			{
-				SlotPadding = InPadding;
-				return *this;
-			}
-
-			UE_DEPRECATED(5.0, "SlotOffset is now deprecated. Use the FSlotArgument or the SetSlotOffset function.")
-			FNodeSlot& SlotOffset( const TAttribute<FVector2D> InOffset )
-			{
-				Offset = InOffset;
-				return *this;
-			}
-
-			UE_DEPRECATED(5.0, "SlotSize is now deprecated. Use the FSlotArgument or the SetSlotSize function.")
-			FNodeSlot& SlotSize( const TAttribute<FVector2D> InSize )
-			{
-				Size = InSize;
-				return *this;
-			}
-
-			UE_DEPRECATED(5.0, "AllowScaling is now deprecated. Use the FSlotArgument or the SetAllowScalingfunction.")
-			FNodeSlot& AllowScaling( const TAttribute<bool> InAllowScale )
-			{
-				AllowScale = InAllowScale;
-				return *this;
-			}
-
 		public:
 			ENodeZone::Type GetZoneType() const
 			{
@@ -389,22 +363,44 @@ public:
 				return SlotPadding.Get();
 			}
 
+			UE_DEPRECATED(5.6, "This exists as a temporary bridge for deprecation of SlotSize Attribute. Please use the FVector2f version.")
 			void SetSlotOffset(TAttribute<FVector2D> InOffset)
+			{
+			PRAGMA_DISABLE_DEPRECATION_WARNINGS
+				Offset_Deprecated = MoveTemp(InOffset);
+				Offset.BindLambda([this]()
+				{
+					return UE::Slate::CastToVector2f(Offset_Deprecated.Get());
+				});
+			PRAGMA_ENABLE_DEPRECATION_WARNINGS
+			}
+			void SetSlotOffset(TAttribute<FVector2f> InOffset)
 			{
 				Offset = MoveTemp(InOffset);
 			}
 
-			FVector2D GetSlotOffset() const
+			UE::Slate::FDeprecateVector2DResult GetSlotOffset() const
 			{
 				return Offset.Get();
 			}
 
+			UE_DEPRECATED(5.6, "This exists as a temporary bridge for deprecation of SlotSize Attribute. Please use the FVector2f version.")
 			void SetSlotSize(TAttribute<FVector2D> InSize)
+			{
+			PRAGMA_DISABLE_DEPRECATION_WARNINGS
+				Size_Deprecated = MoveTemp(InSize);
+				Offset.BindLambda([this]()
+				{
+					return UE::Slate::CastToVector2f(Size_Deprecated.Get());
+				});
+			PRAGMA_ENABLE_DEPRECATION_WARNINGS
+			}
+			void SetSlotSize(TAttribute<FVector2f> InSize)
 			{
 				Size = MoveTemp(InSize);
 			}
 
-			FVector2D GetSlotSize() const
+			UE::Slate::FDeprecateVector2DResult GetSlotSize() const
 			{
 				return Size.Get();
 			}
@@ -423,8 +419,12 @@ public:
 			/** The child widget contained in this slot. */
 			ENodeZone::Type Zone;
 			TAttribute<FMargin> SlotPadding;
-			TAttribute<FVector2D> Offset;
-			TAttribute<FVector2D> Size;
+			TAttribute<FVector2f> Offset;
+			TAttribute<FVector2f> Size;
+			UE_DEPRECATED(5.6, "This exists as a temporary bridge for deprecation of SlotSize Attribute.")
+			TAttribute<FVector2D> Size_Deprecated;
+			UE_DEPRECATED(5.6, "This exists as a temporary bridge for deprecation of SlotSize Attribute.")
+			TAttribute<FVector2D> Offset_Deprecated;
 			TAttribute<bool> AllowScale;
 		};
 
@@ -474,7 +474,7 @@ public:
 					const EVisibility ChildVisibility = CenterZone.GetWidget()->GetVisibility();
 					if( ChildVisibility != EVisibility::Collapsed )
 					{
-						return ( CenterZone.GetWidget()->GetDesiredSize() + CenterZone.SlotPadding.Get().GetDesiredSize() ) * DesiredSizeScale.Get();
+						return ( FVector2D(CenterZone.GetWidget()->GetDesiredSize() + CenterZone.SlotPadding.Get().GetDesiredSize()) ) * FVector2D(DesiredSizeScale.Get());
 					}
 				}
 			}
@@ -507,7 +507,7 @@ public:
 				{
 					const FMargin SlotPadding(CurChild.SlotPadding.Get());
 					// If this child is not allowed to scale, its scale relative to its parent should undo the parent widget's scaling.
-					FVector2D Size;
+					FVector2f Size;
 
 					if( CurChild.Size.IsSet() )
 					{
@@ -517,7 +517,7 @@ public:
 					{
 						AlignmentArrangeResult XResult = AlignChild<Orient_Horizontal>(AllottedGeometry.GetLocalSize().X, CurChild, SlotPadding);
 						AlignmentArrangeResult YResult = AlignChild<Orient_Vertical>(AllottedGeometry.GetLocalSize().Y, CurChild, SlotPadding);
-						Size = FVector2D( XResult.Size, YResult.Size );
+						Size = FVector2f( XResult.Size, YResult.Size );
 					}
 					const FArrangedWidget ChildGeom =
 						AllottedGeometry.MakeChild(
@@ -609,14 +609,24 @@ public:
 		* @param NodeFilter		Set of nodes to prevent movement on, after moving successfully a node is added to this set.
 		* @param bMarkDirty		If we should mark nodes as dirty on move
 		*/
-		virtual void MoveTo(const FVector2D& NewPosition, FNodeSet& NodeFilter, bool bMarkDirty = true)
+		UE_DEPRECATED(5.6, "Use the version of the function accepting FVector2f; this Slate API no longer interfaces directly with double-precision scalars and vectors.")
+		virtual void MoveTo(const FVector2D& NewPosition, FNodeSet& NodeFilter, bool bMarkDirty = true) UE_SLATE_DEPRECATED_VECTOR_VIRTUAL_FUNCTION
+		{
+			MoveTo(UE::Slate::CastToVector2f(NewPosition), NodeFilter, bMarkDirty);
+		}
+		virtual void MoveTo(const FVector2f& NewPosition, FNodeSet& NodeFilter, bool bMarkDirty = true)
 		{
 		}
 
 		/** @return the Node's position within the graph */
-		virtual FVector2D GetPosition() const
+		UE_DEPRECATED(5.6, "Use the version of the function accepting FVector2f; this Slate API no longer interfaces directly with double-precision scalars and vectors.")
+		virtual FVector2D GetPosition() const UE_SLATE_DEPRECATED_VECTOR_VIRTUAL_FUNCTION
+		{
+			return FVector2D(GetPosition2f());
+		}
+		virtual FVector2f GetPosition2f() const
 		{
-			return FVector2D(0.0f, 0.0f);
+			return FVector2f::ZeroVector;
 		}
 
 		/** @return a user-specified comment on this node; the comment gets drawn in a bubble above the node */
@@ -654,12 +664,22 @@ public:
 		}
 
 		/** Populate the brushes array with any overlay brushes to render */
-		virtual void GetOverlayBrushes(bool bSelected, const FVector2D WidgetSize, TArray<FOverlayBrushInfo>& Brushes) const
+		UE_DEPRECATED(5.6, "Use the version of the function accepting FVector2f; this Slate API no longer interfaces directly with double-precision scalars and vectors.")
+		virtual void GetOverlayBrushes(bool bSelected, const FVector2D WidgetSize, TArray<FOverlayBrushInfo>& Brushes) const UE_SLATE_DEPRECATED_VECTOR_VIRTUAL_FUNCTION
+		{
+			GetOverlayBrushes(bSelected, UE::Slate::CastToVector2f(WidgetSize), Brushes);
+		}
+		virtual void GetOverlayBrushes(bool bSelected, const FVector2f& WidgetSize, TArray<FOverlayBrushInfo>& Brushes) const
 		{
 		}
 
 		/** Populate the widgets array with any overlay widgets to render */
-		virtual TArray<FOverlayWidgetInfo> GetOverlayWidgets(bool bSelected, const FVector2D& WidgetSize) const
+		UE_DEPRECATED(5.6, "Use the version of the function accepting FVector2f; this Slate API no longer interfaces directly with double-precision scalars and vectors.")
+		virtual TArray<FOverlayWidgetInfo> GetOverlayWidgets(bool bSelected, const FVector2D& WidgetSize) const UE_SLATE_DEPRECATED_VECTOR_VIRTUAL_FUNCTION
+		{
+			return GetOverlayWidgets(bSelected, UE::Slate::CastToVector2f(WidgetSize));
+		}
+		virtual TArray<FOverlayWidgetInfo> GetOverlayWidgets(bool bSelected, const FVector2f& WidgetSize) const
 		{
 			return TArray<FOverlayWidgetInfo>();
 		}
@@ -687,7 +707,12 @@ public:
 		}
 
 		/** return if the node can be selected, by pointing given location */
-		virtual bool CanBeSelected(const FVector2D& MousePositionInNode) const
+		UE_DEPRECATED(5.6, "Use the version of the function accepting FVector2f; this Slate API no longer interfaces directly with double-precision scalars and vectors.")
+		virtual bool CanBeSelected(const FVector2D& MousePositionInNode) const UE_SLATE_DEPRECATED_VECTOR_VIRTUAL_FUNCTION
+		{
+			return CanBeSelected(UE::Slate::CastToVector2f(MousePositionInNode));
+		}
+		virtual bool CanBeSelected(const FVector2f& MousePositionInNode) const
 		{
 			return true;
 		}
@@ -700,9 +725,14 @@ public:
 		 *	e.g. comment node - only title bar is selectable
 		 *	return size of node used for Marquee selecting
 		 */
-		virtual FVector2D GetDesiredSizeForMarquee() const
+		UE_DEPRECATED(5.6, "Use the version of the function accepting FVector2f; this Slate API no longer interfaces directly with double-precision scalars and vectors.")
+		virtual FVector2D GetDesiredSizeForMarquee() const UE_SLATE_DEPRECATED_VECTOR_VIRTUAL_FUNCTION
+		{
+			return FVector2D(GetDesiredSizeForMarquee2f());
+		}
+		virtual FVector2f GetDesiredSizeForMarquee2f() const
 		{
-			return GetDesiredSize();
+			return UE::Slate::CastToVector2f(GetDesiredSize());
 		}
 
 		// Returns node sort depth, defaults to and is generally 0 for most nodes
@@ -723,7 +753,7 @@ public:
 		SNode()
 		: BorderImage( FCoreStyle::Get().GetBrush( "NoBorder" ) )
 		, BorderBackgroundColor( FAppStyle::GetColor("Graph.ForegroundColor"))
-		, DesiredSizeScale(FVector2D(1,1))
+		, DesiredSizeScale(FVector2f::UnitVector)
 		, Children(this)
 		{
 			bHasRelativeLayoutScale = true;
@@ -734,7 +764,7 @@ public:
 		// SBorder Begin
 		TAttribute<const FSlateBrush*> BorderImage;
 		TAttribute<FSlateColor> BorderBackgroundColor;
-		TAttribute<FVector2D> DesiredSizeScale;
+		TAttribute<FDeprecateSlateVector2D> DesiredSizeScale;
 		/** Whether or not to show the disabled effect when this border is disabled */
 		TAttribute<bool> ShowDisabledEffect;
 		/** Mouse event handlers */
@@ -746,7 +776,7 @@ public:
 
 		// SPanel Begin
 		/** The layout scale to apply to this widget's contents; useful for animation. */
-		TAttribute<FVector2D> ContentScale;
+		TAttribute<FDeprecateSlateVector2D> ContentScale;
 		/** The color and opacity to apply to this widget and all its descendants. */
 		TAttribute<FLinearColor> ColorAndOpacity;
 		/** Optional foreground color that will be inherited by all of this widget's contents */
@@ -805,7 +835,7 @@ public:
 	FSlateColor GetZoomTextColorAndOpacity() const;
 
 	/** @return the view offset in graph space */
-	FVector2D GetViewOffset() const;
+	UE::Slate::FDeprecateVector2DResult GetViewOffset() const;
 
 	/** 
 	 * when a panel is scrolling/zooming to a target, this can be called to get it's destination
@@ -813,16 +843,18 @@ public:
 	 * @param BottomRight bottom right corner of the destination
 	 * @return true if there's a scrolling/zooming target and false if there is no destination
 	 */
+	UE_DEPRECATED(5.6, "Use the version of the function accepting FVector2f; this Slate API no longer interfaces directly with double-precision scalars and vectors.")
 	bool GetZoomTargetRect(FVector2D& TopLeft, FVector2D& BottomRight) const;
+	bool GetZoomTargetRect(FVector2f& TopLeft, FVector2f& BottomRight) const;
 
 	/** @return the current view bookmark ID */
 	const FGuid& GetViewBookmarkId() const { return CurrentBookmarkGuid; }
 
 	/** Given a coordinate in panel space (i.e. panel widget space), return the same coordinate in graph space while taking zoom and panning into account */
-	FVector2D PanelCoordToGraphCoord(const FVector2D& PanelSpaceCoordinate) const;
+	UE::Slate::FDeprecateVector2DResult PanelCoordToGraphCoord(const UE::Slate::FDeprecateVector2DParameter& PanelSpaceCoordinate) const;
 
 	/** Restore the graph panel to the supplied view offset/zoom. Also, optionally set the current bookmark ID. */
-	void RestoreViewSettings(const FVector2D& InViewOffset, float InZoomAmount, const FGuid& InBookmarkGuid = FGuid());
+	void RestoreViewSettings(const UE::Slate::FDeprecateVector2DParameter& InViewOffset, float InZoomAmount, const FGuid& InBookmarkGuid = FGuid());
 
 	/** Get the grid snap size */
 	static uint32 GetSnapGridSize();
@@ -838,10 +870,10 @@ public:
 	bool GetBoundsForSelectedNodes(/*out*/ class FSlateRect& Rect, float Padding = 0.0f);
 
 	/** @return the position where where nodes should be pasted (i.e. from the clipboard) */
-	FVector2D GetPastePosition() const;
+	UE::Slate::FDeprecateVector2DResult GetPastePosition() const;
 
 	/** Ask panel to scroll to location */
-	void RequestDeferredPan(const FVector2D& TargetPosition);
+	void RequestDeferredPan(const UE::Slate::FDeprecateVector2DParameter& TargetPosition);
 
 	/** If it is focusing on a particular object */
 	bool HasDeferredObjectFocus() const;
@@ -870,16 +902,16 @@ protected:
 	 * @param TopLeft The top left corner of the target rect
 	 * @param BottomRight The bottom right corner of the target rect
 	 */
-	void ZoomToTarget(const FVector2D& TopLeft, const FVector2D& BottomRight);
+	void ZoomToTarget(const FVector2f& TopLeft, const FVector2f& BottomRight);
 
 	/** Update the new view offset location  */
-	void UpdateViewOffset(const FGeometry& MyGeometry, const FVector2D& TargetPosition);
+	void UpdateViewOffset(const FGeometry& MyGeometry, const FVector2f& TargetPosition);
 
 	/** Compute much panel needs to change to pan to location  */
-	static FVector2D ComputeEdgePanAmount(const FGeometry& MyGeometry, const FVector2D& MouseEvent);
+	static FVector2f ComputeEdgePanAmount(const FGeometry& MyGeometry, const FVector2f& MouseEvent);
 
 	/** Given a coordinate in graph space (e.g. a node's position), return the same coordinate in widget space while taking zoom and panning into account */
-	FVector2D GraphCoordToPanelCoord(const FVector2D& GraphSpaceCoordinate) const;
+	FVector2f GraphCoordToPanelCoord(const UE::Slate::FDeprecateVector2DParameter& GraphSpaceCoordinate) const;
 
 	/** Given a rectangle in panel space, return a rectangle in graph space. */
 	FSlateRect PanelRectToGraphRect(const FSlateRect& PanelSpaceRect) const;
@@ -890,7 +922,8 @@ protected:
 	 * @param InNodeToDrag   The node that the user wants to drag
 	 * @param GrabOffset     Where within the node the user grabbed relative to split between inputs and outputs.
 	 */
-	virtual void OnBeginNodeInteraction(const TSharedRef<SNode>& InNodeToDrag, const FVector2D& GrabOffset);
+	virtual void OnBeginNodeInteraction(const TSharedRef<SNode>& InNodeToDrag, const FVector2D& GrabOffset) UE_SLATE_DEPRECATED_VECTOR_VIRTUAL_FUNCTION;
+	virtual void OnBeginNodeInteraction(const TSharedRef<SNode>& InNodeToDrag, const FVector2f& GrabOffset);
 
 	/** 
 	 * Lets the CanvasPanel know that the user has ended interacting with a node.
@@ -968,31 +1001,35 @@ protected:
 	 * Get the bounds of the given node
 	 * @return True if successful
 	 */
+	UE_DEPRECATED(5.6, "Use the version of the function accepting FVector2f; this Slate API no longer interfaces directly with double-precision scalars and vectors.")
 	bool GetBoundsForNode(const UObject* InNode, /*out*/ FVector2D& MinCorner, /*out*/ FVector2D& MaxCorner, float Padding = 0.0f) const;
+	bool GetBoundsForNode(const UObject* InNode, /*out*/ FVector2f& MinCorner, /*out*/ FVector2f& MaxCorner, float Padding = 0.0f) const;
 
 	/**
 	 * Get the bounds of the selected nodes 
 	 * @param bSelectionSetOnly If true, limits the query to just the selected nodes.  Otherwise it does all nodes.
 	 * @return True if successful
 	 */
-	bool GetBoundsForNodes(bool bSelectionSetOnly, /*out*/ FVector2D& MinCorner, /*out*/ FVector2D& MaxCorner, float Padding = 0.0f);
+	UE_DEPRECATED(5.6, "Use the version of the function accepting FVector2f; this Slate API no longer interfaces directly with double-precision scalars and vectors.")
+	bool GetBoundsForNodes(bool bSelectionSetOnly, /*out*/ FVector2D& MinCorner, /*out*/ FVector2D& MaxCorner, float Padding = 0.0f) const;
+	bool GetBoundsForNodes(bool bSelectionSetOnly, /*out*/ FVector2f& MinCorner, /*out*/ FVector2f& MaxCorner, float Padding = 0.0f) const;
 
 	/**
 	 * Scroll the view to the desired location 
 	 * @return true when the desired location is reached
 	 */
-	bool ScrollToLocation(const FGeometry& MyGeometry, FVector2D DesiredCenterPosition, const float InDeltaTime);
+	bool ScrollToLocation(const FGeometry& MyGeometry, FVector2f DesiredCenterPosition, const float InDeltaTime);
 
 	/**
 	 * Zoom to fit the desired size 
 	 * @return true when zoom fade has completed & fits the desired size
 	 */
-	bool ZoomToLocation(const FVector2D& CurrentSizeWithoutZoom, const FVector2D& DesiredSize, bool bDoneScrolling);
+	bool ZoomToLocation(const FVector2f& CurrentSizeWithoutZoom, const FVector2f& DesiredSize, bool bDoneScrolling);
 
 	/**
 	 * Change zoom level by the specified zoom level delta, about the specified origin.
 	 */
-	void ChangeZoomLevel(int32 ZoomLevelDelta, const FVector2D& WidgetSpaceZoomOrigin, bool bOverrideZoomLimiting);
+	void ChangeZoomLevel(int32 ZoomLevelDelta, const FVector2f& WidgetSpaceZoomOrigin, bool bOverrideZoomLimiting);
 
 	// Should be called whenever the zoom level has changed
 	void PostChangedZoom();
@@ -1024,10 +1061,10 @@ protected:
 	TUniquePtr<FZoomLevelsContainer> ZoomLevels;
 
 	/** The position within the graph at which the user is looking */
-	FVector2D ViewOffset;
+	FDeprecateSlateVector2D ViewOffset;
 
 	/** The position within the graph at which the user was looking last tick */
-	FVector2D OldViewOffset;
+	FDeprecateSlateVector2D OldViewOffset;
 
 	/** How zoomed in/out we are. e.g. 0.25f results in quarter-sized nodes. */
 	int32 ZoomLevel;
@@ -1052,7 +1089,7 @@ protected:
 	TWeakPtr<SNode> NodeUnderMousePtr;
 
 	/** Where in the title the user grabbed to initiate the drag */
-	FVector2D NodeGrabOffset;
+	FDeprecateSlateVector2D NodeGrabOffset;
 
 	/** The total distance that the mouse has been dragged while down */
 	float TotalMouseDelta;
@@ -1061,7 +1098,7 @@ protected:
 	float TotalMouseDeltaXY;
 
 	/** Offset in the panel the user started the LMB+RMB zoom from */
-	FVector2D ZoomStartOffset;
+	FDeprecateSlateVector2D ZoomStartOffset;
 
 	/** Cumulative magnify delta from trackpad gesture */
 	float TotalGestureMagnify;
@@ -1109,16 +1146,16 @@ protected:
 	FCurveSequence ZoomLevelFade;
 
 	/** The position where we should paste when a user executes the paste command. */
-	FVector2D PastePosition;
+	FDeprecateSlateVector2D PastePosition;
 
 	/** Position to pan to */
-	FVector2D DeferredPanPosition;
+	FDeprecateSlateVector2D DeferredPanPosition;
 
 	/** true if pending request for deferred panning */
 	bool bRequestDeferredPan;
 
 	/**	The current position of the software cursor */
-	FVector2D SoftwareCursorPosition;
+	FDeprecateSlateVector2D SoftwareCursorPosition;
 
 	/**	Whether the software cursor should be drawn */
 	bool bShowSoftwareCursor;
@@ -1143,7 +1180,7 @@ protected:
 	bool bVisualUpdatePending;
 	
 	/** Node positions pre-drag, used to limit transaction creation on drag */
-	TMap<TWeakPtr<SNode>, FVector2D> OriginalNodePositions;
+	TMap<TWeakPtr<SNode>, FVector2f> OriginalNodePositions;
 
 	/** Called when the user left clicks on a node without dragging */
 	SGraphEditor::FOnNodeSingleClicked OnNodeSingleClicked;
@@ -1157,8 +1194,8 @@ private:
 	TWeakPtr<FActiveTimerHandle> ActiveTimerHandle;
 
 	/** Zoom target rectangle */
-	FVector2D ZoomTargetTopLeft;
-	FVector2D ZoomTargetBottomRight;
+	FVector2f ZoomTargetTopLeft;
+	FVector2f ZoomTargetBottomRight;
 
 	/** Current view bookmark info */
 	FGuid CurrentBookmarkGuid;
diff --git a/Engine/Source/Editor/WorldBrowser/Private/Tiles/SWorldComposition.cpp b/Engine/Source/Editor/WorldBrowser/Private/Tiles/SWorldComposition.cpp
--- a/Engine/Source/Editor/WorldBrowser/Private/Tiles/SWorldComposition.cpp
+++ b/Engine/Source/Editor/WorldBrowser/Private/Tiles/SWorldComposition.cpp
@@ -214,7 +214,7 @@ public:
 
 			if (ArrangedChildren.Accepts(ChildVisibility))
 			{
-				FVector2D ChildPos = Child->GetPosition();
+				FVector2f ChildPos = Child->GetPosition2f();
 					
 				ArrangedChildren.AddWidget(ChildVisibility,
 					AllottedGeometry.MakeChild(Child,
@@ -264,11 +264,11 @@ public:
 		if (!WorldModel->IsSimulating())
 		{
 			float ScreenSpaceSize = FLevelCollectionModel::EditableAxisLength()*GetZoomAmount()*2.f;
-			FVector2D PaintSize = FVector2D(ScreenSpaceSize, ScreenSpaceSize);
-			FVector2D PaintPosition = GraphCoordToPanelCoord(FVector2D::ZeroVector) - (PaintSize*0.5f);
+			FVector2f PaintSize = FVector2f(ScreenSpaceSize, ScreenSpaceSize);
+			FVector2f PaintPosition = GraphCoordToPanelCoord(FVector2f::ZeroVector) - (PaintSize*0.5f);
 			float Scale = 0.2f; // Scale down drawing border
 			FSlateLayoutTransform LayoutTransform(Scale, AllottedGeometry.GetAccumulatedLayoutTransform().GetTranslation() + PaintPosition);
-			FSlateRenderTransform SlateRenderTransform(Scale, AllottedGeometry.GetAccumulatedRenderTransform().GetTranslation() + PaintPosition);
+			FSlateRenderTransform SlateRenderTransform(Scale, UE::Slate::CastToVector2f(AllottedGeometry.GetAccumulatedRenderTransform().GetTranslation()) + PaintPosition);
 			FPaintGeometry EditableArea(LayoutTransform, SlateRenderTransform, PaintSize/Scale, !SlateRenderTransform.IsIdentity());
 
 			FLinearColor PaintColor = FLinearColor::Yellow;
@@ -307,7 +307,7 @@ public:
 			FRotator ObserverRotation;
 			if (WorldModel->GetObserverView(ObserverPosition, ObserverRotation))
 			{
-				FVector2D ObserverPositionScreen = GraphCoordToPanelCoord(FVector2D(ObserverPosition.X, ObserverPosition.Y));
+				FVector2f ObserverPositionScreen = GraphCoordToPanelCoord(FVector2f(UE_REAL_TO_FLOAT(ObserverPosition.X), UE_REAL_TO_FLOAT(ObserverPosition.Y)));
 				const FSlateBrush* CameraImage = FAppStyle::GetBrush(TEXT("WorldBrowser.SimulationViewPosition"));
 	
 				//AllottedGeometry.GetAccumulatedRenderTransform();
@@ -335,7 +335,7 @@ public:
 			FRotator PlayerRotation;
 			if (WorldModel->GetPlayerView(PlayerPosition, PlayerRotation))
 			{
-				FVector2D PlayerPositionScreen = GraphCoordToPanelCoord(FVector2D(PlayerPosition.X, PlayerPosition.Y));
+				FVector2f PlayerPositionScreen = GraphCoordToPanelCoord(FVector2f(UE_REAL_TO_FLOAT(PlayerPosition.X), UE_REAL_TO_FLOAT(PlayerPosition.Y)));
 				const FSlateBrush* CameraImage = FAppStyle::GetBrush(TEXT("WorldBrowser.SimulationViewPosition"));
 	
 				FPaintGeometry PaintGeometry = AllottedGeometry.ToPaintGeometry(
@@ -415,15 +415,15 @@ public:
 						if ( !bCursorInDeadZone )
 						{
 							// Note, NodeGrabOffset() comes from the node itself, so it's already scaled correctly.
-							FVector2D AnchorNodeNewPos = PanelCoordToGraphCoord( MyGeometry.AbsoluteToLocal( MouseEvent.GetScreenSpacePosition() ) ) - NodeGrabOffset;
+							FVector2f AnchorNodeNewPos = PanelCoordToGraphCoord( MyGeometry.AbsoluteToLocal( MouseEvent.GetScreenSpacePosition() ) ) - NodeGrabOffset;
 
 							// Dragging an unselected node automatically selects it.
 							SelectionManager.StartDraggingNode(NodeBeingDragged->GetObjectBeingDisplayed(), MouseEvent);
 
 							// Move all the selected nodes.
 							{
-								const FVector2D AnchorNodeOldPos = NodeBeingDragged->GetPosition();
-								const FVector2D DeltaPos = AnchorNodeNewPos - AnchorNodeOldPos;
+								const FVector2f AnchorNodeOldPos = NodeBeingDragged->GetPosition2f();
+								const FVector2f DeltaPos = AnchorNodeNewPos - AnchorNodeOldPos;
 								if (DeltaPos.SizeSquared() > FMath::Square(KINDA_SMALL_NUMBER))
 								{
 									MoveSelectedNodes(NodeBeingDragged, AnchorNodeNewPos);
@@ -476,16 +476,16 @@ protected:
 	/**  Draws background for grid view */
 	uint32 PaintBackground(const FGeometry& AllottedGeometry, const FSlateRect& MyCullingRect, FSlateWindowElementList& OutDrawElements, uint32 LayerId) const
 	{
-		FVector2D ScreenWorldOrigin = GraphCoordToPanelCoord(FVector2D(0, 0));
-		FSlateRect ScreenRect(FVector2D(0, 0), AllottedGeometry.GetLocalSize());
+		FVector2f ScreenWorldOrigin = GraphCoordToPanelCoord(FVector2f::ZeroVector);
+		FSlateRect ScreenRect(FVector2f::ZeroVector, AllottedGeometry.GetLocalSize());
 	
 		// World Y-axis
 		if (ScreenWorldOrigin.X > ScreenRect.Left &&
 			ScreenWorldOrigin.X < ScreenRect.Right)
 		{
-			TArray<FVector2D> LinePoints;
-			LinePoints.Add(FVector2D(ScreenWorldOrigin.X, ScreenRect.Top));
-			LinePoints.Add(FVector2D(ScreenWorldOrigin.X, ScreenRect.Bottom));
+			TArray<FVector2f> LinePoints;
+			LinePoints.Add(FVector2f(ScreenWorldOrigin.X, ScreenRect.Top));
+			LinePoints.Add(FVector2f(ScreenWorldOrigin.X, ScreenRect.Bottom));
 
 			FLinearColor YAxisColor = FLinearColor::Green;
 			YAxisColor.A = 0.4f;
@@ -503,9 +503,9 @@ protected:
 		if (ScreenWorldOrigin.Y > ScreenRect.Top &&
 			ScreenWorldOrigin.Y < ScreenRect.Bottom)
 		{
-			TArray<FVector2D> LinePoints;
-			LinePoints.Add(FVector2D(ScreenRect.Left, ScreenWorldOrigin.Y));
-			LinePoints.Add(FVector2D(ScreenRect.Right, ScreenWorldOrigin.Y));
+			TArray<FVector2f> LinePoints;
+			LinePoints.Add(FVector2f(ScreenRect.Left, ScreenWorldOrigin.Y));
+			LinePoints.Add(FVector2f(ScreenRect.Right, ScreenWorldOrigin.Y));
 
 			FLinearColor XAxisColor = FLinearColor::Red;
 			XAxisColor.A = 0.4f;
@@ -526,14 +526,14 @@ protected:
 	uint32 PaintScaleRuler(const FGeometry& AllottedGeometry, const FSlateRect& MyCullingRect, FSlateWindowElementList& OutDrawElements, uint32 LayerId) const
 	{
 		const float	ScaleRulerLength = 100.f; // pixels
-		TArray<FVector2D> LinePoints;
-		LinePoints.Add(FVector2D::ZeroVector);
-		LinePoints.Add(FVector2D::ZeroVector + FVector2D(ScaleRulerLength, 0.f));
+		TArray<FVector2f> LinePoints;
+		LinePoints.Add(FVector2f::ZeroVector);
+		LinePoints.Add(FVector2f::ZeroVector + FVector2f(ScaleRulerLength, 0.f));
 	
 		FSlateDrawElement::MakeLines( 
 			OutDrawElements,
 			LayerId,
-			AllottedGeometry.ToOffsetPaintGeometry(FVector2D(10, 40)),
+			AllottedGeometry.ToOffsetPaintGeometry(FVector2f(10.0f, 40.0f)),
 			LinePoints,
 			ESlateDrawEffect::None,
 			FColor(200, 200, 200));
@@ -555,7 +555,7 @@ protected:
 		FSlateDrawElement::MakeText(
 			OutDrawElements,
 			LayerId,
-			AllottedGeometry.ToOffsetPaintGeometry(FVector2D(10, 27)),
+			AllottedGeometry.ToOffsetPaintGeometry(FVector2f(10.0f, 27.0f)),
 			RulerText,
 			FAppStyle::GetFontStyle("NormalFont"),
 			ESlateDrawEffect::None,
@@ -603,7 +603,7 @@ protected:
 		{
 			// PRESSING ON A NODE!
 			const FArrangedWidget& NodeGeometry = ArrangedChildren[NodeUnderMouseIndex];
-			const FVector2D MousePositionInNode = NodeGeometry.Geometry.AbsoluteToLocal(MouseEvent.GetScreenSpacePosition());
+			const FVector2f MousePositionInNode = NodeGeometry.Geometry.AbsoluteToLocal(MouseEvent.GetScreenSpacePosition());
 			TSharedRef<SNode> NodeWidgetUnderMouse = StaticCastSharedRef<SNode>( NodeGeometry.Widget );
 
 			if (NodeWidgetUnderMouse->CanBeSelected(MousePositionInNode))
@@ -684,7 +684,7 @@ protected:
 	}
 
 	/**  SNodePanel interface */
-	virtual void OnBeginNodeInteraction(const TSharedRef<SNode>& InNodeToDrag, const FVector2D& GrabOffset) override
+	virtual void OnBeginNodeInteraction(const TSharedRef<SNode>& InNodeToDrag, const FVector2f& GrabOffset) override
 	{
 		bHasNodeInteraction = true;
 		SNodePanel::OnBeginNodeInteraction(InNodeToDrag, GrabOffset);
@@ -758,15 +758,15 @@ protected:
 
 		if (SelectionManager.AreAnyNodesSelected())
 		{
-			FVector2D MinCorner, MaxCorner;
+			FVector2f MinCorner, MaxCorner;
 			if (GetBoundsForNodes(true, MinCorner, MaxCorner, 0.f))
 			{
 				FSlateRect SelectionRect = FSlateRect(GraphCoordToPanelCoord(MinCorner), GraphCoordToPanelCoord(MaxCorner));
-				FSlateRect PanelRect = FSlateRect(FVector2D::ZeroVector, CachedGeometry.GetLocalSize());
+				FSlateRect PanelRect = FSlateRect(FVector2f::ZeroVector, CachedGeometry.GetLocalSize());
 				bool bIsVisible = FSlateRect::DoRectanglesIntersect(PanelRect, SelectionRect);
 				if (!bIsVisible)
 				{
-					FVector2D TargetPosition = MaxCorner/2.f + MinCorner/2.f;
+					FVector2f TargetPosition = MaxCorner/2.f + MinCorner/2.f;
 					FitToSelection_Executed();
 				}
 			}
@@ -793,7 +793,7 @@ protected:
 	/**  FitToSelection command handler */
 	void FitToSelection_Executed()
 	{
-		FVector2D MinCorner, MaxCorner;
+		FVector2f MinCorner, MaxCorner;
 		if (GetBoundsForNodes(true, MinCorner, MaxCorner, 0.f))
 		{
 			// Calculate relative zoom padding, then get the padded bounds
@@ -851,7 +851,7 @@ protected:
 	}	
 	
 	/** Moves selected nodes by specified offset */
-	void MoveSelectedNodes(const TSharedPtr<SNode>& InNodeToDrag, FVector2D NewPosition)
+	void MoveSelectedNodes(const TSharedPtr<SNode>& InNodeToDrag, FVector2f NewPosition)
 	{
 		auto ItemDragged = StaticCastSharedPtr<SWorldTileItem>(InNodeToDrag);
 	
@@ -869,10 +869,10 @@ protected:
 				SnappingDistanceWorld = GEditor->GetGridSize();
 			}
 		
-			FVector2D StartPosition = ItemDragged->GetPosition() - ItemDragged->GetLevelModel()->GetLevelTranslationDelta();
-			FVector2D AbsoluteDelta = NewPosition - StartPosition;
+			FVector2f StartPosition = ItemDragged->GetPosition2f() - UE::Slate::CastToVector2f(ItemDragged->GetLevelModel()->GetLevelTranslationDelta());
+			FVector2f AbsoluteDelta = NewPosition - StartPosition;
 
-			WorldModel->UpdateTranslationDelta(WorldModel->GetSelectedLevels(), AbsoluteDelta, bBoundsSnapping, SnappingDistanceWorld);
+			WorldModel->UpdateTranslationDelta(WorldModel->GetSelectedLevels(), FVector2D(AbsoluteDelta), bBoundsSnapping, SnappingDistanceWorld);
 		}
 	}
 
diff --git a/Engine/Source/Editor/WorldBrowser/Private/Tiles/SWorldTileItem.cpp b/Engine/Source/Editor/WorldBrowser/Private/Tiles/SWorldTileItem.cpp
--- a/Engine/Source/Editor/WorldBrowser/Private/Tiles/SWorldTileItem.cpp
+++ b/Engine/Source/Editor/WorldBrowser/Private/Tiles/SWorldTileItem.cpp
@@ -208,9 +208,9 @@ TSharedRef<SToolTip> SWorldTileItem::CreateToolTipWidget()
 }
 END_SLATE_FUNCTION_BUILD_OPTIMIZATION
 
-FVector2D SWorldTileItem::GetPosition() const
+FVector2f SWorldTileItem::GetPosition2f() const
 {
-	return TileModel->GetLevelPosition2D();
+	return UE::Slate::CastToVector2f(TileModel->GetLevelPosition2D());
 }
 
 TSharedPtr<FLevelModel> SWorldTileItem::GetLevelModel() const
@@ -241,8 +241,8 @@ FOptionalSize SWorldTileItem::GetItemHeight() const
 
 FSlateRect SWorldTileItem::GetItemRect() const
 {
-	FVector2D LevelSize = TileModel->GetLevelSize2D();
-	FVector2D LevelPos = GetPosition();
+	FVector2f LevelSize = UE::Slate::CastToVector2f(TileModel->GetLevelSize2D());
+	FVector2f LevelPos = GetPosition2f();
 	return FSlateRect(LevelPos, LevelPos + LevelSize);
 }
 
@@ -257,15 +257,15 @@ TSharedPtr<IToolTip> SWorldTileItem::GetToolTip()
 	return SNodePanel::SNode::GetToolTip();
 }
 
-FVector2D SWorldTileItem::GetDesiredSizeForMarquee() const
+FVector2f SWorldTileItem::GetDesiredSizeForMarquee2f() const
 {
 	// we don't want to select items in non visible layers
 	if (!WorldModel->PassesAllFilters(*TileModel))
 	{
-		return FVector2D::ZeroVector;
+		return FVector2f::ZeroVector;
 	}
 
-	return SNodePanel::SNode::GetDesiredSizeForMarquee();
+	return SNodePanel::SNode::GetDesiredSizeForMarquee2f();
 }
 
 FVector2D SWorldTileItem::ComputeDesiredSize( float ) const
diff --git a/Engine/Source/Editor/WorldBrowser/Private/Tiles/SWorldTileItem.h b/Engine/Source/Editor/WorldBrowser/Private/Tiles/SWorldTileItem.h
--- a/Engine/Source/Editor/WorldBrowser/Private/Tiles/SWorldTileItem.h
+++ b/Engine/Source/Editor/WorldBrowser/Private/Tiles/SWorldTileItem.h
@@ -66,9 +66,9 @@ public:
 	void Construct(const FArguments& InArgs);
 		
 	// SNodePanel::SNode interface start
-	virtual FVector2D GetDesiredSizeForMarquee() const override;
+	virtual FVector2f GetDesiredSizeForMarquee2f() const override;
 	virtual UObject* GetObjectBeingDisplayed() const override;
-	virtual FVector2D GetPosition() const override;
+	virtual FVector2f GetPosition2f() const override;
 	virtual const FSlateBrush* GetShadowBrush(bool bSelected) const override;
 	virtual void GetDiffHighlightBrushes(const FSlateBrush*& BackgroundOut, const FSlateBrush*& ForegroundOut) const override;
 	// SNodePanel::SNode interface end
