[Truncation Warnings] Update SGraphPin API to use FVector2f

diff --git a/Engine/Plugins/Runtime/Metasound/Source/MetasoundEditor/Private/SMetasoundGraphPin.h b/Engine/Plugins/Runtime/Metasound/Source/MetasoundEditor/Private/SMetasoundGraphPin.h
--- a/Engine/Plugins/Runtime/Metasound/Source/MetasoundEditor/Private/SMetasoundGraphPin.h
+++ b/Engine/Plugins/Runtime/Metasound/Source/MetasoundEditor/Private/SMetasoundGraphPin.h
@@ -69,7 +69,7 @@ namespace Metasound
 				const bool bIsHoveringPin,
 				TSharedPtr<SMetasoundPinValueInspector>& OutPinInspector,
 				TWeakPtr<FPinValueInspectorTooltip>& OutInspectorTooltip,
-				TFunctionRef<void(FVector2D&)> InGetTooltipLocation)
+				TFunctionRef<void(FVector2f&)> InGetTooltipLocation)
 			{
 				if (bIsHoveringPin)
 				{
@@ -90,7 +90,7 @@ namespace Metasound
 							TSharedPtr<FPinValueInspectorTooltip> NewTooltip = OutInspectorTooltip.Pin();
 							if (NewTooltip.IsValid())
 							{
-								FVector2D TooltipLocation;
+								FVector2f TooltipLocation;
 								InGetTooltipLocation(TooltipLocation);
 								NewTooltip->MoveTooltip(TooltipLocation);
 							}
@@ -138,7 +138,7 @@ namespace Metasound
 				const bool bIsHoveringPin,
 				TSharedPtr<FMetasoundPinAudioInspector>& OutPinAudioInspector,
 				TWeakPtr<FPinValueInspectorTooltip>& OutInspectorTooltip,
-				TFunctionRef<void(FVector2D&)> InGetTooltipLocation)
+				TFunctionRef<void(FVector2f&)> InGetTooltipLocation)
 			{
 				if (bIsHoveringPin)
 				{
@@ -151,7 +151,7 @@ namespace Metasound
 							TSharedPtr<FPinValueInspectorTooltip> NewTooltip = OutInspectorTooltip.Pin();
 							if (NewTooltip.IsValid())
 							{
-								FVector2D TooltipLocation;
+								FVector2f TooltipLocation;
 								InGetTooltipLocation(TooltipLocation);
 								NewTooltip->MoveTooltip(TooltipLocation);
 							}
@@ -533,7 +533,7 @@ PRAGMA_ENABLE_DEPRECATION_WARNINGS
 					if (bIsHoveringPin || PinInspector.IsValid())
 					{
 						UpdatePinInspector(*GraphPin, bIsHoveringPin, PinInspector, ParentPinType::ValueInspectorTooltip,
-						[this](FVector2D& OutTooltipLocation)
+						[this](FVector2f& OutTooltipLocation)
 						{
 							ParentPinType::GetInteractiveTooltipLocation(OutTooltipLocation);
 						});
@@ -549,7 +549,7 @@ PRAGMA_ENABLE_DEPRECATION_WARNINGS
 								if (Vertex->TypeName == GetMetasoundDataTypeName<FAudioBuffer>())
 								{
 									UpdatePinAudioInspector(*GraphPin, bIsHoveringPin, PinAudioInspector, ParentPinType::ValueInspectorTooltip,
-									[this](FVector2D& OutTooltipLocation)
+									[this](FVector2f& OutTooltipLocation)
 									{
 										ParentPinType::GetInteractiveTooltipLocation(OutTooltipLocation);
 									});
diff --git a/Engine/Source/Editor/GraphEditor/Private/SGraphPin.cpp b/Engine/Source/Editor/GraphEditor/Private/SGraphPin.cpp
--- a/Engine/Source/Editor/GraphEditor/Private/SGraphPin.cpp
+++ b/Engine/Source/Editor/GraphEditor/Private/SGraphPin.cpp
@@ -104,7 +104,7 @@ TSharedPtr<SGraphPin> FGraphPinHandle::FindInGraphPanel(const SGraphPanel& InPan
 SGraphPin::SGraphPin()
 	: GraphPinObj(nullptr)
 	, PinColorModifier(FLinearColor::White)
-	, CachedNodeOffset(FVector2D::ZeroVector)
+	, CachedNodeOffset(FVector2f::ZeroVector)
 	, Custom_Brush_Connected(nullptr)
 	, Custom_Brush_Disconnected(nullptr)
 	, bGraphDataInvalid(false)
@@ -898,7 +898,7 @@ FReply SGraphPin::OnDrop( const FGeometry& MyGeometry, const FDragDropEvent& Dra
 
 void SGraphPin::Tick( const FGeometry& AllottedGeometry, const double InCurrentTime, const float InDeltaTime )
 {
-	CachedNodeOffset = FVector2D(AllottedGeometry.AbsolutePosition)/AllottedGeometry.Scale - OwnerNodePtr.Pin()->GetUnscaledPosition();
+	CachedNodeOffset = AllottedGeometry.AbsolutePosition / AllottedGeometry.Scale - UE::Slate::CastToVector2f(OwnerNodePtr.Pin()->GetUnscaledPosition());
 	CachedNodeOffset.Y += AllottedGeometry.Size.Y * 0.5f;
 
 	if (!ValueInspectorTooltip.IsValid() && IsHovered() && FKismetDebugUtilities::CanInspectPinValue(GetPinObj()))
@@ -908,7 +908,7 @@ void SGraphPin::Tick( const FGeometry& AllottedGeometry, const double InCurrentT
 
 		if (ValueTooltip.IsValid())
 		{
-			FVector2D TooltipLocation;
+			FVector2f TooltipLocation;
 			GetInteractiveTooltipLocation(TooltipLocation);
 			ValueTooltip->MoveTooltip(TooltipLocation);
 		}
@@ -961,7 +961,7 @@ EVisibility SGraphPin::IsPinVisibleAsAdvanced() const
 	return (bIsAdvancedPin && bHideAdvancedPin && bCanBeHidden) ? EVisibility::Collapsed : EVisibility::Visible;
 }
 
-FVector2D SGraphPin::GetNodeOffset() const
+UE::Slate::FDeprecateVector2DResult SGraphPin::GetNodeOffset() const
 {
 	return CachedNodeOffset;
 }
@@ -1398,6 +1398,13 @@ FText SGraphPin::GetTooltipText() const
 }
 
 void SGraphPin::GetInteractiveTooltipLocation(FVector2D& InOutDesiredLocation) const
+{
+	FVector2f TempVector2f;
+	GetInteractiveTooltipLocation(TempVector2f);
+	InOutDesiredLocation = FVector2D(TempVector2f);
+}
+
+void SGraphPin::GetInteractiveTooltipLocation(FVector2f& InOutDesiredLocation) const
 {
 	TSharedPtr<SGraphNode> OwnerNode = OwnerNodePtr.Pin();
 	if (OwnerNode.IsValid())
@@ -1406,7 +1413,7 @@ void SGraphPin::GetInteractiveTooltipLocation(FVector2D& InOutDesiredLocation) c
 		if (GraphPanel.IsValid())
 		{
 			// Reset to the pin's location in graph space.
-			InOutDesiredLocation = FVector2D(OwnerNode->GetPosition2f()) + CachedNodeOffset;
+			InOutDesiredLocation = OwnerNode->GetPosition2f() + CachedNodeOffset;
 			
 			// Shift the desired location to the right edge of the pin's geometry.
 			InOutDesiredLocation.X += GetTickSpaceGeometry().Size.X;
diff --git a/Engine/Source/Editor/GraphEditor/Public/SGraphPin.h b/Engine/Source/Editor/GraphEditor/Public/SGraphPin.h
--- a/Engine/Source/Editor/GraphEditor/Public/SGraphPin.h
+++ b/Engine/Source/Editor/GraphEditor/Public/SGraphPin.h
@@ -232,7 +232,7 @@ public:
 	EVisibility IsPinVisibleAsAdvanced() const;
 
 	/** Gets Node Offset */
-	FVector2D GetNodeOffset() const;
+	UE::Slate::FDeprecateVector2DResult GetNodeOffset() const;
 
 	/** Returns whether or not this pin is currently connectable */
 	bool GetIsConnectable() const;
@@ -321,7 +321,9 @@ protected:
 	FText GetTooltipText() const;
 
 	/** Gets the window location (in screen coords) for an interactive tooltip (e.g. pin value inspector) */
+	UE_DEPRECATED(5.6, "Use the version of the function accepting FVector2f; this Slate API no longer interfaces directly with double-precision scalars and vectors.")
 	void GetInteractiveTooltipLocation(FVector2D& InOutDesiredLocation) const;
+	void GetInteractiveTooltipLocation(FVector2f& InOutDesiredLocation) const;
 
 	TOptional<EMouseCursor::Type> GetPinCursor() const;
 
@@ -371,7 +373,7 @@ protected:
 	FLinearColor PinColorModifier;
 
 	/** Cached offset from owning node to approximate position of culled pins */
-	FVector2D CachedNodeOffset;
+	FDeprecateSlateVector2D CachedNodeOffset;
 
 	/** Set of pins that are currently being hovered */
 	TSet< FEdGraphPinReference > HoverPinSet;
diff --git a/Engine/Source/Editor/Kismet/Private/SPinValueInspector.cpp b/Engine/Source/Editor/Kismet/Private/SPinValueInspector.cpp
--- a/Engine/Source/Editor/Kismet/Private/SPinValueInspector.cpp
+++ b/Engine/Source/Editor/Kismet/Private/SPinValueInspector.cpp
@@ -325,7 +325,7 @@ void FPinValueInspectorTooltip::ShutdownTooltip()
 	ValueInspectorWidget.Reset();
 }
 
-void FPinValueInspectorTooltip::MoveTooltip(const FVector2D& InNewLocation)
+void FPinValueInspectorTooltip::MoveTooltip(const UE::Slate::FDeprecateVector2DParameter& InNewLocation)
 {
 	if (ensureMsgf(TooltipWindow.IsValid(), TEXT("SPinValueInspector::MoveTooltip was called before any tooltip was created")))
 	{
diff --git a/Engine/Source/Editor/Kismet/Public/SPinValueInspector.h b/Engine/Source/Editor/Kismet/Public/SPinValueInspector.h
--- a/Engine/Source/Editor/Kismet/Public/SPinValueInspector.h
+++ b/Engine/Source/Editor/Kismet/Public/SPinValueInspector.h
@@ -85,7 +85,7 @@ class KISMET_API FPinValueInspectorTooltip
 {
 public:
 	/** Moves the tooltip to the new location */
-	void MoveTooltip(const FVector2D& InNewLocation);
+	void MoveTooltip(const UE::Slate::FDeprecateVector2DParameter& InNewLocation);
 
 	/** 
 	 * Dismisses the current tooltip, if it is not currently hovered
