[Graph Editor] PR #11203: Made FGraphNodeFactory::CreatePinWidget() work

diff --git a/Engine/Source/Editor/GraphEditor/Private/SGraphNode.cpp b/Engine/Source/Editor/GraphEditor/Private/SGraphNode.cpp
--- a/Engine/Source/Editor/GraphEditor/Private/SGraphNode.cpp
+++ b/Engine/Source/Editor/GraphEditor/Private/SGraphNode.cpp
@@ -546,9 +546,14 @@ void SGraphNode::SetOwner( const TSharedRef<SGraphPanel>& OwnerPanel )
 	GraphNode->DEPRECATED_NodeWidget = SharedThis(this);
 
 	/*Once we have an owner, and if hide Unused pins is enabled, we need to remake our pins to drop the hidden ones*/
-	if(OwnerGraphPanelPtr.Pin()->GetPinVisibility() != SGraphEditor::Pin_Show 
+	const bool bNeedToHideUnusedPins = OwnerGraphPanelPtr.Pin()->GetPinVisibility() != SGraphEditor::Pin_Show 
 		&& LeftNodeBox.IsValid()
-		&& RightNodeBox.IsValid())
+		&& RightNodeBox.IsValid();
+
+	//If there is node factory, we need to remake pins.
+	const TSharedPtr<FGraphNodeFactory> NodeFactory = OwnerGraphPanelPtr.Pin()->GetNodeFactory();
+
+	if(bNeedToHideUnusedPins || NodeFactory.IsValid())
 	{
 		this->LeftNodeBox->ClearChildren();
 		this->RightNodeBox->ClearChildren();
@@ -1258,6 +1263,15 @@ void SGraphNode::CreatePinWidgets()
 
 TSharedPtr<SGraphPin> SGraphNode::CreatePinWidget(UEdGraphPin* Pin) const
 {
+	if(OwnerGraphPanelPtr.IsValid())
+	{
+		const TSharedPtr<FGraphNodeFactory> NodeFactory = OwnerGraphPanelPtr.Pin()->GetNodeFactory();
+		if(NodeFactory.IsValid())
+		{
+			return NodeFactory->CreatePinWidget(Pin);	
+		}
+	}
+	
 	return FNodeFactory::CreatePinWidget(Pin);
 }
 
diff --git a/Engine/Source/Editor/GraphEditor/Private/SGraphPanel.cpp b/Engine/Source/Editor/GraphEditor/Private/SGraphPanel.cpp
--- a/Engine/Source/Editor/GraphEditor/Private/SGraphPanel.cpp
+++ b/Engine/Source/Editor/GraphEditor/Private/SGraphPanel.cpp
@@ -2672,3 +2672,8 @@ void SGraphPanel::SetNodeFactory(const TSharedRef<class FGraphNodeFactory>& NewN
 {
 	NodeFactory = NewNodeFactory;
 }
+
+TSharedPtr<FGraphNodeFactory> SGraphPanel::GetNodeFactory() const
+{
+	return NodeFactory;
+}
diff --git a/Engine/Source/Editor/GraphEditor/Public/NodeFactory.h b/Engine/Source/Editor/GraphEditor/Public/NodeFactory.h
--- a/Engine/Source/Editor/GraphEditor/Public/NodeFactory.h
+++ b/Engine/Source/Editor/GraphEditor/Public/NodeFactory.h
@@ -40,7 +40,11 @@ public:
 	/** Create a widget for the supplied node */
 	virtual TSharedPtr<SGraphNode> CreateNodeWidget(UEdGraphNode* InNode);
 
-	/** Create a widget for the supplied pin */
+	/**
+	 * Create a widget for the supplied pin
+	 * It would be possible to need to update the pins of nodes that existed before setting the factory.
+	 * If so, use UEdGraph::NotifyGraphChanged() to refresh pins after setting the factory for SGraphEditor or SGraphPanel.
+	 */
 	virtual TSharedPtr<SGraphPin> CreatePinWidget(UEdGraphPin* InPin);
 
 	/** Create a pin connection factory for the supplied schema */
diff --git a/Engine/Source/Editor/GraphEditor/Public/SGraphPanel.h b/Engine/Source/Editor/GraphEditor/Public/SGraphPanel.h
--- a/Engine/Source/Editor/GraphEditor/Public/SGraphPanel.h
+++ b/Engine/Source/Editor/GraphEditor/Public/SGraphPanel.h
@@ -227,6 +227,9 @@ public:
 	/** Sets the current widget factory. */
 	void SetNodeFactory(const TSharedRef<class FGraphNodeFactory>& NewNodeFactory);
 
+	/** Gets the current widget factory. */
+	TSharedPtr<class FGraphNodeFactory> GetNodeFactory() const;
+
 protected:
 
 	void NotifyGraphChanged ( const struct FEdGraphEditAction& InAction);
