- Add missing null check for left / right node box

#rnx
#tests Open/Create a Niagara System from an existing template, double click any of the modues in the stack

diff --git a/Engine/Source/Editor/GraphEditor/Private/SGraphNode.cpp b/Engine/Source/Editor/GraphEditor/Private/SGraphNode.cpp
--- a/Engine/Source/Editor/GraphEditor/Private/SGraphNode.cpp
+++ b/Engine/Source/Editor/GraphEditor/Private/SGraphNode.cpp
@@ -555,8 +555,14 @@ void SGraphNode::SetOwner( const TSharedRef<SGraphPanel>& OwnerPanel )
 
 	if(bNeedToHideUnusedPins || NodeFactory.IsValid())
 	{
-		this->LeftNodeBox->ClearChildren();
-		this->RightNodeBox->ClearChildren();
+		if (this->LeftNodeBox.IsValid())
+		{
+			this->LeftNodeBox->ClearChildren();
+		}
+		if (this->RightNodeBox.IsValid())
+		{
+			this->RightNodeBox->ClearChildren();
+		}
 		CreatePinWidgets();
 	}
 
