Fix regression when dragging off of a pin and releasing on itself, not useful to show context menu here as we frequently misclick around the pin literal edit box

#rnx

diff --git a/Engine/Source/Editor/GraphEditor/Private/DragConnection.cpp b/Engine/Source/Editor/GraphEditor/Private/DragConnection.cpp
--- a/Engine/Source/Editor/GraphEditor/Private/DragConnection.cpp
+++ b/Engine/Source/Editor/GraphEditor/Private/DragConnection.cpp
@@ -401,7 +401,14 @@ FReply FDragConnection::DroppedOnPin(FVector2D ScreenPosition, FVector2D GraphPo
 		ValidateGraphPinList(/*out*/ PinObjects); // this gathers source pins for the drag operation
 		if(UEdGraphPin* CurrentHoveredPin = GetHoveredPin())
 		{
-			PinObjects.Add(CurrentHoveredPin);
+			if(PinObjects.Contains(CurrentHoveredPin))
+			{
+				PinObjects.Reset(); // user dropped a pin on itself, ignore
+			}
+			else
+			{
+				PinObjects.Add(CurrentHoveredPin);
+			}
 		}
 
 		if(PinObjects.Num() > 0 )
