[Editor] Fixed GraphEditorDragDropAction error icon visibility bug

* Visibility was on the image, but should have been on the scale box, causing collapsed images to still take space

#rnx

diff --git a/Engine/Source/Editor/GraphEditor/Private/GraphEditorDragDropAction.cpp b/Engine/Source/Editor/GraphEditor/Private/GraphEditorDragDropAction.cpp
--- a/Engine/Source/Editor/GraphEditor/Private/GraphEditorDragDropAction.cpp
+++ b/Engine/Source/Editor/GraphEditor/Private/GraphEditorDragDropAction.cpp
@@ -162,7 +162,6 @@ void FGraphEditorDragDropAction::SetSimpleFeedbackMessage(const FSlateBrush* Ico
 	TAttribute<EVisibility> IconVisibility = TAttribute<EVisibility>::Create(TAttribute<EVisibility>::FGetter::CreateRaw(this, &FGraphEditorDragDropAction::GetIconVisible));
 	
 	TSharedRef<SWidget> TypeImage = SPinTypeSelector::ConstructPinTypeImage(Icon, IconColor, SecondaryIcon, SecondaryColor, TSharedPtr<SToolTip>());
-	TypeImage->SetVisibility(IconVisibility);
 
 	SetFeedbackMessage(
 		SNew(SHorizontalBox)
@@ -172,9 +171,9 @@ void FGraphEditorDragDropAction::SetSimpleFeedbackMessage(const FSlateBrush* Ico
 		[
 			SNew(SScaleBox)
 			.Stretch(EStretch::ScaleToFit)
+			.Visibility(ErrorIconVisibility)
 			[
 				SNew(SImage)
-				.Visibility(ErrorIconVisibility)
 				.Image( FAppStyle::GetBrush( TEXT("Graph.ConnectorFeedback.Error") ))
 				.ColorAndOpacity( FLinearColor::White )
 			]
@@ -185,6 +184,7 @@ void FGraphEditorDragDropAction::SetSimpleFeedbackMessage(const FSlateBrush* Ico
 		[
 			SNew(SScaleBox)
 			.Stretch(EStretch::ScaleToFit)
+			.Visibility(IconVisibility)
 			[
 				TypeImage
 			]
