PR #9443: Adds a Stack command for Blueprint node alignment + with some general cleanup and icons

diff --git a/Engine/Source/Editor/BlueprintGraph/Private/EdGraphSchema_K2.cpp b/Engine/Source/Editor/BlueprintGraph/Private/EdGraphSchema_K2.cpp
--- a/Engine/Source/Editor/BlueprintGraph/Private/EdGraphSchema_K2.cpp
+++ b/Engine/Source/Editor/BlueprintGraph/Private/EdGraphSchema_K2.cpp
@@ -1747,6 +1747,12 @@ void UEdGraphSchema_K2::GetContextMenuActions(UToolMenu* Menu, UGraphNodeContext
 							InSection.AddMenuEntry(FGraphEditorCommands::Get().DistributeNodesHorizontally);
 							InSection.AddMenuEntry(FGraphEditorCommands::Get().DistributeNodesVertically);
 						}
+
+						{
+							FToolMenuSection& InSection = AlignmentMenu->AddSection("EdGraphSchemaStack", LOCTEXT("StackHeader", "Stack"));
+							InSection.AddMenuEntry(FGraphEditorCommands::Get().StackNodesHorizontally);
+							InSection.AddMenuEntry(FGraphEditorCommands::Get().StackNodesVertically);						
+						}
 					}));
 				}
 			}
diff --git a/Engine/Source/Editor/EditorStyle/Private/StarshipStyle.cpp b/Engine/Source/Editor/EditorStyle/Private/StarshipStyle.cpp
--- a/Engine/Source/Editor/EditorStyle/Private/StarshipStyle.cpp
+++ b/Engine/Source/Editor/EditorStyle/Private/StarshipStyle.cpp
@@ -4540,9 +4540,12 @@ void FStarshipEditorStyle::FStyle::SetupGraphEditorStyles()
 
 		Set( "GraphEditor.DistributeNodesHorizontally", new IMAGE_BRUSH_SVG( "Starship/GraphEditors/DistributeHorizontally", Icon20x20 ) );
 		Set( "GraphEditor.DistributeNodesVertically", new IMAGE_BRUSH_SVG( "Starship/GraphEditors/DistributeVertically", Icon20x20 ) );
+
+		Set( "GraphEditor.StackNodesHorizontally", new IMAGE_BRUSH_SVG("Starship/GraphEditors/StackHorizontally", Icon20x20) );
+		Set( "GraphEditor.StackNodesVertically", new IMAGE_BRUSH_SVG("Starship/GraphEditors/StackVertically", Icon20x20) );
 		
 		Set( "GraphEditor.ToggleHideUnrelatedNodes", new IMAGE_BRUSH_SVG( "Starship/GraphEditors/HideUnrelated", Icon20x20 ) );
-		Set("GraphEditor.Bookmark", new IMAGE_BRUSH_SVG("Starship/Common/Bookmarks", Icon16x16));
+		Set( "GraphEditor.Bookmark", new IMAGE_BRUSH_SVG("Starship/Common/Bookmarks", Icon16x16) );
 
 		// Graph editor widgets
 		{
diff --git a/Engine/Source/Editor/GraphEditor/Private/GraphEditorActions.cpp b/Engine/Source/Editor/GraphEditor/Private/GraphEditorActions.cpp
--- a/Engine/Source/Editor/GraphEditor/Private/GraphEditorActions.cpp
+++ b/Engine/Source/Editor/GraphEditor/Private/GraphEditorActions.cpp
@@ -55,7 +55,7 @@ void FGraphEditorCommandsImpl::RegisterCommands()
 	UI_COMMAND( ExpandNodes, "Expand Node", "Expands the node's internal graph into the current graph and removes this node.", EUserInterfaceActionType::Button, FInputChord() )
 	UI_COMMAND( CollapseSelectionToFunction, "Collapse to Function", "Collapses selected nodes into a single function node.", EUserInterfaceActionType::Button, FInputChord() )
 	UI_COMMAND( CollapseSelectionToMacro, "Collapse to Macro", "Collapses selected nodes into a single macro node.", EUserInterfaceActionType::Button, FInputChord() )
-	UI_COMMAND( ConvertFunctionToEvent, "Convert Function to Event", "Converts selected function to an event and removes the function defintion.", EUserInterfaceActionType::Button, FInputChord() )
+	UI_COMMAND( ConvertFunctionToEvent, "Convert Function to Event", "Converts selected function to an event and removes the function definition.", EUserInterfaceActionType::Button, FInputChord() )
 	UI_COMMAND( ConvertEventToFunction, "Convert Event to Function", "Converts the selected event to a function graph and removes this event node.", EUserInterfaceActionType::Button, FInputChord() )
 
 	UI_COMMAND( AlignNodesTop, "Align Top", "Aligns the top edges of the selected nodes", EUserInterfaceActionType::Button, FInputChord(EModifierKey::Shift, EKeys::W) )
@@ -70,6 +70,9 @@ void FGraphEditorCommandsImpl::RegisterCommands()
 	UI_COMMAND( DistributeNodesHorizontally, "Distribute Horizontally", "Evenly distributes the selected nodes horizontally", EUserInterfaceActionType::Button, FInputChord() )
 	UI_COMMAND( DistributeNodesVertically, "Distribute Vertically", "Evenly distributes the selected nodes vertically", EUserInterfaceActionType::Button, FInputChord() )
 	
+	UI_COMMAND( StackNodesHorizontally, "Stack Horizontally", "Stacks the selected nodes horizontally", EUserInterfaceActionType::Button, FInputChord() )
+	UI_COMMAND( StackNodesVertically, "Stack Vertically", "Stacks the selected nodes vertically", EUserInterfaceActionType::Button, FInputChord() )
+
 	UI_COMMAND( EnableNodes, "Enable Nodes", "Selected node(s) will be enabled.", EUserInterfaceActionType::Check, FInputChord() )
 	UI_COMMAND( DisableNodes, "Disable Nodes", "Selected node(s) will be disabled.", EUserInterfaceActionType::Check, FInputChord() )
 	UI_COMMAND( EnableNodes_Always, "Enable Nodes (Always)", "Selected node(s) will always be enabled.", EUserInterfaceActionType::RadioButton, FInputChord() )
@@ -84,7 +87,7 @@ void FGraphEditorCommandsImpl::RegisterCommands()
 	UI_COMMAND( FindReferencesByClassMemberGlobal, "By Class Member (All)", "Find references of this item by class member in all blueprints", EUserInterfaceActionType::Button, FInputChord() )
 	UI_COMMAND( FindAndReplaceReferences, "Replace References", "Brings up a window to help find and replace all instances of this item", EUserInterfaceActionType::Button, FInputChord() )
 	
-	UI_COMMAND( GoToDefinition, "Goto Definition", "Jumps to the defintion of the selected node if available, e.g., C++ code for a native function or the graph for a Blueprint function.", EUserInterfaceActionType::Button, FInputChord(EKeys::G, EModifierKey::Alt) )
+	UI_COMMAND( GoToDefinition, "Goto Definition", "Jumps to the definition of the selected node if available, e.g., C++ code for a native function or the graph for a Blueprint function.", EUserInterfaceActionType::Button, FInputChord(EKeys::G, EModifierKey::Alt) )
 
 	UI_COMMAND( BreakThisLink, "Break This Link", "Breaks the selected pin's only link.", EUserInterfaceActionType::Button, FInputChord() )
 	UI_COMMAND( BreakPinLinks, "Break All Pin Links", "Breaks all of the selected pin's links", EUserInterfaceActionType::Button, FInputChord() )
diff --git a/Engine/Source/Editor/GraphEditor/Private/SGraphEditorImpl.cpp b/Engine/Source/Editor/GraphEditor/Private/SGraphEditorImpl.cpp
--- a/Engine/Source/Editor/GraphEditor/Private/SGraphEditorImpl.cpp
+++ b/Engine/Source/Editor/GraphEditor/Private/SGraphEditorImpl.cpp
@@ -1816,6 +1816,113 @@ void SGraphEditorImpl::OnDistributeNodesV()
 	}
 }
 
+/** Distribute the specified array of node data evenly and align left when aligning vertically, and align top when aligning horizontally */
+void StackNodes(TArray<FAlignmentData>& InData, bool bIsHorizontal)
+{
+	// Sort the data
+	InData.Sort([](const FAlignmentData& A, const FAlignmentData& B) {
+		return A.TargetProperty < B.TargetProperty;
+	});
+
+	float TargetPosition = InData[0].TargetProperty;
+	float StackAxisPosition = 0.0f;
+	float PaddingAmount = SNodePanel::GetSnapGridSize();
+
+	if (bIsHorizontal)
+	{
+		StackAxisPosition = InData[0].Node->NodePosY;
+		PaddingAmount *= 2.0f;
+	}
+	else
+	{
+		StackAxisPosition = InData[0].Node->NodePosX;
+	}
+
+	// Now set all the properties on the target
+	if (InData.Num() > 1)
+	{
+		UEdGraph* Graph = InData[0].Node->GetGraph();
+		if (Graph)
+		{
+			const UEdGraphSchema* Schema = Graph->GetSchema();
+
+			if (Schema)
+			{
+				for (int32 Index = 0; Index < InData.Num(); ++Index)
+				{
+					FAlignmentData& Entry = InData[Index];
+
+					FVector2D Target2DPosition(Entry.Node->NodePosX, Entry.Node->NodePosY);
+
+					if (bIsHorizontal)
+					{
+						Target2DPosition.X = TargetPosition;
+						Target2DPosition.Y = StackAxisPosition;
+					}
+					else
+					{
+						Target2DPosition.X = StackAxisPosition;
+						Target2DPosition.Y = TargetPosition;
+					}
+
+					Schema->SetNodePosition(Entry.Node, Target2DPosition);
+
+					TargetPosition = Entry.GetTarget() + PaddingAmount;
+				}
+
+				return;
+			}
+		}
+
+		// fall back to the old approach if there isn't a schema
+		for (int32 Index = 0; Index < InData.Num() - 1; ++Index)
+		{
+			FAlignmentData& Entry = InData[Index];
+
+			Entry.Node->Modify();
+			Entry.TargetProperty = TargetPosition;
+
+			TargetPosition = Entry.GetTarget() + PaddingAmount;
+		}
+	}
+}
+
+void SGraphEditorImpl::OnStackNodesH()
+{
+	TArray<FAlignmentData> AlignData;
+	for (UObject* It : GetSelectedNodes())
+	{
+		if (UEdGraphNode* Node = Cast<UEdGraphNode>(It))
+		{
+			AlignData.Add(FAlignmentData(Node, Node->NodePosX, GetNodeSize(*this, Node).X));
+		}
+	}
+
+	if (AlignData.Num() > 1)
+	{
+		const FScopedTransaction Transaction(FGraphEditorCommands::Get().StackNodesHorizontally->GetLabel());
+		StackNodes(AlignData, true);
+	}
+}
+
+void SGraphEditorImpl::OnStackNodesV()
+{
+	TArray<FAlignmentData> AlignData;
+	for (UObject* It : GetSelectedNodes())
+	{
+		if (UEdGraphNode* Node = Cast<UEdGraphNode>(It))
+		{
+			AlignData.Add(FAlignmentData(Node, Node->NodePosY, GetNodeSize(*this, Node).Y));
+		}
+	}
+
+	if (AlignData.Num() > 1)
+	{
+		const FScopedTransaction Transaction(FGraphEditorCommands::Get().StackNodesVertically->GetLabel());
+		StackNodes(AlignData, false);
+	}
+}
+
 int32 SGraphEditorImpl::GetNumberOfSelectedNodes() const
 {
 	return GetSelectedNodes().Num();
diff --git a/Engine/Source/Editor/GraphEditor/Private/SGraphEditorImpl.h b/Engine/Source/Editor/GraphEditor/Private/SGraphEditorImpl.h
--- a/Engine/Source/Editor/GraphEditor/Private/SGraphEditorImpl.h
+++ b/Engine/Source/Editor/GraphEditor/Private/SGraphEditorImpl.h
@@ -354,6 +354,9 @@ public:
 	virtual void OnDistributeNodesH() override;
 	virtual void OnDistributeNodesV() override;
 
+	virtual void OnStackNodesH() override;
+	virtual void OnStackNodesV() override;
+
 	virtual int32 GetNumberOfSelectedNodes() const override;
 
 	virtual UEdGraphNode* GetSingleSelectedNode() const override;
diff --git a/Engine/Source/Editor/GraphEditor/Public/GraphEditorActions.h b/Engine/Source/Editor/GraphEditor/Public/GraphEditorActions.h
--- a/Engine/Source/Editor/GraphEditor/Public/GraphEditorActions.h
+++ b/Engine/Source/Editor/GraphEditor/Public/GraphEditorActions.h
@@ -91,6 +91,9 @@ public:
 
 	TSharedPtr< FUICommandInfo > DistributeNodesHorizontally;
 	TSharedPtr< FUICommandInfo > DistributeNodesVertically;
+
+	TSharedPtr< FUICommandInfo > StackNodesHorizontally;
+	TSharedPtr< FUICommandInfo > StackNodesVertically;
 	
 	// Enable/disable commands
 	TSharedPtr< FUICommandInfo > EnableNodes;
diff --git a/Engine/Source/Editor/Kismet/Private/BlueprintEditor.cpp b/Engine/Source/Editor/Kismet/Private/BlueprintEditor.cpp
--- a/Engine/Source/Editor/Kismet/Private/BlueprintEditor.cpp
+++ b/Engine/Source/Editor/Kismet/Private/BlueprintEditor.cpp
@@ -1567,6 +1567,14 @@ TSharedRef<SGraphEditor> FBlueprintEditor::CreateGraphEditorWidget(TSharedRef<FT
 				FExecuteAction::CreateSP( this, &FBlueprintEditor::OnDistributeNodesV )
 				);
 
+			GraphEditorCommands->MapAction(FGraphEditorCommands::Get().StackNodesHorizontally,
+				FExecuteAction::CreateSP(this, &FBlueprintEditor::OnStackNodesH)
+				);
+
+			GraphEditorCommands->MapAction(FGraphEditorCommands::Get().StackNodesVertically,
+				FExecuteAction::CreateSP(this, &FBlueprintEditor::OnStackNodesV)
+				);
+
 			GraphEditorCommands->MapAction( FGenericCommands::Get().Rename,
 				FExecuteAction::CreateSP( this, &FBlueprintEditor::OnRenameNode ),
 				FCanExecuteAction::CreateSP( this, &FBlueprintEditor::CanRenameNodes )
@@ -6971,6 +6979,24 @@ void FBlueprintEditor::OnDistributeNodesV()
 	}
 }
 
+void FBlueprintEditor::OnStackNodesH()
+{
+	TSharedPtr<SGraphEditor> FocusedGraphEd = FocusedGraphEdPtr.Pin();
+	if (FocusedGraphEd.IsValid())
+	{
+		FocusedGraphEd->OnStackNodesH();
+	}
+}
+
+void FBlueprintEditor::OnStackNodesV()
+{
+	TSharedPtr<SGraphEditor> FocusedGraphEd = FocusedGraphEdPtr.Pin();
+	if (FocusedGraphEd.IsValid())
+	{
+		FocusedGraphEd->OnStackNodesV();
+	}
+}
+
 void FBlueprintEditor::SelectAllNodes()
 {
 	TSharedPtr<SGraphEditor> FocusedGraphEd = FocusedGraphEdPtr.Pin();
diff --git a/Engine/Source/Editor/Kismet/Public/BlueprintEditor.h b/Engine/Source/Editor/Kismet/Public/BlueprintEditor.h
--- a/Engine/Source/Editor/Kismet/Public/BlueprintEditor.h
+++ b/Engine/Source/Editor/Kismet/Public/BlueprintEditor.h
@@ -1074,6 +1074,9 @@ protected:
 	void OnDistributeNodesH();
 	void OnDistributeNodesV();
 
+	void OnStackNodesH();
+	void OnStackNodesV();
+
 	void SelectAllNodes();
 	bool CanSelectAllNodes() const;
 
diff --git a/Engine/Source/Editor/UnrealEd/Public/GraphEditor.h b/Engine/Source/Editor/UnrealEd/Public/GraphEditor.h
--- a/Engine/Source/Editor/UnrealEd/Public/GraphEditor.h
+++ b/Engine/Source/Editor/UnrealEd/Public/GraphEditor.h
@@ -628,6 +628,21 @@ public:
 		}
 	}
 
+	virtual void OnStackNodesH()
+	{
+		if (Implementation.IsValid())
+		{
+			Implementation->OnStackNodesH();
+		}
+	}
+
+	virtual void OnStackNodesV()
+	{
+		if (Implementation.IsValid())
+		{
+			Implementation->OnStackNodesV();
+		}
+	}
 
 	virtual int32 GetNumberOfSelectedNodes() const
 	{
