[Truncation Warnings] Enable truncation warnings in GraphEditor module

diff --git a/Engine/Plugins/Runtime/RigVM/Source/RigVMEditor/Private/EdGraph/RigVMEdGraphConnectionDrawingPolicy.cpp b/Engine/Plugins/Runtime/RigVM/Source/RigVMEditor/Private/EdGraph/RigVMEdGraphConnectionDrawingPolicy.cpp
--- a/Engine/Plugins/Runtime/RigVM/Source/RigVMEditor/Private/EdGraph/RigVMEdGraphConnectionDrawingPolicy.cpp
+++ b/Engine/Plugins/Runtime/RigVM/Source/RigVMEditor/Private/EdGraph/RigVMEdGraphConnectionDrawingPolicy.cpp
@@ -174,9 +174,9 @@ bool FRigVMEdGraphConnectionDrawingPolicy::ShouldChangeTangentForReouteControlPo
 		}
 		else
 		{
-			FVector2D AverageLeftPin;
-			FVector2D AverageRightPin;
-			FVector2D CenterPin;
+			FVector2f AverageLeftPin;
+			FVector2f AverageRightPin;
+			FVector2f CenterPin;
 
 			const TArray<UEdGraphPin*>& Pins = Node->GetAllPins();
 
@@ -211,14 +211,14 @@ bool FRigVMEdGraphConnectionDrawingPolicy::ShouldChangeTangentForReouteControlPo
 }
 
 // Average of the positions of all pins connected to InPin
-bool FRigVMEdGraphConnectionDrawingPolicy::GetAverageConnectedPositionForPin(UEdGraphPin* InPin, FVector2D& OutPos) const
+bool FRigVMEdGraphConnectionDrawingPolicy::GetAverageConnectedPositionForPin(UEdGraphPin* InPin, FVector2f& OutPos) const
 {
-	FVector2D Result = FVector2D::ZeroVector;
+	FVector2f Result = FVector2f::ZeroVector;
 	int32 ResultCount = 0;
 
 	for (UEdGraphPin* LinkedPin : InPin->LinkedTo)
 	{
-		FVector2D CenterPoint;
+		FVector2f CenterPoint;
 		if (FindPinCenter(LinkedPin, /*out*/ CenterPoint))
 		{
 			Result += CenterPoint;
diff --git a/Engine/Plugins/Runtime/RigVM/Source/RigVMEditor/Public/EdGraph/RigVMEdGraphConnectionDrawingPolicy.h b/Engine/Plugins/Runtime/RigVM/Source/RigVMEditor/Public/EdGraph/RigVMEdGraphConnectionDrawingPolicy.h
--- a/Engine/Plugins/Runtime/RigVM/Source/RigVMEditor/Public/EdGraph/RigVMEdGraphConnectionDrawingPolicy.h
+++ b/Engine/Plugins/Runtime/RigVM/Source/RigVMEditor/Public/EdGraph/RigVMEdGraphConnectionDrawingPolicy.h
@@ -41,7 +41,7 @@ private:
 
 	bool ShouldChangeTangentForReouteControlPoint(URigVMEdGraphNode* Node);
 	// Average of the positions of all pins connected to InPin
-	bool GetAverageConnectedPositionForPin(UEdGraphPin* InPin, FVector2D& OutPos) const;
+	bool GetAverageConnectedPositionForPin(UEdGraphPin* InPin, FVector2f& OutPos) const;
 
 	const FSlateBrush* CastImage;
 };
diff --git a/Engine/Source/Editor/GraphEditor/GraphEditor.Build.cs b/Engine/Source/Editor/GraphEditor/GraphEditor.Build.cs
--- a/Engine/Source/Editor/GraphEditor/GraphEditor.Build.cs
+++ b/Engine/Source/Editor/GraphEditor/GraphEditor.Build.cs
@@ -6,7 +6,10 @@ public class GraphEditor : ModuleRules
 {
 	public GraphEditor(ReadOnlyTargetRules Target) : base(Target)
 	{
-        PublicIncludePathModuleNames.AddRange(
+		// Enable truncation warnings in this module
+		CppCompileWarningSettings.UnsafeTypeCastWarningLevel = WarningLevel.Warning;
+
+		PublicIncludePathModuleNames.AddRange(
             new string[] {                
 				"ClassViewer",
 				"StructViewer",
diff --git a/Engine/Source/Editor/GraphEditor/Private/BlueprintConnectionDrawingPolicy.cpp b/Engine/Source/Editor/GraphEditor/Private/BlueprintConnectionDrawingPolicy.cpp
--- a/Engine/Source/Editor/GraphEditor/Private/BlueprintConnectionDrawingPolicy.cpp
+++ b/Engine/Source/Editor/GraphEditor/Private/BlueprintConnectionDrawingPolicy.cpp
@@ -48,7 +48,7 @@ FKismetConnectionDrawingPolicy::FKismetConnectionDrawingPolicy(int32 InBackLayer
 {
 	// Don't want to draw ending arrowheads
 	ArrowImage = nullptr;
-	ArrowRadius = FVector2D::ZeroVector;
+	ArrowRadius = FVector2f::ZeroVector;
 
 	// But we do want to draw midpoint arrowheads
 	if (GetDefault<UBlueprintEditorSettings>()->bDrawMidpointArrowsInBlueprints)
@@ -212,7 +212,7 @@ void FKismetConnectionDrawingPolicy::BuildExecutionRoadmap()
 	{
 		double& ObservationTime = SequentialNodeTimes[i];
 
-		const float PositionRatio = (SequentialNodeTimes.Num() - i) * InvNumNodes;
+		const float PositionRatio = static_cast<float>(SequentialNodeTimes.Num() - i) * InvNumNodes;
 		const float PositionBonus = FMath::Pow(PositionRatio, TracePositionExponent) * TracePositionBonusPeriod;
 		ObservationTime += PositionBonus;
 
@@ -382,7 +382,7 @@ FKismetConnectionDrawingPolicy::FTimePair const* FKismetConnectionDrawingPolicy:
 	return FoundExecPath;
 }
 
-bool FKismetConnectionDrawingPolicy::FindPinCenter(UEdGraphPin* Pin, FVector2D& OutCenter) const
+bool FKismetConnectionDrawingPolicy::FindPinCenter(UEdGraphPin* Pin, FVector2f& OutCenter) const
 {
 	if (const TSharedPtr<SGraphPin>* pPinWidget = PinToPinWidgetMap.Find(Pin))
 	{
@@ -411,15 +411,15 @@ UObject* FKismetConnectionDrawingPolicy::GetObjectBeingDebugged(UBlueprint* Targ
 	return ActiveObject;
 }
 
-bool FKismetConnectionDrawingPolicy::GetAverageConnectedPosition(class UK2Node_Knot* Knot, EEdGraphPinDirection Direction, FVector2D& OutPos) const
+bool FKismetConnectionDrawingPolicy::GetAverageConnectedPosition(class UK2Node_Knot* Knot, EEdGraphPinDirection Direction, FVector2f& OutPos) const
 {
-	FVector2D Result = FVector2D::ZeroVector;
+	FVector2f Result = FVector2f::ZeroVector;
 	int32 ResultCount = 0;
 
 	UEdGraphPin* Pin = (Direction == EGPD_Input) ? Knot->GetInputPin() : Knot->GetOutputPin();
 	for (UEdGraphPin* LinkedPin : Pin->LinkedTo)
 	{
-		FVector2D CenterPoint;
+		FVector2f CenterPoint;
 		if (FindPinCenter(LinkedPin, /*out*/ CenterPoint))
 		{
 			Result += CenterPoint;
@@ -429,7 +429,7 @@ bool FKismetConnectionDrawingPolicy::GetAverageConnectedPosition(class UK2Node_K
 
 	if (ResultCount > 0)
 	{
-		OutPos = Result * (1.0f / ResultCount);
+		OutPos = Result * (1.0f / static_cast<float>(ResultCount));
 		return true;
 	}
 	else
@@ -448,9 +448,9 @@ bool FKismetConnectionDrawingPolicy::ShouldChangeTangentForKnot(UK2Node_Knot* Kn
 	{
 		bool bPinReversed = false;
 
-		FVector2D AverageLeftPin;
-		FVector2D AverageRightPin;
-		FVector2D CenterPin;
+		FVector2f AverageLeftPin;
+		FVector2f AverageRightPin;
+		FVector2f CenterPin;
 		bool bCenterValid = FindPinCenter(Knot->GetOutputPin(), /*out*/ CenterPin);
 		bool bLeftValid = GetAverageConnectedPosition(Knot, EGPD_Input, /*out*/ AverageLeftPin);
 		bool bRightValid = GetAverageConnectedPosition(Knot, EGPD_Output, /*out*/ AverageRightPin);
diff --git a/Engine/Source/Editor/GraphEditor/Private/KismetNodes/SGraphNodeK2Composite.cpp b/Engine/Source/Editor/GraphEditor/Private/KismetNodes/SGraphNodeK2Composite.cpp
--- a/Engine/Source/Editor/GraphEditor/Private/KismetNodes/SGraphNodeK2Composite.cpp
+++ b/Engine/Source/Editor/GraphEditor/Private/KismetNodes/SGraphNodeK2Composite.cpp
@@ -87,7 +87,7 @@ void SGraphNodeK2Composite::UpdateGraphNode()
 		[
 			SNew(SBorder)
 			.BorderImage( FAppStyle::GetBrush( "Graph.CollapsedNode.Body" ) )
-			.Padding(0)
+			.Padding(0.0f)
 			[
 				SNew(SOverlay)
 				+SOverlay::Slot()
@@ -111,7 +111,7 @@ void SGraphNodeK2Composite::UpdateGraphNode()
 						[
 							SNew(SBorder)
 							.BorderImage( FAppStyle::GetBrush("NoBorder") )  // Graph.CollapsedNode.ColorSpill
-							.Padding( FMargin(10,5,30,3) )
+							.Padding( FMargin(10.0f,5.0f,30.0f,3.0f) )
 							[
 								SNew(SVerticalBox)
 								+SVerticalBox::Slot()
@@ -279,7 +279,7 @@ TSharedRef<SWidget> SGraphNodeK2Composite::CreateNodeBody()
 			.BorderImage( FAppStyle::GetBrush("NoBorder") )
 			.HAlign(HAlign_Fill)
 			.VAlign(VAlign_Fill)
-			.Padding( FMargin(0,3) )
+			.Padding( FMargin(0.0f,3.0f) )
 			[
 				SNew(SHorizontalBox)
 				+SHorizontalBox::Slot()
diff --git a/Engine/Source/Editor/GraphEditor/Private/KismetNodes/SGraphNodeK2Copy.cpp b/Engine/Source/Editor/GraphEditor/Private/KismetNodes/SGraphNodeK2Copy.cpp
--- a/Engine/Source/Editor/GraphEditor/Private/KismetNodes/SGraphNodeK2Copy.cpp
+++ b/Engine/Source/Editor/GraphEditor/Private/KismetNodes/SGraphNodeK2Copy.cpp
@@ -91,7 +91,7 @@ public:
 			SNew(SButton)
 			.ButtonStyle(FAppStyle::Get(), NAME_NoBorder)
 			.Visibility(this, &SCopyNodeGraphPin::GetPinStatusIconVisibility)
-			.ContentPadding(0)
+			.ContentPadding(0.0f)
 			.OnClicked(this, &SCopyNodeGraphPin::ClickedOnPinStatusIcon)
 			[
 				SNew(SImage)
@@ -117,7 +117,7 @@ public:
 				+ SHorizontalBox::Slot()
 				.AutoWidth()
 				.VAlign(VAlign_Center)
-				.Padding(0, 0, InArgs._SideToSideMargin, 0)
+				.Padding(0.0f, 0.0f, InArgs._SideToSideMargin, 0.0f)
 				[
 					ActualPinWidget
 				]
@@ -141,7 +141,7 @@ public:
 			+ SHorizontalBox::Slot()
 				.AutoWidth()
 				.VAlign(VAlign_Center)
-				.Padding(InArgs._SideToSideMargin, 0, 0, 0)
+				.Padding(InArgs._SideToSideMargin, 0.0f, 0.0f, 0.0f)
 				[
 					ActualPinWidget
 				];
diff --git a/Engine/Source/Editor/GraphEditor/Private/KismetNodes/SGraphNodeK2CreateDelegate.cpp b/Engine/Source/Editor/GraphEditor/Private/KismetNodes/SGraphNodeK2CreateDelegate.cpp
--- a/Engine/Source/Editor/GraphEditor/Private/KismetNodes/SGraphNodeK2CreateDelegate.cpp
+++ b/Engine/Source/Editor/GraphEditor/Private/KismetNodes/SGraphNodeK2CreateDelegate.cpp
@@ -350,7 +350,7 @@ void SGraphNodeK2CreateDelegate::CreateBelowPinControls(TSharedPtr<SVerticalBox>
 				.OptionsSource(&FunctionOptionList)
 				.OnGenerateWidget(this, &SGraphNodeK2CreateDelegate::MakeFunctionOptionComboWidget)
 				.OnSelectionChanged(this, &SGraphNodeK2CreateDelegate::OnFunctionSelected)
-				.ContentPadding(2)
+				.ContentPadding(2.0f)
 				.MaxListHeight(200.0f)
 				.Content()
 				[
diff --git a/Engine/Source/Editor/GraphEditor/Private/KismetNodes/SGraphNodeSpawnActor.cpp b/Engine/Source/Editor/GraphEditor/Private/KismetNodes/SGraphNodeSpawnActor.cpp
--- a/Engine/Source/Editor/GraphEditor/Private/KismetNodes/SGraphNodeSpawnActor.cpp
+++ b/Engine/Source/Editor/GraphEditor/Private/KismetNodes/SGraphNodeSpawnActor.cpp
@@ -107,15 +107,15 @@ protected:
 
 		return
 			SNew(SBox)
-			.WidthOverride(280)
+			.WidthOverride(280.0f)
 			[
 				SNew(SVerticalBox)
 				+SVerticalBox::Slot()
 				.AutoHeight()
-				.MaxHeight(500)
+				.MaxHeight(500.0f)
 				[ 
 					SNew(SBorder)
-					.Padding(4)
+					.Padding(4.0f)
 					.BorderImage( FAppStyle::GetBrush("ToolPanel.GroupBorder") )
 					[
 						ClassViewerModule.CreateClassViewer(Options, FOnClassPicked::CreateSP(this, &SGraphPinActorBasedBlueprintClass::OnClassPicked))
diff --git a/Engine/Source/Editor/GraphEditor/Private/KismetNodes/SGraphNodeSpawnActorFromClass.cpp b/Engine/Source/Editor/GraphEditor/Private/KismetNodes/SGraphNodeSpawnActorFromClass.cpp
--- a/Engine/Source/Editor/GraphEditor/Private/KismetNodes/SGraphNodeSpawnActorFromClass.cpp
+++ b/Engine/Source/Editor/GraphEditor/Private/KismetNodes/SGraphNodeSpawnActorFromClass.cpp
@@ -107,15 +107,15 @@ protected:
 
 		return 
 			SNew(SBox)
-			.WidthOverride(280)
+			.WidthOverride(280.0f)
 			[
 				SNew(SVerticalBox)
 				+SVerticalBox::Slot()
 				.AutoHeight()
-				.MaxHeight(500)
+				.MaxHeight(500.0f)
 				[
 					SNew(SBorder)
-					.Padding(4)
+					.Padding(4.0f)
 					.BorderImage( FAppStyle::GetBrush("ToolPanel.GroupBorder") )
 					[
 						ClassViewerModule.CreateClassViewer(Options, FOnClassPicked::CreateSP(this, &SGraphPinActorBasedClass::OnPickedNewClass))
diff --git a/Engine/Source/Editor/GraphEditor/Private/KismetPins/SGraphPinClass.cpp b/Engine/Source/Editor/GraphEditor/Private/KismetPins/SGraphPinClass.cpp
--- a/Engine/Source/Editor/GraphEditor/Private/KismetPins/SGraphPinClass.cpp
+++ b/Engine/Source/Editor/GraphEditor/Private/KismetPins/SGraphPinClass.cpp
@@ -191,10 +191,10 @@ TSharedRef<SWidget> SGraphPinClass::GenerateAssetPicker()
 			SNew(SVerticalBox)
 			+ SVerticalBox::Slot()
 			.FillHeight(1.0f)
-			.MaxHeight(500)
+			.MaxHeight(500.0f)
 			[ 
 				SNew(SBorder)
-				.Padding(4)
+				.Padding(4.0f)
 				.BorderImage( FAppStyle::GetBrush("ToolPanel.GroupBorder") )
 				[
 					ClassViewerModule.CreateClassViewer(Options, FOnClassPicked::CreateSP(this, &SGraphPinClass::OnPickedNewClass))
diff --git a/Engine/Source/Editor/GraphEditor/Private/KismetPins/SGraphPinColor.cpp b/Engine/Source/Editor/GraphEditor/Private/KismetPins/SGraphPinColor.cpp
--- a/Engine/Source/Editor/GraphEditor/Private/KismetPins/SGraphPinColor.cpp
+++ b/Engine/Source/Editor/GraphEditor/Private/KismetPins/SGraphPinColor.cpp
@@ -40,7 +40,7 @@ TSharedRef<SWidget>	SGraphPinColor::GetDefaultValueWidget()
 		.BorderImage( FAppStyle::GetBrush("FilledBorder") )
 		.Visibility( this, &SGraphPin::GetDefaultValueVisibility )
 		.IsEnabled(this, &SGraphPin::GetDefaultValueIsEditable)
-		.Padding(1)
+		.Padding(1.0f)
 		[
 			SNew(SHorizontalBox)
 			+SHorizontalBox::Slot()
diff --git a/Engine/Source/Editor/GraphEditor/Private/KismetPins/SGraphPinEnum.cpp b/Engine/Source/Editor/GraphEditor/Private/KismetPins/SGraphPinEnum.cpp
--- a/Engine/Source/Editor/GraphEditor/Private/KismetPins/SGraphPinEnum.cpp
+++ b/Engine/Source/Editor/GraphEditor/Private/KismetPins/SGraphPinEnum.cpp
@@ -52,7 +52,7 @@ void SPinComboBox::Construct( const FArguments& InArgs )
 	this->ChildSlot
 	[
 		SAssignNew( ComboButton, SComboButton )
-		.ContentPadding(3)
+		.ContentPadding(3.0f)
 		.MenuPlacement(MenuPlacement_BelowAnchor)
 		.ButtonContent()
 		[
@@ -69,7 +69,7 @@ void SPinComboBox::Construct( const FArguments& InArgs )
 		[
 			SNew(SBorder)
 			.BorderImage(FAppStyle::Get().GetBrush("Brushes.Recessed"))
-			.Padding(0) // No padding here so that the inner box controls matching width with padding
+			.Padding(0.0f) // No padding here so that the inner box controls matching width with padding
 			[
 				SNew(SBox)
 				.MaxDesiredHeight(450.0f)
diff --git a/Engine/Source/Editor/GraphEditor/Private/KismetPins/SGraphPinInteger.cpp b/Engine/Source/Editor/GraphEditor/Private/KismetPins/SGraphPinInteger.cpp
--- a/Engine/Source/Editor/GraphEditor/Private/KismetPins/SGraphPinInteger.cpp
+++ b/Engine/Source/Editor/GraphEditor/Private/KismetPins/SGraphPinInteger.cpp
@@ -137,7 +137,7 @@ TSharedRef<SWidget>	SGraphPinInteger::GetDefaultValueWidget()
 		};
 
 		return SNew(SComboButton)
-			.ContentPadding(3)
+			.ContentPadding(3.0f)
 			.MenuPlacement(MenuPlacement_BelowAnchor)
 			.Visibility(this, &SGraphPin::GetDefaultValueVisibility)
 			.ButtonContent()
diff --git a/Engine/Source/Editor/GraphEditor/Private/KismetPins/SGraphPinIntegerSlider.cpp b/Engine/Source/Editor/GraphEditor/Private/KismetPins/SGraphPinIntegerSlider.cpp
--- a/Engine/Source/Editor/GraphEditor/Private/KismetPins/SGraphPinIntegerSlider.cpp
+++ b/Engine/Source/Editor/GraphEditor/Private/KismetPins/SGraphPinIntegerSlider.cpp
@@ -137,7 +137,7 @@ TSharedRef<SWidget>	SGraphPinIntegerSlider::GetDefaultValueWidget()
 		};
 
 		return SNew(SComboButton)
-			.ContentPadding(3)
+			.ContentPadding(3.0f)
 			.MenuPlacement(MenuPlacement_BelowAnchor)
 			.Visibility(this, &SGraphPin::GetDefaultValueVisibility)
 			.ButtonContent()
diff --git a/Engine/Source/Editor/GraphEditor/Private/KismetPins/SGraphPinObject.cpp b/Engine/Source/Editor/GraphEditor/Private/KismetPins/SGraphPinObject.cpp
--- a/Engine/Source/Editor/GraphEditor/Private/KismetPins/SGraphPinObject.cpp
+++ b/Engine/Source/Editor/GraphEditor/Private/KismetPins/SGraphPinObject.cpp
@@ -108,13 +108,13 @@ TSharedRef<SWidget>	SGraphPinObject::GetDefaultValueWidget()
 			.Visibility( this, &SGraphPin::GetDefaultValueVisibility )
 			+SHorizontalBox::Slot()
 			.AutoWidth()
-			.Padding(2,0)
+			.Padding(2.0f,0.0f)
 			.MaxWidth(100.0f)
 			[
 				SAssignNew(AssetPickerAnchor, SComboButton)
 				.ButtonStyle( FAppStyle::Get(), "PropertyEditor.AssetComboStyle" )
 				.ForegroundColor( this, &SGraphPinObject::OnGetComboForeground)
-				.ContentPadding( FMargin(2,2,2,1) )
+				.ContentPadding( FMargin(2.0f,2.0f,2.0f,1.0f) )
 				.ButtonColorAndOpacity( this, &SGraphPinObject::OnGetWidgetBackground )
 				.MenuPlacement(MenuPlacement_BelowAnchor)
 				.IsEnabled(this, &SGraphPin::IsEditingEnabled)
@@ -132,7 +132,7 @@ TSharedRef<SWidget>	SGraphPinObject::GetDefaultValueWidget()
 			// Use button
 			+SHorizontalBox::Slot()
 			.AutoWidth()
-			.Padding(1,0)
+			.Padding(1.0f,0.0f)
 			.VAlign(VAlign_Center)
 			[
 				SAssignNew(UseButton, SButton)
@@ -151,14 +151,14 @@ TSharedRef<SWidget>	SGraphPinObject::GetDefaultValueWidget()
 			// Browse button
 			+SHorizontalBox::Slot()
 			.AutoWidth()
-			.Padding(1,0)
+			.Padding(1.0f,0.0f)
 			.VAlign(VAlign_Center)
 			[
 				SAssignNew(BrowseButton, SButton)
 				.ButtonStyle( FAppStyle::Get(), "NoBorder" )
 				.ButtonColorAndOpacity( this, &SGraphPinObject::OnGetWidgetBackground )
 				.OnClicked( GetOnBrowseButtonDelegate() )
-				.ContentPadding(0)
+				.ContentPadding(0.0f)
 				.ToolTipText(NSLOCTEXT("GraphEditor", "ObjectGraphPin_Browse_Tooltip", "Browse"))
 				[
 					SNew(SImage)
@@ -287,8 +287,8 @@ TSharedRef<SWidget> SGraphPinObject::GenerateAssetPicker()
 
 	return
 		SNew(SBox)
-		.HeightOverride(300)
-		.WidthOverride(300)
+		.HeightOverride(300.0f)
+		.WidthOverride(300.0f)
 		[
 			SNew(SBorder)
 			.BorderImage( FAppStyle::GetBrush("Menu.Background") )
diff --git a/Engine/Source/Editor/GraphEditor/Private/KismetPins/SGraphPinString.cpp b/Engine/Source/Editor/GraphEditor/Private/KismetPins/SGraphPinString.cpp
--- a/Engine/Source/Editor/GraphEditor/Private/KismetPins/SGraphPinString.cpp
+++ b/Engine/Source/Editor/GraphEditor/Private/KismetPins/SGraphPinString.cpp
@@ -39,8 +39,8 @@ TSharedRef<SWidget>	SGraphPinString::GetDefaultValueWidget()
 	if (bIsMultiLine)
 	{
 		return SNew(SBox)
-			.MinDesiredWidth(18)
-			.MaxDesiredHeight(200)
+			.MinDesiredWidth(18.0f)
+			.MaxDesiredHeight(200.0f)
 			[
 				SNew(SMultiLineEditableTextBox)
 				.Style(FAppStyle::Get(), "Graph.EditableTextBox")
@@ -50,7 +50,7 @@ TSharedRef<SWidget>	SGraphPinString::GetDefaultValueWidget()
 				.IsReadOnly(this, &SGraphPinString::GetDefaultValueIsReadOnly)
 				.OnTextCommitted(this, &SGraphPinString::SetTypeInValue)
 				.ForegroundColor(FSlateColor::UseForeground())
-				.WrapTextAt(400)
+				.WrapTextAt(400.0f)
 				.ModiferKeyForNewLine(EModifierKey::Shift)
 			];
 	}
@@ -63,8 +63,8 @@ TSharedRef<SWidget>	SGraphPinString::GetDefaultValueWidget()
 		}
 
 		return SNew(SBox)
-			.MinDesiredWidth(18)
-			.MaxDesiredWidth(400)
+			.MinDesiredWidth(18.0f)
+			.MaxDesiredWidth(400.0f)
 			[
 				SNew(SEditableTextBox)
 					.Style(FAppStyle::Get(), "Graph.EditableTextBox")
@@ -187,8 +187,8 @@ TSharedPtr<SWidget> SGraphPinString::TryBuildComboBoxWidget()
 	});
 
 	return SNew(SBox)
-		.MinDesiredWidth(18)
-		.MaxDesiredWidth(400)
+		.MinDesiredWidth(18.0f)
+		.MaxDesiredWidth(400.0f)
 		[
 			SNew(SSearchableComboBox)
 				.OptionsSource(&ComboBoxOptions)
diff --git a/Engine/Source/Editor/GraphEditor/Private/KismetPins/SGraphPinStruct.cpp b/Engine/Source/Editor/GraphEditor/Private/KismetPins/SGraphPinStruct.cpp
--- a/Engine/Source/Editor/GraphEditor/Private/KismetPins/SGraphPinStruct.cpp
+++ b/Engine/Source/Editor/GraphEditor/Private/KismetPins/SGraphPinStruct.cpp
@@ -120,16 +120,16 @@ TSharedRef<SWidget> SGraphPinStruct::GenerateAssetPicker()
 
 	return
 		SNew(SBox)
-		.WidthOverride(280)
+		.WidthOverride(280.0f)
 		[
 			SNew(SVerticalBox)
 
 			+ SVerticalBox::Slot()
 			.FillHeight(1.0f)
-			.MaxHeight(500)
+			.MaxHeight(500.0f)
 			[ 
 				SNew(SBorder)
-				.Padding(4)
+				.Padding(4.0f)
 				.BorderImage( FAppStyle::GetBrush("ToolPanel.GroupBorder") )
 				[
 					StructViewerModule.CreateStructViewer(Options, FOnStructPicked::CreateSP(this, &SGraphPinStruct::OnPickedNewStruct))
diff --git a/Engine/Source/Editor/GraphEditor/Private/KismetPins/SGraphPinStructInstance.cpp b/Engine/Source/Editor/GraphEditor/Private/KismetPins/SGraphPinStructInstance.cpp
--- a/Engine/Source/Editor/GraphEditor/Private/KismetPins/SGraphPinStructInstance.cpp
+++ b/Engine/Source/Editor/GraphEditor/Private/KismetPins/SGraphPinStructInstance.cpp
@@ -162,8 +162,8 @@ TSharedRef<SWidget> SGraphPinStructInstance::GetEditContent()
 	StructureDetailsView->GetOnFinishedChangingPropertiesDelegate().AddSP(this, &SGraphPinStructInstance::PropertyValueChanged);
 
 	return SNew(SBox)
-		.MaxDesiredHeight(350)
-		.MinDesiredWidth(350)
+		.MaxDesiredHeight(350.0f)
+		.MinDesiredWidth(350.0f)
 		[
 			SNew(SBorder)
 			.BorderImage(FAppStyle::GetBrush("ToolPanel.GroupBorder"))
diff --git a/Engine/Source/Editor/GraphEditor/Private/KismetPins/SGraphPinText.cpp b/Engine/Source/Editor/GraphEditor/Private/KismetPins/SGraphPinText.cpp
--- a/Engine/Source/Editor/GraphEditor/Private/KismetPins/SGraphPinText.cpp
+++ b/Engine/Source/Editor/GraphEditor/Private/KismetPins/SGraphPinText.cpp
@@ -113,7 +113,7 @@ TSharedRef<SWidget>	SGraphPinText::GetDefaultValueWidget()
 		.Visibility(this, &SGraphPin::GetDefaultValueVisibility)
 		.ForegroundColor(FSlateColor::UseForeground())
 		.IsEnabled(this, &SGraphPin::GetDefaultValueIsEditable)
-		.WrapTextAt(400)
+		.WrapTextAt(400.0f)
 		.MinDesiredWidth(18.0f)
-		.MaxDesiredHeight(200);
+		.MaxDesiredHeight(200.0f);
 }
diff --git a/Engine/Source/Editor/GraphEditor/Private/MaterialGraphConnectionDrawingPolicy.cpp b/Engine/Source/Editor/GraphEditor/Private/MaterialGraphConnectionDrawingPolicy.cpp
--- a/Engine/Source/Editor/GraphEditor/Private/MaterialGraphConnectionDrawingPolicy.cpp
+++ b/Engine/Source/Editor/GraphEditor/Private/MaterialGraphConnectionDrawingPolicy.cpp
@@ -34,13 +34,13 @@ FMaterialGraphConnectionDrawingPolicy::FMaterialGraphConnectionDrawingPolicy(int
 {
 	// Don't want to draw ending arrowheads
 	ArrowImage = nullptr;
-	ArrowRadius = FVector2D::ZeroVector;
+	ArrowRadius = FVector2f::ZeroVector;
 
 	// Still need to be able to perceive the graph while dragging connectors, esp over comment boxes
 	HoverDeemphasisDarkFraction = 0.4f;
 }
 
-bool FMaterialGraphConnectionDrawingPolicy::FindPinCenter(UEdGraphPin* Pin, FVector2D& OutCenter) const
+bool FMaterialGraphConnectionDrawingPolicy::FindPinCenter(UEdGraphPin* Pin, FVector2f& OutCenter) const
 {
 	if (const TSharedPtr<SGraphPin>* pPinWidget = PinToPinWidgetMap.Find(Pin))
 	{
@@ -54,15 +54,15 @@ bool FMaterialGraphConnectionDrawingPolicy::FindPinCenter(UEdGraphPin* Pin, FVec
 	return false;
 }
 
-bool FMaterialGraphConnectionDrawingPolicy::GetAverageConnectedPosition(class UMaterialGraphNode_Knot* Knot, EEdGraphPinDirection Direction, FVector2D& OutPos) const
+bool FMaterialGraphConnectionDrawingPolicy::GetAverageConnectedPosition(class UMaterialGraphNode_Knot* Knot, EEdGraphPinDirection Direction, FVector2f& OutPos) const
 {
-	FVector2D Result = FVector2D::ZeroVector;
+	FVector2f Result = FVector2f::ZeroVector;
 	int32 ResultCount = 0;
 
 	UEdGraphPin* Pin = (Direction == EGPD_Input) ? Knot->GetInputPin() : Knot->GetOutputPin();
 	for (UEdGraphPin* LinkedPin : Pin->LinkedTo)
 	{
-		FVector2D CenterPoint;
+		FVector2f CenterPoint;
 		if (FindPinCenter(LinkedPin, /*out*/ CenterPoint))
 		{
 			Result += CenterPoint;
@@ -72,7 +72,7 @@ bool FMaterialGraphConnectionDrawingPolicy::GetAverageConnectedPosition(class UM
 
 	if (ResultCount > 0)
 	{
-		OutPos = Result * (1.0f / ResultCount);
+		OutPos = Result * (1.0f / static_cast<float>(ResultCount));
 		return true;
 	}
 	else
@@ -91,9 +91,9 @@ bool FMaterialGraphConnectionDrawingPolicy::ShouldChangeTangentForKnot(UMaterial
 	{
 		bool bPinReversed = false;
 
-		FVector2D AverageLeftPin;
-		FVector2D AverageRightPin;
-		FVector2D CenterPin;
+		FVector2f AverageLeftPin;
+		FVector2f AverageRightPin;
+		FVector2f CenterPin;
 		bool bCenterValid = FindPinCenter(Knot->GetOutputPin(), /*out*/ CenterPin);
 		bool bLeftValid = GetAverageConnectedPosition(Knot, EGPD_Input, /*out*/ AverageLeftPin);
 		bool bRightValid = GetAverageConnectedPosition(Knot, EGPD_Output, /*out*/ AverageRightPin);
diff --git a/Engine/Source/Editor/GraphEditor/Private/MaterialGraphConnectionDrawingPolicy.h b/Engine/Source/Editor/GraphEditor/Private/MaterialGraphConnectionDrawingPolicy.h
--- a/Engine/Source/Editor/GraphEditor/Private/MaterialGraphConnectionDrawingPolicy.h
+++ b/Engine/Source/Editor/GraphEditor/Private/MaterialGraphConnectionDrawingPolicy.h
@@ -29,8 +29,8 @@ protected:
 	TMap<class UMaterialGraphNode_Knot*, bool> KnotToReversedDirectionMap;
 
 	bool ShouldChangeTangentForKnot(class UMaterialGraphNode_Knot* Knot);
-	bool GetAverageConnectedPosition(class UMaterialGraphNode_Knot* Knot, EEdGraphPinDirection Direction, FVector2D& OutPos) const;
-	bool FindPinCenter(UEdGraphPin* Pin, FVector2D& OutCenter) const;
+	bool GetAverageConnectedPosition(class UMaterialGraphNode_Knot* Knot, EEdGraphPinDirection Direction, FVector2f& OutPos) const;
+	bool FindPinCenter(UEdGraphPin* Pin, FVector2f& OutCenter) const;
 
 public:
 	FMaterialGraphConnectionDrawingPolicy(int32 InBackLayerID, int32 InFrontLayerID, float ZoomFactor, const FSlateRect& InClippingRect, FSlateWindowElementList& InDrawElements, UEdGraph* InGraphObj);
diff --git a/Engine/Source/Editor/GraphEditor/Private/MaterialNodes/SGraphNodeMaterialBase.cpp b/Engine/Source/Editor/GraphEditor/Private/MaterialNodes/SGraphNodeMaterialBase.cpp
--- a/Engine/Source/Editor/GraphEditor/Private/MaterialNodes/SGraphNodeMaterialBase.cpp
+++ b/Engine/Source/Editor/GraphEditor/Private/MaterialNodes/SGraphNodeMaterialBase.cpp
@@ -280,7 +280,7 @@ void FPreviewElement::Draw_RenderThread(FRDGBuilder& GraphBuilder, const FDrawPa
 		{
 			// Check realtime mode for whether to pass current time to canvas
 			double CurrentTime = bIsRealtime ? (FApp::GetCurrentTime() - GStartTime) : 0.0;
-			float DeltaTime = bIsRealtime ? FApp::GetDeltaTime() : 0.0f;
+			float DeltaTime = bIsRealtime ? static_cast<float>(FApp::GetDeltaTime()) : 0.0f;
 
 			FCanvas* Canvas = GraphBuilder.AllocObject<FCanvas>(RenderTarget, nullptr, FGameTime::CreateUndilated(CurrentTime, DeltaTime), GMaxRHIFeatureLevel);
 			{
diff --git a/Engine/Source/Editor/GraphEditor/Private/MaterialNodes/SGraphSubstrateMaterial.cpp b/Engine/Source/Editor/GraphEditor/Private/MaterialNodes/SGraphSubstrateMaterial.cpp
--- a/Engine/Source/Editor/GraphEditor/Private/MaterialNodes/SGraphSubstrateMaterial.cpp
+++ b/Engine/Source/Editor/GraphEditor/Private/MaterialNodes/SGraphSubstrateMaterial.cpp
@@ -28,7 +28,7 @@ enum class ESubstrateWidgetOutputType : uint8
 
 static EStyleColor GetSubstrateWidgetColor0() { return EStyleColor::AccentBlue;  }
 static EStyleColor GetSubstrateWidgetColor1() { return EStyleColor::AccentGreen;  }
-FLinearColor FSubstrateWidget::GetConnectionColor() { return FLinearColor(0.16, 0.015, 0.24) * 4.f; }
+FLinearColor FSubstrateWidget::GetConnectionColor() { return FLinearColor(0.16f, 0.015f, 0.24f) * 4.f; }
 
 bool FSubstrateWidget::HasInputSubstrateType(const UEdGraphPin* InPin)
 {
@@ -256,7 +256,7 @@ static const TSharedRef<SWidget> InternalProcessOperator(
 						[
 							SNew(SErrorText)
 							.ErrorText(Message)
-							.BackgroundColor(FSlateColor(FLinearColor(0.65,0.25,0.05)))
+							.BackgroundColor(FSlateColor(FLinearColor(0.65f,0.25f,0.05f)))
 						];
 					}
 
@@ -273,7 +273,7 @@ static const TSharedRef<SWidget> InternalProcessOperator(
 							[
 								SNew(SErrorText)
 								.ErrorText(Message)
-								.BackgroundColor(FSlateColor(FLinearColor(0.16, 0.015, 0.24)))
+								.BackgroundColor(FSlateColor(FLinearColor(0.16f, 0.015f, 0.24f)))
 							];
 						};
 
@@ -317,7 +317,7 @@ static const TSharedRef<SWidget> InternalProcessOperator(
 			}
 
 			const EStyleColor Color = OverrideColor != EStyleColor::MAX ? OverrideColor : (bIsCurrent ? EStyleColor::AccentGreen : EStyleColor::AccentGray);
-			const FSlateColor SlateColor = OutputType == ESubstrateWidgetOutputType::Node ? FSlateColor(Color) : FSlateColor(FLinearColor(0.16, 0.015, 0.24));
+			const FSlateColor SlateColor = OutputType == ESubstrateWidgetOutputType::Node ? FSlateColor(Color) : FSlateColor(FLinearColor(0.16f, 0.015f, 0.24f));
 			auto BSDF = SNew(SErrorText)
 				.ErrorText(FText::FromString(BSDFDesc))
 				.BackgroundColor(SlateColor)
diff --git a/Engine/Source/Editor/GraphEditor/Private/SGraphEditorActionMenu.cpp b/Engine/Source/Editor/GraphEditor/Private/SGraphEditorActionMenu.cpp
--- a/Engine/Source/Editor/GraphEditor/Private/SGraphEditorActionMenu.cpp
+++ b/Engine/Source/Editor/GraphEditor/Private/SGraphEditorActionMenu.cpp
@@ -31,12 +31,12 @@ void SGraphEditorActionMenu::Construct( const FArguments& InArgs )
 	// Build the widget layout
 	SBorder::Construct( SBorder::FArguments()
 		.BorderImage( FAppStyle::GetBrush("Menu.Background") )
-		.Padding(5)
+		.Padding(5.0f)
 		[
 			// Achieving fixed width by nesting items within a fixed width box.
 			SNew(SBox)
-			.WidthOverride(400)
-			.HeightOverride(400)
+			.WidthOverride(400.0f)
+			.HeightOverride(400.0f)
 			[
 				SAssignNew(GraphActionMenu, SGraphActionMenu)
 				.OnActionSelected(this, &SGraphEditorActionMenu::OnActionSelected)
diff --git a/Engine/Source/Editor/GraphEditor/Private/SGraphEditorImpl.cpp b/Engine/Source/Editor/GraphEditor/Private/SGraphEditorImpl.cpp
--- a/Engine/Source/Editor/GraphEditor/Private/SGraphEditorImpl.cpp
+++ b/Engine/Source/Editor/GraphEditor/Private/SGraphEditorImpl.cpp
@@ -73,15 +73,15 @@ struct FSlateBrush;
 
 #define LOCTEXT_NAMESPACE "GraphEditorModule"
 
-FVector2D GetNodeSize(const SGraphEditor& GraphEditor, const UEdGraphNode* Node)
+FVector2f GetNodeSize(const SGraphEditor& GraphEditor, const UEdGraphNode* Node)
 {
 	FSlateRect Rect;
 	if (GraphEditor.GetBoundsForNode(Node, Rect, 0.f))
 	{
-		return FVector2D(Rect.Right - Rect.Left, Rect.Bottom - Rect.Top);
+		return Rect.GetSize();
 	}
 
-	return FVector2D(Node->NodeWidth, Node->NodeHeight);
+	return Node->GetSize();
 }
 
 /////////////////////////////////////////////////////
@@ -860,7 +860,7 @@ void SGraphEditorImpl::Construct( const FArguments& InArgs )
 		]
 	];
 
-	GraphPanel->RestoreViewSettings(FVector2D::ZeroVector, -1);
+	GraphPanel->RestoreViewSettings(FVector2f::ZeroVector, -1.0f);
 
 	NotifyGraphChanged();
 }
@@ -1717,7 +1717,7 @@ void DistributeNodes(TArray<FAlignmentData>& InData, bool bIsHorizontal)
 {
 	// Sort the data
 	InData.Sort([](const FAlignmentData& A, const FAlignmentData& B) {
-		return A.TargetProperty + A.TargetOffset / 2 < B.TargetProperty + B.TargetOffset / 2;
+		return A.TargetProperty + FMath::RoundToInt(A.TargetOffset / 2) < B.TargetProperty + FMath::RoundToInt(B.TargetOffset / 2);
 	});
 
 	// Measure the available space
@@ -1727,8 +1727,8 @@ void DistributeNodes(TArray<FAlignmentData>& InData, bool bIsHorizontal)
 		TotalWidthOfNodes += InData[Index].TargetOffset;
 	}
 
-	const float SpaceToDistributeIn = InData.Last().TargetProperty - InData[0].GetTarget();
-	const float PaddingAmount = ((SpaceToDistributeIn - TotalWidthOfNodes) / (InData.Num() - 1));
+	const float SpaceToDistributeIn = static_cast<float>(InData.Last().TargetProperty) - InData[0].GetTarget();
+	const float PaddingAmount = ((SpaceToDistributeIn - TotalWidthOfNodes) / static_cast<float>(InData.Num() - 1));
 
 	float TargetPosition = InData[0].GetTarget() + PaddingAmount;
 
@@ -1747,7 +1747,7 @@ void DistributeNodes(TArray<FAlignmentData>& InData, bool bIsHorizontal)
 				{
 					FAlignmentData& Entry = InData[Index];
 
-					FVector2f Target2DPosition(Entry.Node->NodePosX, Entry.Node->NodePosY);
+					FVector2f Target2DPosition = Entry.Node->GetPosition();
  
 					if (bIsHorizontal)
 					{
@@ -1773,7 +1773,7 @@ void DistributeNodes(TArray<FAlignmentData>& InData, bool bIsHorizontal)
 			FAlignmentData& Entry = InData[Index];
 
 			Entry.Node->Modify();
-			Entry.TargetProperty = TargetPosition;
+			Entry.TargetProperty = UE::LWC::FloatToIntCastChecked<int32>(TargetPosition);
 
 			TargetPosition = Entry.GetTarget() + PaddingAmount;
 		}
@@ -1824,18 +1824,18 @@ void StackNodes(TArray<FAlignmentData>& InData, bool bIsHorizontal)
 		return A.TargetProperty < B.TargetProperty;
 	});
 
-	float TargetPosition = InData[0].TargetProperty;
+	float TargetPosition = static_cast<float>(InData[0].TargetProperty);
 	float StackAxisPosition = 0.0f;
-	float PaddingAmount = SNodePanel::GetSnapGridSize();
+	float PaddingAmount = static_cast<float>(SNodePanel::GetSnapGridSize());
 
 	if (bIsHorizontal)
 	{
-		StackAxisPosition = InData[0].Node->NodePosY;
+		StackAxisPosition = InData[0].Node->GetNodePosY();
 		PaddingAmount *= 2.0f;
 	}
 	else
 	{
-		StackAxisPosition = InData[0].Node->NodePosX;
+		StackAxisPosition = InData[0].Node->GetNodePosX();
 	}
 
 	// Now set all the properties on the target
@@ -1852,7 +1852,7 @@ void StackNodes(TArray<FAlignmentData>& InData, bool bIsHorizontal)
 				{
 					FAlignmentData& Entry = InData[Index];
 
-					FVector2f AdjustedTargetPosition(Entry.Node->NodePosX, Entry.Node->NodePosY);
+					FVector2f AdjustedTargetPosition = Entry.Node->GetPosition();
 
 					if (bIsHorizontal)
 					{
@@ -1880,7 +1880,7 @@ void StackNodes(TArray<FAlignmentData>& InData, bool bIsHorizontal)
 			FAlignmentData& Entry = InData[Index];
 
 			Entry.Node->Modify();
-			Entry.TargetProperty = TargetPosition;
+			Entry.TargetProperty = UE::LWC::FloatToIntCastChecked<int32>(TargetPosition);
 
 			TargetPosition = Entry.GetTarget() + PaddingAmount;
 		}
diff --git a/Engine/Source/Editor/GraphEditor/Private/SGraphEditorImpl.h b/Engine/Source/Editor/GraphEditor/Private/SGraphEditorImpl.h
--- a/Engine/Source/Editor/GraphEditor/Private/SGraphEditorImpl.h
+++ b/Engine/Source/Editor/GraphEditor/Private/SGraphEditorImpl.h
@@ -122,8 +122,8 @@ struct FAlignmentHelper
 					
 					for (FAlignmentData& Entry : AlignmentData)
 					{
-						int32 TargetProperty = Target - Entry.TargetOffset; 
-						FVector2f TargetPosition(Entry.Node->NodePosX, Entry.Node->NodePosY);
+						float TargetProperty = Target - Entry.TargetOffset; 
+						FVector2f TargetPosition = Entry.Node->GetPosition();
 
 						if (Orientation == EOrientation::Orient_Horizontal)
 						{
@@ -180,7 +180,7 @@ private:
 			{
 				SumTotal += Entry.GetTarget();
 			}
-			return SumTotal / AlignmentData.Num();
+			return SumTotal / static_cast<float>(AlignmentData.Num());
 		}
 	}
 
diff --git a/Engine/Source/Editor/GraphEditor/Private/SGraphNode.cpp b/Engine/Source/Editor/GraphEditor/Private/SGraphNode.cpp
--- a/Engine/Source/Editor/GraphEditor/Private/SGraphNode.cpp
+++ b/Engine/Source/Editor/GraphEditor/Private/SGraphNode.cpp
@@ -83,7 +83,7 @@ void SNodeTitle::Construct(const FArguments& InArgs, UEdGraphNode* InNode)
 
 	StyleSet = InArgs._StyleSet;
 
-	CachedSize = FVector2D::ZeroVector;
+	CachedSize = FVector2f::ZeroVector;
 
 	// If the user set the text, use it, otherwise use the node title by default
 	if(InArgs._Text.IsSet())
@@ -129,7 +129,7 @@ FText SNodeTitle::GetHeadTitle() const
 	return (GraphNode.IsValid() && GraphNode->GetCanRenameNode()) ? GraphNode->GetNodeTitle(ENodeTitleType::EditableTitle) : CachedHeadTitle;
 }
 
-FVector2D SNodeTitle::GetTitleSize() const
+UE::Slate::FDeprecateVector2DResult SNodeTitle::GetTitleSize() const
 {
 	return CachedSize;
 }
@@ -494,7 +494,7 @@ void SGraphNode::OnToolTipClosing()
 
 void SGraphNode::Tick( const FGeometry& AllottedGeometry, const double InCurrentTime, const float InDeltaTime )
 {
-	CachedUnscaledPosition = FVector2D(AllottedGeometry.AbsolutePosition/AllottedGeometry.Scale);
+	CachedUnscaledPosition = AllottedGeometry.AbsolutePosition/AllottedGeometry.Scale;
 
 	SNodePanel::SNode::Tick(AllottedGeometry, InCurrentTime, InDeltaTime);
 
@@ -585,8 +585,7 @@ void SGraphNode::MoveTo( const FVector2f& NewPosition, FNodeSet& NodeFilter, boo
 		{
 			NodeFilter.Add( SharedThis( this ) );
 			GraphNode->Modify(bMarkDirty);
-			GraphNode->NodePosX = NewPosition.X;
-			GraphNode->NodePosY = NewPosition.Y;
+			GraphNode->SetPosition(NewPosition);
 		}
 	}
 }
@@ -594,7 +593,7 @@ void SGraphNode::MoveTo( const FVector2f& NewPosition, FNodeSet& NodeFilter, boo
 /** @return the Node's position within the graph */
 FVector2f SGraphNode::GetPosition2f() const
 {
-	return FVector2f( GraphNode->NodePosX, GraphNode->NodePosY );
+	return GraphNode->GetPosition();
 }
 
 FString SGraphNode::GetEditableNodeTitle() const
@@ -1410,7 +1409,7 @@ SGraphNode::SGraphNode()
 	, bProvidedComplexTooltip(false)
 	, bRenameIsPending( false )
 	, ErrorColor( FLinearColor::White )
-	, CachedUnscaledPosition( FVector2D::ZeroVector )
+	, CachedUnscaledPosition( FVector2f::ZeroVector )
 	, Settings( GetDefault<UGraphEditorSettings>() )
 {
 	// Set up animation
@@ -1438,47 +1437,46 @@ void SGraphNode::PositionThisNodeBetweenOtherNodes(const TMap< UObject*, TShared
 void SGraphNode::PositionThisNodeBetweenOtherNodes(const TMap< UObject*, TSharedRef<SNode> >& NodeToWidgetLookup, TSet<UEdGraphNode*>& PreviousNodes, TSet<UEdGraphNode*>& NextNodes, float HeightAboveWire) const
 {
 	// Find the previous position centroid
-	FVector2D PrevPos(0.0f, 0.0f);
+	FVector2f PrevPos(0.0f, 0.0f);
 	for (auto NodeIt = PreviousNodes.CreateConstIterator(); NodeIt; ++NodeIt)
 	{
 		UEdGraphNode* PreviousNode = *NodeIt;
-		const FVector2D CornerPos(PreviousNode->NodePosX, PreviousNode->NodePosY);
+		const FVector2f CornerPos = PreviousNode->GetPosition();
 		PrevPos += CornerPos + NodeToWidgetLookup.FindChecked(PreviousNode)->GetDesiredSize() * 0.5f;
 	}
 
 	// Find the next position centroid
-	FVector2D NextPos(0.0f, 0.0f);
+	FVector2f NextPos(0.0f, 0.0f);
 	for (auto NodeIt = NextNodes.CreateConstIterator(); NodeIt; ++NodeIt)
 	{
 		UEdGraphNode* NextNode = *NodeIt;
-		const FVector2D CornerPos(NextNode->NodePosX, NextNode->NodePosY);
+		const FVector2f CornerPos = NextNode->GetPosition();
 		NextPos += CornerPos + NodeToWidgetLookup.FindChecked(NextNode)->GetDesiredSize() * 0.5f;
 	}
 
 	PositionThisNodeBetweenOtherNodes(PrevPos, NextPos, HeightAboveWire);
 }
 
-void SGraphNode::PositionThisNodeBetweenOtherNodes(const FVector2D& PrevPos, const FVector2D& NextPos, float HeightAboveWire) const
+void SGraphNode::PositionThisNodeBetweenOtherNodes(const FVector2f& PrevPos, const FVector2f& NextPos, float HeightAboveWire) const
 {
-	const FVector2D DesiredNodeSize = GetDesiredSize();
+	const FVector2f DesiredNodeSize = GetDesiredSize();
 
-	FVector2D DeltaPos(NextPos - PrevPos);
+	FVector2f DeltaPos(NextPos - PrevPos);
 	if (DeltaPos.IsNearlyZero())
 	{
-		DeltaPos = FVector2D(10.0f, 0.0f);
+		DeltaPos = FVector2f(10.0f, 0.0f);
 	}
 
-	const FVector2D Normal = FVector2D(DeltaPos.Y, -DeltaPos.X).GetSafeNormal();
+	const FVector2f Normal = FVector2f(DeltaPos.Y, -DeltaPos.X).GetSafeNormal();
 
-	const FVector2D SlidingCapsuleBias = FVector2D::ZeroVector;//(0.5f * FMath::Sin(Normal.X * (float)HALF_PI) * DesiredNodeSize.X, 0.0f);
+	const FVector2f SlidingCapsuleBias = FVector2f::ZeroVector;//(0.5f * FMath::Sin(Normal.X * (float)HALF_PI) * DesiredNodeSize.X, 0.0f);
 
-	const FVector2D NewCenter = PrevPos + (0.5f * DeltaPos) + (HeightAboveWire * Normal) + SlidingCapsuleBias;
+	const FVector2f NewCenter = PrevPos + (0.5f * DeltaPos) + (HeightAboveWire * Normal) + SlidingCapsuleBias;
 
 	// Now we need to adjust the new center by the node size and zoom factor
-	const FVector2D NewCorner = NewCenter - (0.5f * DesiredNodeSize);
+	const FVector2f NewCorner = NewCenter - (0.5f * DesiredNodeSize);
 
-	GraphNode->NodePosX = NewCorner.X;
-	GraphNode->NodePosY = NewCorner.Y;
+	GraphNode->SetPosition(NewCorner);
 }
 
 FText SGraphNode::GetErrorMsgToolTip() const
diff --git a/Engine/Source/Editor/GraphEditor/Private/SGraphNodeComment.cpp b/Engine/Source/Editor/GraphEditor/Private/SGraphNodeComment.cpp
--- a/Engine/Source/Editor/GraphEditor/Private/SGraphNodeComment.cpp
+++ b/Engine/Source/Editor/GraphEditor/Private/SGraphNodeComment.cpp
@@ -85,8 +85,7 @@ void SGraphNodeComment::Construct(const FArguments& InArgs, UEdGraphNode_Comment
 	this->UpdateGraphNode();
 
 	// Pull out sizes
-	UserSize.X = InNode->NodeWidth;
-	UserSize.Y = InNode->NodeHeight;
+	UserSize = InNode->GetSize();
 
 	// Cache desired size so we cull correctly. We can do this as our ComputeDesiredSize ignores the layout scale.
 	CacheDesiredSize(1.0f);
@@ -136,7 +135,7 @@ void SGraphNodeComment::OnDragEnter( const FGeometry& MyGeometry, const FDragDro
 
 float SGraphNodeComment::GetWrapAt() const
 {
-	return (float)(CachedWidth - 32.0f);
+	return static_cast<float>(CachedWidth) - 32.0f;
 }
 
 bool SGraphNodeComment::IsNameReadOnly() const
@@ -167,8 +166,8 @@ void SGraphNodeComment::UpdateGraphNode()
 	CachedFontSize = CommentNode->GetFontSize();
 
 	CommentStyle = FAppStyle::Get().GetWidgetStyle<FInlineEditableTextBlockStyle>("Graph.CommentBlock.TitleInlineEditableText");
-	CommentStyle.EditableTextBoxStyle.TextStyle.Font.Size = CachedFontSize;
-	CommentStyle.TextStyle.Font.Size = CachedFontSize;
+	CommentStyle.EditableTextBoxStyle.TextStyle.Font.Size = static_cast<float>(CachedFontSize);
+	CommentStyle.TextStyle.Font.Size = static_cast<float>(CachedFontSize);
 
 	this->ContentScale.Bind( this, &SGraphNode::GetContentScale );
 	this->GetOrAddSlot( ENodeZone::Center )
@@ -404,8 +403,7 @@ void SGraphNodeComment::MoveTo( const FVector2f& NewPosition, FNodeSet& NodeFilt
 					{
 						NodeFilter.Add(Node->DEPRECATED_NodeWidget.Pin());
 						Node->Modify(bMarkDirty);
-						Node->NodePosX += PositionDelta.X;
-						Node->NodePosY += PositionDelta.Y;
+						Node->SetPosition(Node->GetPosition() + PositionDelta);
 					}
 				}
 			}
diff --git a/Engine/Source/Editor/GraphEditor/Private/SGraphNodeDocumentation.cpp b/Engine/Source/Editor/GraphEditor/Private/SGraphNodeDocumentation.cpp
--- a/Engine/Source/Editor/GraphEditor/Private/SGraphNodeDocumentation.cpp
+++ b/Engine/Source/Editor/GraphEditor/Private/SGraphNodeDocumentation.cpp
@@ -80,8 +80,7 @@ void SGraphNodeDocumentation::Construct( const FArguments& InArgs, UEdGraphNode*
 		ZoomCurve = SpawnAnim.AddCurve( 0, 0.1f );
 		FadeCurve = SpawnAnim.AddCurve( 0.15f, 0.15f );
 	}
-	UserSize.X = InNode->NodeWidth;
-	UserSize.Y = InNode->NodeHeight;
+	UserSize = InNode->GetSize();
 	bUserIsDragging = false;
 	UpdateGraphNode();
 }
@@ -256,11 +255,10 @@ TSharedPtr<SWidget> SGraphNodeDocumentation::CreateDocumentationPage()
 			UserSize = GraphNodeDocumentationDefs::MaximumNodeSize;
 			ContentWidget->SlatePrepass();
 			DocumentationSize = ContentWidget->GetDesiredSize();
-			if( GraphNode->NodeWidth != 0.f && GraphNode->NodeHeight != 0.f )
+			if( GraphNode->NodeWidth != 0 && GraphNode->NodeHeight != 0 )
 			{
 				// restore node size
-				UserSize.X = GraphNode->NodeWidth;
-				UserSize.Y = GraphNode->NodeHeight;
+				UserSize = GraphNode->GetSize();
 			}
 			else
 			{
@@ -375,8 +373,8 @@ void SGraphNodeDocumentation::Tick( const FGeometry& AllottedGeometry, const dou
 	if( CachedDocumentationLink != GraphNode->GetDocumentationLink() ||
 		CachedDocumentationExcerpt != GraphNode->GetDocumentationExcerptName())
 	{
-		GraphNode->NodeWidth = 0.f;
-		GraphNode->NodeHeight = 0.f;
+		GraphNode->NodeWidth = 0;
+		GraphNode->NodeHeight = 0;
 		UpdateGraphNode();
 	}
 }
diff --git a/Engine/Source/Editor/GraphEditor/Private/SGraphNodeResizable.cpp b/Engine/Source/Editor/GraphEditor/Private/SGraphNodeResizable.cpp
--- a/Engine/Source/Editor/GraphEditor/Private/SGraphNodeResizable.cpp
+++ b/Engine/Source/Editor/GraphEditor/Private/SGraphNodeResizable.cpp
@@ -169,7 +169,7 @@ FReply SGraphNodeResizable::OnMouseMove(const FGeometry& MyGeometry, const FPoin
 		DragSize.X += DeltaNodeSize.X;
 		DragSize.Y += DeltaNodeSize.Y;
 		// apply snap
-		const uint32 SnapSize = SNodePanel::GetSnapGridSize();
+		const float SnapSize = static_cast<float>(SNodePanel::GetSnapGridSize());
 		FVector2f SnappedSize;
 		SnappedSize.X = SnapSize * FMath::RoundToFloat( DragSize.X/SnapSize );
 		SnappedSize.Y = SnapSize * FMath::RoundToFloat( DragSize.Y/SnapSize );
diff --git a/Engine/Source/Editor/GraphEditor/Private/SGraphPanel.cpp b/Engine/Source/Editor/GraphEditor/Private/SGraphPanel.cpp
--- a/Engine/Source/Editor/GraphEditor/Private/SGraphPanel.cpp
+++ b/Engine/Source/Editor/GraphEditor/Private/SGraphPanel.cpp
@@ -786,24 +786,25 @@ FReply SGraphPanel::OnKeyDown( const FGeometry& MyGeometry, const FKeyEvent& InK
 		const bool bIsModifierActive = InKeyEvent.IsCommandDown() || InKeyEvent.IsAltDown() || InKeyEvent.IsShiftDown() || InKeyEvent.IsControlDown();
 		if (!bIsModifierActive)
 		{
+			const float SnapGridSize = static_cast<float>(GetSnapGridSize());
 			if( InKeyEvent.GetKey() == EKeys::Up || InKeyEvent.GetKey() == EKeys::NumPadEight )
 			{
-				UpdateSelectedNodesPositions(FVector2f(0.0f,-1.0f * GetSnapGridSize()));
+				UpdateSelectedNodesPositions(FVector2f(0.0f,-1.0f * SnapGridSize));
 				return FReply::Handled();
 			}
 			if( InKeyEvent.GetKey() == EKeys::Down || InKeyEvent.GetKey() == EKeys::NumPadTwo )
 			{
-				UpdateSelectedNodesPositions(FVector2f(0.0f,GetSnapGridSize()));
+				UpdateSelectedNodesPositions(FVector2f(0.0f, SnapGridSize));
 				return FReply::Handled();
 			}
 			if( InKeyEvent.GetKey() == EKeys::Right || InKeyEvent.GetKey() == EKeys::NumPadSix )
 			{
-				UpdateSelectedNodesPositions(FVector2f(GetSnapGridSize(),0.0f));
+				UpdateSelectedNodesPositions(FVector2f(SnapGridSize,0.0f));
 				return FReply::Handled();
 			}
 			if( InKeyEvent.GetKey() == EKeys::Left || InKeyEvent.GetKey() == EKeys::NumPadFour )
 			{
-				UpdateSelectedNodesPositions(FVector2f(-1.0f * GetSnapGridSize(),0.0f));
+				UpdateSelectedNodesPositions(FVector2f(-1.0f * SnapGridSize,0.0f));
 				return FReply::Handled();
 			}
 		}
@@ -1369,7 +1370,7 @@ FReply SGraphPanel::OnDrop( const FGeometry& MyGeometry, const FDragDropEvent& D
 	else if (Operation->IsOfType<FGraphNodeDragDropOp>())
 	{
 		TSharedPtr<FGraphNodeDragDropOp> NodeDropOp = StaticCastSharedPtr<FGraphNodeDragDropOp>(Operation);
-		NodeDropOp->OnPerformDropToGraph.ExecuteIfBound(NodeDropOp, GraphObj, FVector2D(NodeAddPosition), FVector2D(DragDropEvent.GetScreenSpacePosition()));
+		NodeDropOp->OnPerformDropToGraphAtLocation.ExecuteIfBound(NodeDropOp, GraphObj, NodeAddPosition, DragDropEvent.GetScreenSpacePosition());
 		return FReply::Handled();
 	}
 	else
@@ -1488,7 +1489,7 @@ void SGraphPanel::MoveNodesToAnchorPinAtGraphPosition(TArrayView<UEdGraphNode* c
 	{
 		static int32 RoundToGrid(int32 Value, int32 GridSize)
 		{
-			return FMath::RoundToInt(static_cast<float>(Value) / GridSize) * GridSize;
+			return FMath::RoundToInt(static_cast<float>(Value) / static_cast<float>(GridSize)) * GridSize;
 		}
 
 		// The standard SnapToGrid() will floor values, but this will round them instead
@@ -1598,8 +1599,7 @@ void SGraphPanel::MoveNodesToAnchorPinAtGraphPosition(TArrayView<UEdGraphNode* c
 				// Extra safety in case this was called from within a timer
 				if (IsValid(SpawnedNode))
 				{
-					SpawnedNode->NodePosX += Delta.X;
-					SpawnedNode->NodePosY += Delta.Y;
+					SpawnedNode->SetPosition(SpawnedNode->GetPosition() + Delta);
 
 					// Note: Not using the standard SnapToGrid() on purpose since we actually want to be as close
 					// to the user's dragged location as possible, and flooring ends up with noticeably more error
@@ -1767,7 +1767,7 @@ TSharedPtr<SWidget> SGraphPanel::SummonContextMenu(const UE::Slate::FDeprecateVe
 void SGraphPanel::SummonCreateNodeMenuFromUICommand(uint32 NumNodesAdded)
 {
 	FVector2f WhereToSummonMenu = LastPointerEvent.GetScreenSpacePosition();
-	const float AdditionalOffset = 1.0f + (NumNodesAdded * GetDefault<UGraphEditorSettings>()->PaddingAutoCollateIncrement);
+	const float AdditionalOffset = 1.0f + (static_cast<float>(NumNodesAdded) * GetDefault<UGraphEditorSettings>()->PaddingAutoCollateIncrement);
 	FVector2f WhereToAddNode = PastePosition + AdditionalOffset;
 
 	TSharedPtr<SGraphNode> NodeUnderMouse = GetGraphNodeUnderMouse(LastPointerGeometry, LastPointerEvent);
@@ -1810,7 +1810,7 @@ void SGraphPanel::AttachGraphEvents(TSharedPtr<SGraphNode> CreatedSubNode)
 bool SGraphPanel::GetBoundsForNode(const UObject* InNode, FVector2D& MinCorner, FVector2D& MaxCorner, float Padding) const
 {
 	FVector2f TempMinCorner, TempMaxCorner;
-	bool bRetVal = GetBoundsForNode(InNode, TempMinCorner, TempMaxCorner, Padding);
+	bool bRetVal = SNodePanel::GetBoundsForNode(InNode, TempMinCorner, TempMaxCorner, Padding);
 	MinCorner = FVector2D(TempMinCorner);
 	MaxCorner = FVector2D(TempMaxCorner);
 
@@ -1818,7 +1818,14 @@ bool SGraphPanel::GetBoundsForNode(const UObject* InNode, FVector2D& MinCorner,
 }
 bool SGraphPanel::GetBoundsForNode(const UObject* InNode, FVector2f& MinCorner, FVector2f& MaxCorner, float Padding) const
 {
-	return SNodePanel::GetBoundsForNode(InNode, MinCorner, MaxCorner, Padding);
+PRAGMA_DISABLE_DEPRECATION_WARNINGS
+	FVector2D TempMinCorner, TempMaxCorner;
+	bool bRetVal = GetBoundsForNode(InNode, TempMinCorner, TempMaxCorner, Padding);
+	MinCorner = UE::Slate::CastToVector2f(TempMinCorner);
+	MaxCorner = UE::Slate::CastToVector2f(TempMaxCorner);
+PRAGMA_ENABLE_DEPRECATION_WARNINGS
+
+	return bRetVal;
 }
 
 class FConnectionAligner
@@ -1900,7 +1907,7 @@ private:
 			float AlignmentDelta = 0.f;
 			for (const FPinPair& Pins : NodeToPins.Value)
 			{
-				AlignmentDelta += (Node->NodePosY + Pins.SrcPin->GetNodeOffset().Y) - (NodeToPins.Key->NodePosY + Pins.DstPin->GetNodeOffset().Y);
+				AlignmentDelta += (Node->GetNodePosY() + Pins.SrcPin->GetNodeOffset().Y) - (NodeToPins.Key->GetNodePosY() + Pins.DstPin->GetNodeOffset().Y);
 			}
 
 			UEdGraph* GraphObj = NodeToPins.Key->GetGraph();
@@ -1909,8 +1916,8 @@ private:
 
 			const UEdGraphSchema* Schema = GraphObj->GetSchema();
 
-			float NewNodePosY = NodeToPins.Key->NodePosY + (AlignmentDelta / NodeToPins.Value.Num());
-			Schema->SetNodePosition(NodeToPins.Key, FVector2f(NodeToPins.Key->NodePosX, NewNodePosY));
+			float NewNodePosY = NodeToPins.Key->GetNodePosY() + (AlignmentDelta / static_cast<float>(NodeToPins.Value.Num()));
+			Schema->SetNodePosition(NodeToPins.Key, FVector2f(NodeToPins.Key->GetNodePosX(), NewNodePosY));
 	
 			VisitedNodes.Add(Node);
 			VisitedNodes.Add(NodeToPins.Key);
diff --git a/Engine/Source/Editor/GraphEditor/Private/SGraphPin.cpp b/Engine/Source/Editor/GraphEditor/Private/SGraphPin.cpp
--- a/Engine/Source/Editor/GraphEditor/Private/SGraphPin.cpp
+++ b/Engine/Source/Editor/GraphEditor/Private/SGraphPin.cpp
@@ -236,7 +236,7 @@ void SGraphPin::Construct(const FArguments& InArgs, UEdGraphPin* InPin)
 		SNew(SButton)
 		.ButtonStyle(FAppStyle::Get(), NAME_NoBorder)
 		.Visibility(this, &SGraphPin::GetPinStatusIconVisibility)
-		.ContentPadding(0)
+		.ContentPadding(0.0f)
 		.OnClicked(this, &SGraphPin::ClickedOnPinStatusIcon)
 		[
 			SNew(SImage)
@@ -278,7 +278,7 @@ void SGraphPin::Construct(const FArguments& InArgs, UEdGraphPin* InPin)
 		{
 			TSharedPtr<SBox> ValueBox;
 			LabelAndValue->AddSlot()
-				.Padding(bIsInput ? FMargin(InArgs._SideToSideMargin, 0, 0, 0) : FMargin(0, 0, InArgs._SideToSideMargin, 0))
+				.Padding(bIsInput ? FMargin(InArgs._SideToSideMargin, 0.0f, 0.0f, 0.0f) : FMargin(0.0f, 0.0f, InArgs._SideToSideMargin, 0.0f))
 				.VAlign(VAlign_Center)
 				[
 					SAssignNew(ValueBox, SBox)
@@ -310,7 +310,7 @@ void SGraphPin::Construct(const FArguments& InArgs, UEdGraphPin* InPin)
 			+SHorizontalBox::Slot()
 			.AutoWidth()
 			.VAlign(VAlign_Center)
-			.Padding(0, 0, InArgs._SideToSideMargin, 0)
+			.Padding(0.0f, 0.0f, InArgs._SideToSideMargin, 0.0f)
 			[
 				PinWidgetRef
 			]
@@ -334,7 +334,7 @@ void SGraphPin::Construct(const FArguments& InArgs, UEdGraphPin* InPin)
 			+SHorizontalBox::Slot()
 			.AutoWidth()
 			.VAlign(VAlign_Center)
-			.Padding(InArgs._SideToSideMargin, 0, 0, 0)
+			.Padding(InArgs._SideToSideMargin, 0.0f, 0.0f, 0.0f)
 			[
 				PinWidgetRef
 			];
@@ -1399,13 +1399,7 @@ FText SGraphPin::GetTooltipText() const
 
 void SGraphPin::GetInteractiveTooltipLocation(FVector2D& InOutDesiredLocation) const
 {
-	FVector2f TempVector2f;
-	GetInteractiveTooltipLocation(TempVector2f);
-	InOutDesiredLocation = FVector2D(TempVector2f);
-}
-
-void SGraphPin::GetInteractiveTooltipLocation(FVector2f& InOutDesiredLocation) const
-{
+	FVector2f OutDesiredLocation = UE::Slate::CastToVector2f(InOutDesiredLocation);
 	TSharedPtr<SGraphNode> OwnerNode = OwnerNodePtr.Pin();
 	if (OwnerNode.IsValid())
 	{
@@ -1413,10 +1407,10 @@ void SGraphPin::GetInteractiveTooltipLocation(FVector2f& InOutDesiredLocation) c
 		if (GraphPanel.IsValid())
 		{
 			// Reset to the pin's location in graph space.
-			InOutDesiredLocation = OwnerNode->GetPosition2f() + CachedNodeOffset;
-			
+			OutDesiredLocation = OwnerNode->GetPosition2f() + CachedNodeOffset;
+
 			// Shift the desired location to the right edge of the pin's geometry.
-			InOutDesiredLocation.X += GetTickSpaceGeometry().Size.X;
+			OutDesiredLocation.X += GetTickSpaceGeometry().Size.X;
 
 			// Align to the first entry in the inspector's tree view.
 			TSharedPtr<FPinValueInspectorTooltip> Inspector = ValueInspectorTooltip.Pin();
@@ -1428,22 +1422,33 @@ void SGraphPin::GetInteractiveTooltipLocation(FVector2f& InOutDesiredLocation) c
 
 				if (Inspector->ValueInspectorWidget->ShouldShowSearchFilter())
 				{
-					InOutDesiredLocation.Y -= VerticalOffsetWithSearchFilter;
+					OutDesiredLocation.Y -= VerticalOffsetWithSearchFilter;
 				}
 				else
 				{
-					InOutDesiredLocation.Y -= VerticalOffsetWithoutSearchFilter;
+					OutDesiredLocation.Y -= VerticalOffsetWithoutSearchFilter;
 				}
 			}
 
 			// Convert our desired location from graph coordinates into panel space.
-			InOutDesiredLocation -= GraphPanel->GetViewOffset();
-			InOutDesiredLocation *= GraphPanel->GetZoomAmount();
+			OutDesiredLocation -= GraphPanel->GetViewOffset();
+			OutDesiredLocation *= GraphPanel->GetZoomAmount();
 
 			// Finally, convert the modified location from panel space into screen space.
-			InOutDesiredLocation = GraphPanel->GetTickSpaceGeometry().LocalToAbsolute(InOutDesiredLocation);
+			OutDesiredLocation = GraphPanel->GetTickSpaceGeometry().LocalToAbsolute(InOutDesiredLocation);
 		}
 	}
+
+	InOutDesiredLocation = FVector2D(OutDesiredLocation);
+}
+
+void SGraphPin::GetInteractiveTooltipLocation(FVector2f& InOutDesiredLocation) const
+{
+PRAGMA_DISABLE_DEPRECATION_WARNINGS
+	FVector2D TempVector2f = FVector2D(InOutDesiredLocation);
+	GetInteractiveTooltipLocation(TempVector2f);
+	InOutDesiredLocation = UE::Slate::CastToVector2f(TempVector2f);
+PRAGMA_ENABLE_DEPRECATION_WARNINGS
 }
 
 bool SGraphPin::IsEditingEnabled() const
diff --git a/Engine/Source/Editor/GraphEditor/Private/SNodePanel.cpp b/Engine/Source/Editor/GraphEditor/Private/SNodePanel.cpp
--- a/Engine/Source/Editor/GraphEditor/Private/SNodePanel.cpp
+++ b/Engine/Source/Editor/GraphEditor/Private/SNodePanel.cpp
@@ -800,7 +800,7 @@ FReply SNodePanel::OnMouseMove(const FGeometry& MyGeometry, const FPointerEvent&
 			// Get rid of mouse movement that's been 'used up' by zooming
 			if (ZoomLevelDelta != 0)
 			{
-				TotalMouseDeltaXY -= (ZoomLevelDelta / NodePanelDefs::MouseZoomScaling);
+				TotalMouseDeltaXY -= (static_cast<float>(ZoomLevelDelta) / NodePanelDefs::MouseZoomScaling);
 			}
 
 			// Perform zoom centered on the cached start offset
@@ -1397,8 +1397,8 @@ void SNodePanel::PaintBackgroundAsLines(const FSlateBrush* BackgroundImage, cons
 	
 	const float GridCellSize = NominalGridSize * ZoomFactor * Inflation;
 
-	const float GraphSpaceGridX0 = FancyMod(ViewOffset.X, Inflation * NominalGridSize * RulePeriod);
-	const float GraphSpaceGridY0 = FancyMod(ViewOffset.Y, Inflation * NominalGridSize * RulePeriod);
+	const float GraphSpaceGridX0 = FancyMod(ViewOffset.X, Inflation * NominalGridSize * static_cast<float>(RulePeriod));
+	const float GraphSpaceGridY0 = FancyMod(ViewOffset.Y, Inflation * NominalGridSize * static_cast<float>(RulePeriod));
 
 	float ImageOffsetX = GraphSpaceGridX0 * -ZoomFactor;
 	float ImageOffsetY = GraphSpaceGridY0 * -ZoomFactor;
diff --git a/Engine/Source/Editor/GraphEditor/Public/BlueprintConnectionDrawingPolicy.h b/Engine/Source/Editor/GraphEditor/Public/BlueprintConnectionDrawingPolicy.h
--- a/Engine/Source/Editor/GraphEditor/Public/BlueprintConnectionDrawingPolicy.h
+++ b/Engine/Source/Editor/GraphEditor/Public/BlueprintConnectionDrawingPolicy.h
@@ -138,7 +138,7 @@ protected:
 	FTimePair const* BackTraceExecPath(UEdGraphPin const* const OutputPin, FExecPairingMap const* const NodeExecutionList);
 
 	bool ShouldChangeTangentForKnot(class UK2Node_Knot* Knot);
-	bool GetAverageConnectedPosition(class UK2Node_Knot* Knot, EEdGraphPinDirection Direction, FVector2D& OutPos) const;
-	bool FindPinCenter(UEdGraphPin* Pin, FVector2D& OutCenter) const;
+	bool GetAverageConnectedPosition(class UK2Node_Knot* Knot, EEdGraphPinDirection Direction, FVector2f& OutPos) const;
+	bool FindPinCenter(UEdGraphPin* Pin, FVector2f& OutCenter) const;
 	static UObject* GetObjectBeingDebugged(UBlueprint* TargetBP);
 };
diff --git a/Engine/Source/Editor/GraphEditor/Public/KismetPins/SGraphPinNum.h b/Engine/Source/Editor/GraphEditor/Public/KismetPins/SGraphPinNum.h
--- a/Engine/Source/Editor/GraphEditor/Public/KismetPins/SGraphPinNum.h
+++ b/Engine/Source/Editor/GraphEditor/Public/KismetPins/SGraphPinNum.h
@@ -28,8 +28,8 @@ protected:
 	virtual TSharedRef<SWidget>	GetDefaultValueWidget() override
 	{
 		return SNew(SBox)
-			.MinDesiredWidth(18)
-			.MaxDesiredWidth(400)
+			.MinDesiredWidth(18.0f)
+			.MaxDesiredWidth(400.0f)
 			[
 				SNew(SNumericEntryBox<NumericType>)
 				.EditableTextBoxStyle(FAppStyle::Get(), "Graph.EditableTextBox")
diff --git a/Engine/Source/Editor/GraphEditor/Public/KismetPins/SGraphPinNumSlider.h b/Engine/Source/Editor/GraphEditor/Public/KismetPins/SGraphPinNumSlider.h
--- a/Engine/Source/Editor/GraphEditor/Public/KismetPins/SGraphPinNumSlider.h
+++ b/Engine/Source/Editor/GraphEditor/Public/KismetPins/SGraphPinNumSlider.h
@@ -53,7 +53,7 @@ protected:
 		
 		return SNew(SBox)
 			.MinDesiredWidth(MinDesiredBoxWidth)
-			.MaxDesiredWidth(400)
+			.MaxDesiredWidth(400.0f)
 			[
 				SNew(SNumericEntryBox<NumericType>)
 				.EditableTextBoxStyle(FAppStyle::Get(), "Graph.EditableTextBox")
@@ -164,7 +164,7 @@ private:
 	FProperty* PinProperty;
 	NumericType LastSliderCommittedValue;
 	NumericType SliderValue;
-	float MinDesiredBoxWidth = 0;
+	float MinDesiredBoxWidth = 0.0f;
 	bool bIsUsingSlider = false;
 	bool bShouldShowDisabledWhenConnected = false;
 };
diff --git a/Engine/Source/Editor/GraphEditor/Public/SGraphNode.h b/Engine/Source/Editor/GraphEditor/Public/SGraphNode.h
--- a/Engine/Source/Editor/GraphEditor/Public/SGraphNode.h
+++ b/Engine/Source/Editor/GraphEditor/Public/SGraphNode.h
@@ -89,7 +89,7 @@ public:
 	FText GetHeadTitle() const;
 
 	/** Get the size of this title the last time it was drawn */
-	FVector2D GetTitleSize() const;
+	UE::Slate::FDeprecateVector2DResult GetTitleSize() const;
 
 	/** Mark the cached title dirty */
 	void MarkDirty();
@@ -107,7 +107,7 @@ protected:
 	TAttribute< FText > TitleText;
 
 	/** The cached size of the title */
-	FVector2D CachedSize;
+	FVector2f CachedSize;
 
 protected:
 	// Gets the expected node title
@@ -265,7 +265,7 @@ public:
 	void NotifyDisallowedPinConnection(const class UEdGraphPin* PinA, const class UEdGraphPin* PinB) const;
 
 	/** Gets the unscaled position of the node from the last tick */
-	FVector2D GetUnscaledPosition() const {return CachedUnscaledPosition;}
+	UE::Slate::FDeprecateVector2DResult GetUnscaledPosition() const {return CachedUnscaledPosition;}
 
 	/** Returns the current Node LOD or Highest LOD if unable to query */
 	EGraphRenderingLOD::Type GetCurrentLOD() const;
@@ -293,7 +293,7 @@ protected:
 
 	void PositionThisNodeBetweenOtherNodes(const TMap< UObject*, TSharedRef<SNode> >& NodeToWidgetLookup, UEdGraphNode* PreviousNode, UEdGraphNode* NextNode, float HeightAboveWire) const;
 	void PositionThisNodeBetweenOtherNodes(const TMap< UObject*, TSharedRef<SNode> >& NodeToWidgetLookup, TSet<UEdGraphNode*>& PreviousNodes, TSet<UEdGraphNode*>& NextNodes, float HeightAboveWire) const;
-	void PositionThisNodeBetweenOtherNodes(const FVector2D& PrevPos, const FVector2D& NextPos, float HeightAboveWire) const;
+	void PositionThisNodeBetweenOtherNodes(const FVector2f& PrevPos, const FVector2f& NextPos, float HeightAboveWire) const;
 
 	/**
 	 * Check whether drag and drop functionality is permitted on the given node 
@@ -484,7 +484,7 @@ protected:
 
 
 	/** Caches true position of node */
-	FVector2D CachedUnscaledPosition;
+	FVector2f CachedUnscaledPosition;
 
 	/** Cached icon color for the node */
 	FLinearColor IconColor;
diff --git a/Engine/Source/Runtime/Engine/Classes/EdGraph/EdGraphNode.h b/Engine/Source/Runtime/Engine/Classes/EdGraph/EdGraphNode.h
--- a/Engine/Source/Runtime/Engine/Classes/EdGraph/EdGraphNode.h
+++ b/Engine/Source/Runtime/Engine/Classes/EdGraph/EdGraphNode.h
@@ -494,6 +494,20 @@ public:
 	ENGINE_API virtual void BeginDestroy() override;
 	// End of UObject interface
 
+	// Helper conversions for the editor which typically expects the Node positions and sizes in floats instead of int32
+	float GetWidth() const { return static_cast<float>(NodeWidth); }
+	void SetWidth(float InWidth) { NodeWidth = UE::LWC::FloatToIntCastChecked<int32>(InWidth); }
+	float GetHeight() const { return static_cast<float>(NodeHeight); }
+	void SetHeight(float InHeight) { NodeHeight = UE::LWC::FloatToIntCastChecked<int32>(InHeight); }
+	float GetNodePosX() const { return static_cast<float>(NodePosX); }
+	void SetNodePosX(float InX) { NodePosX = UE::LWC::FloatToIntCastChecked<int32>(InX); }
+	float GetNodePosY() const { return static_cast<float>(NodePosY); }
+	void SetNodePosY(float InY) { NodePosY = UE::LWC::FloatToIntCastChecked<int32>(InY); }
+	FVector2f GetSize() const { return FVector2f(GetHeight(), GetWidth()); }
+	void SetSize(const FVector2f& InSize) { SetWidth(InSize.X); SetHeight(InSize.Y); }
+	FVector2f GetPosition() const { return FVector2f(GetNodePosX(), GetNodePosY()); }
+	void SetPosition(const FVector2f& InPosition) { SetNodePosX(InPosition.X); SetNodePosY(InPosition.Y); }
+
 	/** widget representing this node if it exists; Note: This is not safe to use in general and will be removed in the future, as there is no guarantee that only one graph editor/panel is viewing a given graph */
 	TWeakPtr<SGraphNode> DEPRECATED_NodeWidget;
 
