[Truncation Warnings] Final pass on GraphEditor module for FVector2D usage, and confirming shims went in the right direction

diff --git a/Engine/Source/Editor/GraphEditor/Private/DragConnection.cpp b/Engine/Source/Editor/GraphEditor/Private/DragConnection.cpp
--- a/Engine/Source/Editor/GraphEditor/Private/DragConnection.cpp
+++ b/Engine/Source/Editor/GraphEditor/Private/DragConnection.cpp
@@ -53,7 +53,7 @@ void FDragConnection::OnDrop(bool bDropWasHandled, const FPointerEvent& MouseEve
 
 void FDragConnection::OnDragged(const class FDragDropEvent& DragDropEvent)
 {
-	FVector2D TargetPosition = DragDropEvent.GetScreenSpacePosition();
+	FVector2f TargetPosition = DragDropEvent.GetScreenSpacePosition();
 
 	// Reposition the info window wrt to the drag
 	CursorDecoratorWindow->MoveWindowTo(DragDropEvent.GetScreenSpacePosition() + DecoratorAdjust);
diff --git a/Engine/Source/Editor/GraphEditor/Private/DragNode.cpp b/Engine/Source/Editor/GraphEditor/Private/DragNode.cpp
--- a/Engine/Source/Editor/GraphEditor/Private/DragNode.cpp
+++ b/Engine/Source/Editor/GraphEditor/Private/DragNode.cpp
@@ -66,7 +66,7 @@ UEdGraphNode* FDragNode::GetGraphNodeForSGraphNode(TSharedRef<SGraphNode>& SNode
 
 void FDragNode::OnDragged(const class FDragDropEvent& DragDropEvent)
 {
-	FVector2D TargetPosition = DragDropEvent.GetScreenSpacePosition();
+	FVector2f TargetPosition = DragDropEvent.GetScreenSpacePosition();
 
 	// Reposition the info window to the dragged position
 	CursorDecoratorWindow->MoveWindowTo(DragDropEvent.GetScreenSpacePosition() + DecoratorAdjust);
diff --git a/Engine/Source/Editor/GraphEditor/Private/DragNode.h b/Engine/Source/Editor/GraphEditor/Private/DragNode.h
--- a/Engine/Source/Editor/GraphEditor/Private/DragNode.h
+++ b/Engine/Source/Editor/GraphEditor/Private/DragNode.h
@@ -51,7 +51,7 @@ protected:
 	TArray< TSharedRef<SGraphNode> > DraggedNodes;
 
 	/** Offset information for the decorator widget */
-	FVector2D	DecoratorAdjust;
+	FVector2f	DecoratorAdjust;
 
 	/** if we can drop our node here*/
 	bool bValidOperation;
diff --git a/Engine/Source/Editor/GraphEditor/Private/KismetNodes/SGraphNodeK2CreateDelegate.cpp b/Engine/Source/Editor/GraphEditor/Private/KismetNodes/SGraphNodeK2CreateDelegate.cpp
--- a/Engine/Source/Editor/GraphEditor/Private/KismetNodes/SGraphNodeK2CreateDelegate.cpp
+++ b/Engine/Source/Editor/GraphEditor/Private/KismetNodes/SGraphNodeK2CreateDelegate.cpp
@@ -219,9 +219,9 @@ void SGraphNodeK2CreateDelegate::OnFunctionSelected(TSharedPtr<FString> Function
 
 				FName NewEventName = FBlueprintEditorUtils::FindUniqueCustomEventName(NodeBP);
 				UBlueprintEventNodeSpawner* Spawner = UBlueprintEventNodeSpawner::Create(UK2Node_CustomEvent::StaticClass(), NewEventName);
-				FVector2D SpawnPos = EventGraph->GetGoodPlaceForNewNode();
+				FVector2f SpawnPos = EventGraph->GetGoodPlaceForNewNode();
 
-				UEdGraphNode* NewNode = Spawner->Invoke(EventGraph, IBlueprintNodeBinder::FBindingSet(), SpawnPos);
+				UEdGraphNode* NewNode = Spawner->Invoke(EventGraph, IBlueprintNodeBinder::FBindingSet(), FDeprecateSlateVector2D(SpawnPos));
 
 				UK2Node_CustomEvent* NewEventNode = CastChecked<UK2Node_CustomEvent>(NewNode);
 				NewEventNode->SetDelegateSignature(Node->GetDelegateSignature());
diff --git a/Engine/Source/Editor/GraphEditor/Private/MaterialNodes/SGraphNodeMaterialConvert.cpp b/Engine/Source/Editor/GraphEditor/Private/MaterialNodes/SGraphNodeMaterialConvert.cpp
--- a/Engine/Source/Editor/GraphEditor/Private/MaterialNodes/SGraphNodeMaterialConvert.cpp
+++ b/Engine/Source/Editor/GraphEditor/Private/MaterialNodes/SGraphNodeMaterialConvert.cpp
@@ -49,7 +49,7 @@ TSharedPtr<SConvertInnerPin> FConvertDragDropOp::GetSourcePin() const
 	return SourcePin;
 }
 
-FVector2D FConvertDragDropOp::GetScreenPosition() const
+UE::Slate::FDeprecateVector2DResult FConvertDragDropOp::GetScreenPosition() const
 { 
 	return ScreenPosition;
 }
@@ -350,7 +350,7 @@ void SConvertInnerPin::CancelDragDrop()
 	}
 }
 
-FVector2D SConvertInnerPin::GetPinCenterAbsolute() const
+UE::Slate::FDeprecateVector2DResult SConvertInnerPin::GetPinCenterAbsolute() const
 { 
 	return CenterAbsolute; 
 }
@@ -679,8 +679,8 @@ int32 SGraphNodeMaterialConvert::OnPaint(
 		TSharedPtr<SConvertInnerPin> OutputInnerPin = GetInnerPin(false, ConvertMapping.OutputIndex, ConvertMapping.OutputComponentIndex);
 		if (InputInnerPin && OutputInnerPin)
 		{
-			const FVector2D CurveStart = InputInnerPin->GetPinCenterAbsolute();
-			const FVector2D CurveEnd = OutputInnerPin->GetPinCenterAbsolute();
+			const FVector2f CurveStart = InputInnerPin->GetPinCenterAbsolute();
+			const FVector2f CurveEnd = OutputInnerPin->GetPinCenterAbsolute();
 			MakeConnectionCurve(AllottedGeometry, OutDrawElements, LayerId, CurveStart, CurveEnd);
 		}
 	}
@@ -689,8 +689,8 @@ int32 SGraphNodeMaterialConvert::OnPaint(
 	{
 		if (TSharedPtr<SConvertInnerPin> SourcePin = CurrentDragDropOp->GetSourcePin())
 		{
-			const FVector2D CurveStart = SourcePin->GetPinCenterAbsolute();
-			const FVector2D CurveEnd = CurrentDragDropOp->GetScreenPosition() + FVector2D(Inverse(Args.GetWindowToDesktopTransform()));
+			const FVector2f CurveStart = SourcePin->GetPinCenterAbsolute();
+			const FVector2f CurveEnd = CurrentDragDropOp->GetScreenPosition() + Inverse(Args.GetWindowToDesktopTransform());
 			MakeConnectionCurve(AllottedGeometry, OutDrawElements, LayerId, CurveStart, CurveEnd);
 		}
 	}
@@ -869,12 +869,12 @@ void SGraphNodeMaterialConvert::MakeConnectionCurve(
 	const FGeometry& InAllottedGeometry,
 	FSlateWindowElementList& OutDrawElements,
 	int32& InOutLayerId,
-	const FVector2D InCurveStart,
-	const FVector2D InCurveEnd
+	const FVector2f& InCurveStart,
+	const FVector2f& InCurveEnd
 ) const
 {
-	FVector2D LocalStart = InAllottedGeometry.AbsoluteToLocal(InCurveStart);
-	FVector2D LocalEnd = InAllottedGeometry.AbsoluteToLocal(InCurveEnd);
+	FVector2f LocalStart = InAllottedGeometry.AbsoluteToLocal(InCurveStart);
+	FVector2f LocalEnd = InAllottedGeometry.AbsoluteToLocal(InCurveEnd);
 
 	//If LocalStart.X is > LocalEnd.X, swap values
 	if (LocalStart.X > LocalEnd.X)
@@ -892,7 +892,7 @@ void SGraphNodeMaterialConvert::MakeConnectionCurve(
 	{
 		/* Because we make sure the left-most position is always the start position,
 		 * curve direction is always pointing to the right */
-		static const FVector2D Direction = FVector2D(100, 0);
+		static const FVector2f Direction = FVector2f(100.0f, 0.0f);
 
 		FSlateDrawElement::MakeSpline(
 			OutDrawElements,
diff --git a/Engine/Source/Editor/GraphEditor/Private/SGraphNodeKnot.cpp b/Engine/Source/Editor/GraphEditor/Private/SGraphNodeKnot.cpp
--- a/Engine/Source/Editor/GraphEditor/Private/SGraphNodeKnot.cpp
+++ b/Engine/Source/Editor/GraphEditor/Private/SGraphNodeKnot.cpp
@@ -59,7 +59,7 @@ namespace SKnotNodeDefinitions
 	static const float KnotCenterBubbleAdjust = 20.f;
 
 	/** Knot node spacer sizes */
-	static const FVector2D NodeSpacerSize(42.0f, 24.0f);
+	static const FVector2f NodeSpacerSize(42.0f, 24.0f);
 }
 
 class FAmbivalentDirectionDragConnection : public FDragConnection
@@ -88,8 +88,8 @@ protected:
 	FAmbivalentDirectionDragConnection(UEdGraphNode* InKnot, const TSharedRef<SGraphPanel>& InGraphPanel, const FDraggedPinTable& InStartingPins)
 		: FDragConnection(InGraphPanel, InStartingPins)
 		, KnotPtr(InKnot)
-		, StartScreenPos(FVector2D::ZeroVector)
-		, MostRecentScreenPos(FVector2D::ZeroVector)
+		, StartScreenPos(FVector2f::ZeroVector)
+		, MostRecentScreenPos(FVector2f::ZeroVector)
 		, bLatchedStartScreenPos(false)
 	{
 	}
@@ -97,8 +97,8 @@ protected:
 protected:
 	TWeakObjectPtr<UEdGraphNode> KnotPtr;
 
-	FVector2D StartScreenPos;
-	FVector2D MostRecentScreenPos;
+	FVector2f StartScreenPos;
+	FVector2f MostRecentScreenPos;
 	bool bLatchedStartScreenPos;
 };
 
@@ -123,7 +123,7 @@ void FAmbivalentDirectionDragConnection::OnDragged(const class FDragDropEvent& D
 {
 	if (bLatchedStartScreenPos)
 	{
-		const FVector2D LastScreenPos = MostRecentScreenPos;
+		const FVector2f LastScreenPos = MostRecentScreenPos;
 		MostRecentScreenPos = DragDropEvent.GetScreenSpacePosition();
 
 		// Switch directions on the preview connector as we cross from left to right of the starting drag point (or vis versa)
diff --git a/Engine/Source/Editor/GraphEditor/Private/SGraphNodeResizable.cpp b/Engine/Source/Editor/GraphEditor/Private/SGraphNodeResizable.cpp
--- a/Engine/Source/Editor/GraphEditor/Private/SGraphNodeResizable.cpp
+++ b/Engine/Source/Editor/GraphEditor/Private/SGraphNodeResizable.cpp
@@ -264,16 +264,19 @@ UE::Slate::FDeprecateVector2DResult SGraphNodeResizable::GetCorrectedNodePositio
 
 	return CorrectedPos;
 }
-SGraphNodeResizable::EResizableWindowZone SGraphNodeResizable::FindMouseZone(const FVector2D& LocalMouseCoordinates) const
+SGraphNodeResizable::EResizableWindowZone SGraphNodeResizable::FindMouseZone(const FVector2f& LocalMouseCoordinates) const
 {
-	return FindMouseZone(UE::Slate::CastToVector2f(LocalMouseCoordinates));
+PRAGMA_DISABLE_DEPRECATION_WARNINGS
+	return FindMouseZone(FVector2D(LocalMouseCoordinates));
+PRAGMA_ENABLE_DEPRECATION_WARNINGS
 }
 
-SGraphNodeResizable::EResizableWindowZone SGraphNodeResizable::FindMouseZone(const FVector2f& LocalMouseCoordinates) const
+SGraphNodeResizable::EResizableWindowZone SGraphNodeResizable::FindMouseZone(const FVector2D& LocalMouseCoordinates) const
 {
 	EResizableWindowZone OutMouseZone = CRWZ_NotInWindow;
 	FSlateRect HitResultBorderSize = GetHitTestingBorder();
 	const FVector2f NodeSize = GetDesiredSize();
+	const FVector2f& LocalMouseCoords = UE::Slate::CastToVector2f(LocalMouseCoordinates);
 	const float TitleBarHeight = GetTitleBarHeight();
 
 	const float Zoom = GetOwnerPanel()->GetZoomAmount();
@@ -286,7 +289,7 @@ SGraphNodeResizable::EResizableWindowZone SGraphNodeResizable::FindMouseZone(con
 	}
 	
 	// Test for hit in location of 'grab' zone
-	if (LocalMouseCoordinates.Y > ( NodeSize.Y - HitResultBorderSize.Bottom))
+	if (LocalMouseCoords.Y > (NodeSize.Y - HitResultBorderSize.Bottom))
 	{
 		OutMouseZone = CRWZ_BottomBorder;
 	}
@@ -297,25 +300,25 @@ SGraphNodeResizable::EResizableWindowZone SGraphNodeResizable::FindMouseZone(con
 		
 		// Allow half of the title bar to be treated as the top edge
 		// so that the title bar can still be grabbed
-		if (LocalMouseCoordinates.Y < LargestZone * 0.4)
+		if (LocalMouseCoords.Y < LargestZone * 0.4)
 		{
 			OutMouseZone = CRWZ_TopBorder;
 		}
-		else if (LocalMouseCoordinates.Y < LargestZone)
+		else if (LocalMouseCoords.Y < LargestZone)
 		{
 			OutMouseZone = CRWZ_TitleBar;
 		}
 	}
-	else if (LocalMouseCoordinates.Y <= (HitResultBorderSize.Top))
+	else if (LocalMouseCoords.Y <= (HitResultBorderSize.Top))
 	{
 		OutMouseZone = CRWZ_TopBorder;
 	}
-	else if (LocalMouseCoordinates.Y <= TitleBarHeight)
+	else if (LocalMouseCoords.Y <= TitleBarHeight)
 	{
 		OutMouseZone = CRWZ_TitleBar;
 	}
 
-	if (LocalMouseCoordinates.X > (NodeSize.X - HitResultBorderSize.Right))
+	if (LocalMouseCoords.X > (NodeSize.X - HitResultBorderSize.Right))
 	{
 		if (OutMouseZone == CRWZ_BottomBorder)
 		{
@@ -330,7 +333,7 @@ SGraphNodeResizable::EResizableWindowZone SGraphNodeResizable::FindMouseZone(con
 			OutMouseZone = CRWZ_RightBorder;
 		}
 	}
-	else if (LocalMouseCoordinates.X <= HitResultBorderSize.Left)
+	else if (LocalMouseCoords.X <= HitResultBorderSize.Left)
 	{
 		if (OutMouseZone == CRWZ_TopBorder)
 		{
@@ -349,11 +352,11 @@ SGraphNodeResizable::EResizableWindowZone SGraphNodeResizable::FindMouseZone(con
 	// Test for hit on rest of frame
 	if (OutMouseZone == CRWZ_NotInWindow)
 	{
-		if (LocalMouseCoordinates.Y > HitResultBorderSize.Top) 
+		if (LocalMouseCoords.Y > HitResultBorderSize.Top)
 		{
 			OutMouseZone = CRWZ_InWindow;
 		}
-		else if (LocalMouseCoordinates.X > HitResultBorderSize.Left) 
+		else if (LocalMouseCoords.X > HitResultBorderSize.Left)
 		{
 			OutMouseZone = CRWZ_InWindow;
 		}
@@ -369,22 +372,26 @@ float SGraphNodeResizable::GetTitleBarHeight() const
 
 FVector2D SGraphNodeResizable::GetNodeMinimumSize() const
 {
-	return FVector2D(GetNodeMinimumSize2f());
+	return FVector2D(GraphNodeResizableDefs::MinNodeSize);
 }
 
 FVector2f SGraphNodeResizable::GetNodeMinimumSize2f() const
 {
-	return GraphNodeResizableDefs::MinNodeSize;
+PRAGMA_DISABLE_DEPRECATION_WARNINGS
+	return UE::Slate::CastToVector2f(GetNodeMinimumSize());
+PRAGMA_ENABLE_DEPRECATION_WARNINGS
 }
 
 FVector2D SGraphNodeResizable::GetNodeMaximumSize() const
 {
-	return FVector2D(GetNodeMaximumSize2f());
+	return GraphNodeResizableDefs::MaxNodeSize;
 }
 
 FVector2f SGraphNodeResizable::GetNodeMaximumSize2f() const
 {
-	return GraphNodeResizableDefs::MaxNodeSize;
+PRAGMA_DISABLE_DEPRECATION_WARNINGS
+	return UE::Slate::CastToVector2f(GetNodeMaximumSize());
+PRAGMA_ENABLE_DEPRECATION_WARNINGS
 }
 
 FSlateRect SGraphNodeResizable::GetHitTestingBorder() const
diff --git a/Engine/Source/Editor/GraphEditor/Private/SNodePanel.cpp b/Engine/Source/Editor/GraphEditor/Private/SNodePanel.cpp
--- a/Engine/Source/Editor/GraphEditor/Private/SNodePanel.cpp
+++ b/Engine/Source/Editor/GraphEditor/Private/SNodePanel.cpp
@@ -518,13 +518,15 @@ FSlateRect SNodePanel::PanelRectToGraphRect( const FSlateRect& PanelSpaceRect )
 
 void SNodePanel::OnBeginNodeInteraction(const TSharedRef<SNode>& InNodeToDrag, const FVector2D& GrabOffset)
 {
-	OnBeginNodeInteraction(InNodeToDrag, UE::Slate::CastToVector2f(GrabOffset));
+	NodeUnderMousePtr = InNodeToDrag;
+	NodeGrabOffset = UE::Slate::CastToVector2f(GrabOffset);
 }
 
 void SNodePanel::OnBeginNodeInteraction(const TSharedRef<SNode>& InNodeToDrag, const FVector2f& GrabOffset)
 {
-	NodeUnderMousePtr = InNodeToDrag;
-	NodeGrabOffset = GrabOffset;
+PRAGMA_DISABLE_DEPRECATION_WARNINGS
+	OnBeginNodeInteraction(InNodeToDrag, FVector2D(GrabOffset));
+PRAGMA_ENABLE_DEPRECATION_WARNINGS
 }
 
 void SNodePanel::OnEndNodeInteraction(const TSharedRef<SNode>& InNodeToDrag)
diff --git a/Engine/Source/Editor/GraphEditor/Public/MaterialNodes/SGraphNodeMaterialConvert.h b/Engine/Source/Editor/GraphEditor/Public/MaterialNodes/SGraphNodeMaterialConvert.h
--- a/Engine/Source/Editor/GraphEditor/Public/MaterialNodes/SGraphNodeMaterialConvert.h
+++ b/Engine/Source/Editor/GraphEditor/Public/MaterialNodes/SGraphNodeMaterialConvert.h
@@ -5,6 +5,7 @@
 #include "Input/DragAndDrop.h"
 #include "MaterialNodes/SGraphNodeMaterialBase.h"
 #include "SGraphPin.h"
+#include "Types/SlateVector2.h"
 #include "Widgets/SCompoundWidget.h"
 
 class FConvertDragDropOp;
@@ -30,11 +31,11 @@ public:
 	//~ End FDragDropOperation interface
 
 	TSharedPtr<SConvertInnerPin> GetSourcePin() const;
-	FVector2D GetScreenPosition() const;
+	UE::Slate::FDeprecateVector2DResult GetScreenPosition() const;
 
 protected:
 	TSharedPtr<SConvertInnerPin> SourcePin;
-	FVector2D ScreenPosition;
+	FVector2f ScreenPosition;
 };
 
 ////////////////////////////////////////////////////////////////////////////////////////////////////////////////////////////////////////////
@@ -98,7 +99,7 @@ public:
 
 	const FSlateBrush* GetPinBrush() const;
 	void CancelDragDrop();
-	FVector2D GetPinCenterAbsolute() const;
+	UE::Slate::FDeprecateVector2DResult GetPinCenterAbsolute() const;
 
 	/** Adds a connection to WeakConnectedPins */
 	void AddConnection(TSharedPtr<SConvertInnerPin> InOtherPin);
@@ -133,7 +134,7 @@ protected:
 
 	TSharedPtr<SImage> PinImage;
 	TSet<TWeakPtr<SConvertInnerPin>> WeakConnectedPins;
-	mutable FVector2D CenterAbsolute;
+	mutable FVector2f CenterAbsolute;
 };
 
 ////////////////////////////////////////////////////////////////////////////////////////////////////////////////////////////////////////////
@@ -191,8 +192,8 @@ protected:
 		const FGeometry& InAllottedGeometry,
 		FSlateWindowElementList& OutDrawElements,
 		int32& InOutLayerId, 
-		const FVector2D InCurveStart, 
-		const FVector2D InCurveEnd
+		const FVector2f& InCurveStart, 
+		const FVector2f& InCurveEnd
 	) const;
 
 private:
