Code readibility fix for graph editor action menu

diff --git a/Engine/Source/Editor/GraphEditor/Private/SGraphEditorActionMenu.cpp b/Engine/Source/Editor/GraphEditor/Private/SGraphEditorActionMenu.cpp
--- a/Engine/Source/Editor/GraphEditor/Private/SGraphEditorActionMenu.cpp
+++ b/Engine/Source/Editor/GraphEditor/Private/SGraphEditorActionMenu.cpp
@@ -53,7 +53,9 @@ void SGraphEditorActionMenu::RefreshAllActions()
 {
 	if (GraphActionMenu)
 	{
-		GraphActionMenu->RefreshAllActions(true, false);
+		constexpr bool bPreserveExpansion = true;
+		constexpr bool bHandleOnSelectionEvent = false;
+		GraphActionMenu->RefreshAllActions(bPreserveExpansion, bHandleOnSelectionEvent);
 	}
 }
 
