Updated headers using UnrealCodeFixup to make sure dllstorage is on methods/staticvars instead of types. Part 5/n

diff --git a/Engine/Source/Editor/GraphEditor/Public/KismetNodes/SGraphNodeK2Base.h b/Engine/Source/Editor/GraphEditor/Public/KismetNodes/SGraphNodeK2Base.h
--- a/Engine/Source/Editor/GraphEditor/Public/KismetNodes/SGraphNodeK2Base.h
+++ b/Engine/Source/Editor/GraphEditor/Public/KismetNodes/SGraphNodeK2Base.h
@@ -11,6 +11,8 @@
 #include "SNodePanel.h"
 #include "Templates/SharedPointer.h"
 
+#define UE_API GRAPHEDITOR_API
+
 class SToolTip;
 class UObject;
 struct FGraphInformationPopupInfo;
@@ -19,50 +21,52 @@ struct FNodeInfoContext;
 struct FOverlayBrushInfo;
 struct FSlateBrush;
 
-class GRAPHEDITOR_API SGraphNodeK2Base : public SGraphNode
+class SGraphNodeK2Base : public SGraphNode
 {
 public:
 	SLATE_BEGIN_ARGS(SGraphNodeK2Base) { }
 		SLATE_ARGUMENT(TSharedPtr<ISlateStyle>, Style)
 	SLATE_END_ARGS()
 
-	void Construct(const FArguments& InArgs);
+	UE_API void Construct(const FArguments& InArgs);
 
 	//~ Begin SGraphNode Interface
-	virtual void UpdateGraphNode() override;
+	UE_API virtual void UpdateGraphNode() override;
 
 	//~ Begin SNodePanel::SNode Interface
-	virtual void GetOverlayBrushes(bool bSelected, const FVector2f& WidgetSize, TArray<FOverlayBrushInfo>& Brushes) const override;
-	virtual void GetNodeInfoPopups(FNodeInfoContext* Context, TArray<FGraphInformationPopupInfo>& Popups) const override;
-	virtual const FSlateBrush* GetShadowBrush(bool bSelected) const override;
-	virtual void GetDiffHighlightBrushes(const FSlateBrush*& BackgroundOut, const FSlateBrush*& ForegroundOut) const override;
-	void PerformSecondPassLayout(const TMap< UObject*, TSharedRef<SNode> >& NodeToWidgetLookup) const override;
+	UE_API virtual void GetOverlayBrushes(bool bSelected, const FVector2f& WidgetSize, TArray<FOverlayBrushInfo>& Brushes) const override;
+	UE_API virtual void GetNodeInfoPopups(FNodeInfoContext* Context, TArray<FGraphInformationPopupInfo>& Popups) const override;
+	UE_API virtual const FSlateBrush* GetShadowBrush(bool bSelected) const override;
+	UE_API virtual void GetDiffHighlightBrushes(const FSlateBrush*& BackgroundOut, const FSlateBrush*& ForegroundOut) const override;
+	UE_API void PerformSecondPassLayout(const TMap< UObject*, TSharedRef<SNode> >& NodeToWidgetLookup) const override;
 
 protected :
 	//~ Begin SGraphNode Interface
-	virtual TSharedPtr<SToolTip> GetComplexTooltip() override;
+	UE_API virtual TSharedPtr<SToolTip> GetComplexTooltip() override;
 	//~ End SGraphNode Interface
 
 	/** Set up node in 'standard' mode */
-	void UpdateStandardNode();
+	UE_API void UpdateStandardNode();
 	/** Set up node in 'compact' mode */
-	void UpdateCompactNode();
+	UE_API void UpdateCompactNode();
 
 	/** Get title in compact mode */
-	FText GetNodeCompactTitle() const;
+	UE_API FText GetNodeCompactTitle() const;
 
 	/** Retrieves text to tack on to the top of the tooltip (above the standard text) */
-	FText GetToolTipHeading() const;
+	UE_API FText GetToolTipHeading() const;
 
 	/** Returns the slate style to use for this node */
-	const ISlateStyle& GetStyleSet() const;
+	UE_API const ISlateStyle& GetStyleSet() const;
 
 	/** The slate style to use. We'll fall back to App Style if it's unset. */
 	TSharedPtr<ISlateStyle> Style;
 
 protected:
-	static const FLinearColor BreakpointHitColor;
-	static const FLinearColor LatentBubbleColor;
-	static const FLinearColor TimelineBubbleColor;
-	static const FLinearColor PinnedWatchColor;
+	static UE_API const FLinearColor BreakpointHitColor;
+	static UE_API const FLinearColor LatentBubbleColor;
+	static UE_API const FLinearColor TimelineBubbleColor;
+	static UE_API const FLinearColor PinnedWatchColor;
 };
+
+#undef UE_API
diff --git a/Engine/Source/Editor/GraphEditor/Public/KismetNodes/SGraphNodeK2Composite.h b/Engine/Source/Editor/GraphEditor/Public/KismetNodes/SGraphNodeK2Composite.h
--- a/Engine/Source/Editor/GraphEditor/Public/KismetNodes/SGraphNodeK2Composite.h
+++ b/Engine/Source/Editor/GraphEditor/Public/KismetNodes/SGraphNodeK2Composite.h
@@ -9,28 +9,32 @@
 #include "Widgets/DeclarativeSyntaxSupport.h"
 #include "Widgets/SWidget.h"
 
+#define UE_API GRAPHEDITOR_API
+
 class SToolTip;
 class SWidget;
 class UEdGraph;
 class UK2Node_Composite;
 
-class GRAPHEDITOR_API SGraphNodeK2Composite : public SGraphNodeK2Base
+class SGraphNodeK2Composite : public SGraphNodeK2Base
 {
 public:
 	SLATE_BEGIN_ARGS(SGraphNodeK2Composite){}
 	SLATE_END_ARGS()
 
-	void Construct(const FArguments& InArgs, UK2Node_Composite* InNode);
+	UE_API void Construct(const FArguments& InArgs, UK2Node_Composite* InNode);
 
 	// SGraphNode interface
-	 virtual void UpdateGraphNode() override;
-	virtual TSharedPtr<SToolTip> GetComplexTooltip() override;
+	 UE_API virtual void UpdateGraphNode() override;
+	UE_API virtual TSharedPtr<SToolTip> GetComplexTooltip() override;
 	// End of SGraphNode interface
 protected:
-	virtual UEdGraph* GetInnerGraph() const;
+	UE_API virtual UEdGraph* GetInnerGraph() const;
 
-	FText GetPreviewCornerText() const;
-	FText GetTooltipTextForNode() const;
+	UE_API FText GetPreviewCornerText() const;
+	UE_API FText GetTooltipTextForNode() const;
 
-	virtual TSharedRef<SWidget> CreateNodeBody();
+	UE_API virtual TSharedRef<SWidget> CreateNodeBody();
 };
+
+#undef UE_API
diff --git a/Engine/Source/Editor/GraphEditor/Public/KismetNodes/SGraphNodeK2Default.h b/Engine/Source/Editor/GraphEditor/Public/KismetNodes/SGraphNodeK2Default.h
--- a/Engine/Source/Editor/GraphEditor/Public/KismetNodes/SGraphNodeK2Default.h
+++ b/Engine/Source/Editor/GraphEditor/Public/KismetNodes/SGraphNodeK2Default.h
@@ -6,13 +6,17 @@
 #include "KismetNodes/SGraphNodeK2Base.h"
 #include "Widgets/DeclarativeSyntaxSupport.h"
 
+#define UE_API GRAPHEDITOR_API
+
 class UK2Node;
 
-class GRAPHEDITOR_API SGraphNodeK2Default : public SGraphNodeK2Base
+class SGraphNodeK2Default : public SGraphNodeK2Base
 {
 public:
 	SLATE_BEGIN_ARGS(SGraphNodeK2Default){}
 	SLATE_END_ARGS()
 
-	void Construct(const FArguments& InArgs, UK2Node* InNode);
+	UE_API void Construct(const FArguments& InArgs, UK2Node* InNode);
 };
+
+#undef UE_API
diff --git a/Engine/Source/Editor/GraphEditor/Public/KismetNodes/SGraphNodeK2Event.h b/Engine/Source/Editor/GraphEditor/Public/KismetNodes/SGraphNodeK2Event.h
--- a/Engine/Source/Editor/GraphEditor/Public/KismetNodes/SGraphNodeK2Event.h
+++ b/Engine/Source/Editor/GraphEditor/Public/KismetNodes/SGraphNodeK2Event.h
@@ -10,20 +10,22 @@
 #include "Widgets/SOverlay.h"
 #include "Widgets/SWidget.h"
 
+#define UE_API GRAPHEDITOR_API
+
 class SGraphPin;
 class SNodeTitle;
 class SOverlay;
 class SWidget;
 
-class GRAPHEDITOR_API SGraphNodeK2Event : public SGraphNodeK2Default
+class SGraphNodeK2Event : public SGraphNodeK2Default
 {
 public:
 	SGraphNodeK2Event() : SGraphNodeK2Default(), bHasDelegateOutputPin(false) {}
 
 protected:
-	virtual TSharedRef<SWidget> CreateTitleWidget(TSharedPtr<SNodeTitle> NodeTitle) override;
-	virtual bool UseLowDetailNodeTitles() const override;
-	virtual void AddPin( const TSharedRef<SGraphPin>& PinToAdd ) override;
+	UE_API virtual TSharedRef<SWidget> CreateTitleWidget(TSharedPtr<SNodeTitle> NodeTitle) override;
+	UE_API virtual bool UseLowDetailNodeTitles() const override;
+	UE_API virtual void AddPin( const TSharedRef<SGraphPin>& PinToAdd ) override;
 
 
 	virtual void SetDefaultTitleAreaWidget(TSharedRef<SOverlay> DefaultTitleAreaWidget) override
@@ -37,8 +39,10 @@ private:
 		return SGraphNodeK2Default::UseLowDetailNodeTitles();
 	}
 
-	EVisibility GetTitleVisibility() const;
+	UE_API EVisibility GetTitleVisibility() const;
 
 	TSharedPtr<SOverlay> TitleAreaWidget;
 	bool bHasDelegateOutputPin;
 };
+
+#undef UE_API
diff --git a/Engine/Source/Editor/GraphEditor/Public/KismetNodes/SGraphNodeK2Var.h b/Engine/Source/Editor/GraphEditor/Public/KismetNodes/SGraphNodeK2Var.h
--- a/Engine/Source/Editor/GraphEditor/Public/KismetNodes/SGraphNodeK2Var.h
+++ b/Engine/Source/Editor/GraphEditor/Public/KismetNodes/SGraphNodeK2Var.h
@@ -9,31 +9,35 @@
 #include "Types/SlateEnums.h"
 #include "Widgets/DeclarativeSyntaxSupport.h"
 
+#define UE_API GRAPHEDITOR_API
+
 class FText;
 class SWidget;
 class UK2Node;
 struct FMargin;
 struct FSlateBrush;
 
-class GRAPHEDITOR_API SGraphNodeK2Var : public SGraphNodeK2Base
+class SGraphNodeK2Var : public SGraphNodeK2Base
 {
 public:
 	SLATE_BEGIN_ARGS(SGraphNodeK2Var){}
 	SLATE_END_ARGS()
 
-	void Construct( const FArguments& InArgs, UK2Node* InNode );
+	UE_API void Construct( const FArguments& InArgs, UK2Node* InNode );
 
-	virtual const FSlateBrush* GetShadowBrush(bool bSelected) const override;
+	UE_API virtual const FSlateBrush* GetShadowBrush(bool bSelected) const override;
 
-	virtual void GetDiffHighlightBrushes(const FSlateBrush*& BackgroundOut, const FSlateBrush*& ForegroundOut) const override;
+	UE_API virtual void GetDiffHighlightBrushes(const FSlateBrush*& BackgroundOut, const FSlateBrush*& ForegroundOut) const override;
 
 	// SGraphNode interface
-	virtual void UpdateGraphNode() override;
+	UE_API virtual void UpdateGraphNode() override;
 	// End of SGraphNode interface
 
 protected:
-	FSlateColor GetVariableColor() const;
+	UE_API FSlateColor GetVariableColor() const;
 
 	// Allow derived classes to override title widget
-	virtual TSharedRef<SWidget> UpdateTitleWidget(FText InTitleText, TSharedPtr<SWidget> InTitleWidget, EHorizontalAlignment& InOutTitleHAlign, FMargin& InOutTitleMargin) const;
+	UE_API virtual TSharedRef<SWidget> UpdateTitleWidget(FText InTitleText, TSharedPtr<SWidget> InTitleWidget, EHorizontalAlignment& InOutTitleHAlign, FMargin& InOutTitleMargin) const;
 };
+
+#undef UE_API
diff --git a/Engine/Source/Editor/GraphEditor/Public/KismetPins/SGraphPinBool.h b/Engine/Source/Editor/GraphEditor/Public/KismetPins/SGraphPinBool.h
--- a/Engine/Source/Editor/GraphEditor/Public/KismetPins/SGraphPinBool.h
+++ b/Engine/Source/Editor/GraphEditor/Public/KismetPins/SGraphPinBool.h
@@ -9,25 +9,29 @@
 #include "Widgets/DeclarativeSyntaxSupport.h"
 #include "Widgets/SWidget.h"
 
+#define UE_API GRAPHEDITOR_API
+
 class SWidget;
 class UEdGraphPin;
 
-class GRAPHEDITOR_API SGraphPinBool : public SGraphPin
+class SGraphPinBool : public SGraphPin
 {
 public:
 	SLATE_BEGIN_ARGS(SGraphPinBool) {}
 	SLATE_END_ARGS()
 
-	void Construct(const FArguments& InArgs, UEdGraphPin* InGraphPinObj);
+	UE_API void Construct(const FArguments& InArgs, UEdGraphPin* InGraphPinObj);
 
 protected:
 	//~ Begin SGraphPin Interface
-	virtual TSharedRef<SWidget>	GetDefaultValueWidget() override;
+	UE_API virtual TSharedRef<SWidget>	GetDefaultValueWidget() override;
 	//~ End SGraphPin Interface
 
 	/** Determine if the check box should be checked or not */
-	ECheckBoxState IsDefaultValueChecked() const;
+	UE_API ECheckBoxState IsDefaultValueChecked() const;
 
 	/** Called when check box is changed */
-	void OnDefaultValueCheckBoxChanged( ECheckBoxState InIsChecked );
+	UE_API void OnDefaultValueCheckBoxChanged( ECheckBoxState InIsChecked );
 };
+
+#undef UE_API
diff --git a/Engine/Source/Editor/GraphEditor/Public/KismetPins/SGraphPinClass.h b/Engine/Source/Editor/GraphEditor/Public/KismetPins/SGraphPinClass.h
--- a/Engine/Source/Editor/GraphEditor/Public/KismetPins/SGraphPinClass.h
+++ b/Engine/Source/Editor/GraphEditor/Public/KismetPins/SGraphPinClass.h
@@ -12,6 +12,8 @@
 #include "Widgets/DeclarativeSyntaxSupport.h"
 #include "Widgets/SWidget.h"
 
+#define UE_API GRAPHEDITOR_API
+
 class SWidget;
 class UClass;
 class UEdGraphPin;
@@ -19,26 +21,26 @@ class UEdGraphPin;
 /////////////////////////////////////////////////////
 // SGraphPinClass
 
-class GRAPHEDITOR_API SGraphPinClass : public SGraphPinObject
+class SGraphPinClass : public SGraphPinObject
 {
 public:
 	SLATE_BEGIN_ARGS(SGraphPinClass) {}
 	SLATE_END_ARGS()
 
-	void Construct(const FArguments& InArgs, UEdGraphPin* InGraphPinObj);
+	UE_API void Construct(const FArguments& InArgs, UEdGraphPin* InGraphPinObj);
 	void SetAllowAbstractClasses(bool bAllow) { bAllowAbstractClasses = bAllow; }
 protected:
 
 	// Called when a new class was picked via the asset picker
-	void OnPickedNewClass(UClass* ChosenClass);
+	UE_API void OnPickedNewClass(UClass* ChosenClass);
 
 	//~ Begin SGraphPinObject Interface
-	virtual FReply OnClickUse() override;
+	UE_API virtual FReply OnClickUse() override;
 	virtual bool AllowSelfPinWidget() const override { return false; }
-	virtual TSharedRef<SWidget> GenerateAssetPicker() override;
-	virtual FText GetDefaultComboText() const override;
-	virtual FOnClicked GetOnUseButtonDelegate() override;
-	virtual const FAssetData& GetAssetData(bool bRuntimePath) const override;
+	UE_API virtual TSharedRef<SWidget> GenerateAssetPicker() override;
+	UE_API virtual FText GetDefaultComboText() const override;
+	UE_API virtual FOnClicked GetOnUseButtonDelegate() override;
+	UE_API virtual const FAssetData& GetAssetData(bool bRuntimePath) const override;
 	//~ End SGraphPinObject Interface
 
 	/** Cached AssetData without the _C */
@@ -47,3 +49,5 @@ protected:
 	/** Whether abstract classes should be filtered out in the class viewer */
 	bool bAllowAbstractClasses;
 };
+
+#undef UE_API
diff --git a/Engine/Source/Editor/GraphEditor/Public/KismetPins/SGraphPinColor.h b/Engine/Source/Editor/GraphEditor/Public/KismetPins/SGraphPinColor.h
--- a/Engine/Source/Editor/GraphEditor/Public/KismetPins/SGraphPinColor.h
+++ b/Engine/Source/Editor/GraphEditor/Public/KismetPins/SGraphPinColor.h
@@ -11,34 +11,38 @@
 #include "Widgets/DeclarativeSyntaxSupport.h"
 #include "Widgets/SWidget.h"
 
+#define UE_API GRAPHEDITOR_API
+
 class SWidget;
 class UEdGraphPin;
 struct FGeometry;
 struct FPointerEvent;
 
-class GRAPHEDITOR_API SGraphPinColor : public SGraphPin
+class SGraphPinColor : public SGraphPin
 {
 public:
 	SLATE_BEGIN_ARGS(SGraphPinColor) {}
 	SLATE_END_ARGS()
 
-	void Construct(const FArguments& InArgs, UEdGraphPin* InGraphPinObj);
+	UE_API void Construct(const FArguments& InArgs, UEdGraphPin* InGraphPinObj);
 
 	/** Return the current color value stored in the pin */
-	FLinearColor GetColor () const;
+	UE_API FLinearColor GetColor () const;
 
 protected:
 	//~ Begin SGraphPin Interface
-	virtual TSharedRef<SWidget>	GetDefaultValueWidget() override;
+	UE_API virtual TSharedRef<SWidget>	GetDefaultValueWidget() override;
 	//~ End SGraphPin Interface
 
 	/** Called when clicking on the color button */
-	FReply OnColorBoxClicked(const FGeometry& MyGeometry, const FPointerEvent& MouseEvent);
+	UE_API FReply OnColorBoxClicked(const FGeometry& MyGeometry, const FPointerEvent& MouseEvent);
 
 	/** Called when the color picker commits a color value */
-	void OnColorCommitted(FLinearColor color);
+	UE_API void OnColorCommitted(FLinearColor color);
 
 
 private:
 	TSharedPtr<SWidget> DefaultValueWidget;
 };
+
+#undef UE_API
diff --git a/Engine/Source/Editor/GraphEditor/Public/KismetPins/SGraphPinEnum.h b/Engine/Source/Editor/GraphEditor/Public/KismetPins/SGraphPinEnum.h
--- a/Engine/Source/Editor/GraphEditor/Public/KismetPins/SGraphPinEnum.h
+++ b/Engine/Source/Editor/GraphEditor/Public/KismetPins/SGraphPinEnum.h
@@ -11,16 +11,18 @@
 #include "Types/SlateEnums.h"
 #include "Widgets/DeclarativeSyntaxSupport.h"
 
+#define UE_API GRAPHEDITOR_API
+
 class SWidget;
 class UEdGraphPin;
 
-class GRAPHEDITOR_API SGraphPinEnum : public SGraphPin
+class SGraphPinEnum : public SGraphPin
 {
 public:
 	SLATE_BEGIN_ARGS(SGraphPinEnum) {}
 	SLATE_END_ARGS()
 
-	void Construct(const FArguments& InArgs, UEdGraphPin* InGraphPinObj);
+	UE_API void Construct(const FArguments& InArgs, UEdGraphPin* InGraphPinObj);
 
 protected:
 
@@ -29,7 +31,7 @@ protected:
 	 *
 	 *	@return Reference to the newly created widget object
 	 */
-	virtual TSharedRef<SWidget>	GetDefaultValueWidget() override;
+	UE_API virtual TSharedRef<SWidget>	GetDefaultValueWidget() override;
 
 private:
 
@@ -38,14 +40,14 @@ private:
 	 *
 	 *	@return currently selected string
 	 */
-	FString OnGetText() const;
+	UE_API FString OnGetText() const;
 
 	/**
 	 *	Function to generate the list of indexes from the enum object
 	 *
 	 *	@param OutComboBoxIndexes - Int array reference to store the list of indexes
 	 */
-	void GenerateComboBoxIndexes( TArray< TSharedPtr<int32> >& OutComboBoxIndexes );
+	UE_API void GenerateComboBoxIndexes( TArray< TSharedPtr<int32> >& OutComboBoxIndexes );
 
 	/**
 	 *	Function to set the newly selected index
@@ -53,21 +55,23 @@ private:
 	 * @param NewSelection The newly selected item in the combo box
 	 * @param SelectInfo Provides context on how the selection changed
 	 */
-	void ComboBoxSelectionChanged( TSharedPtr<int32> NewSelection, ESelectInfo::Type SelectInfo );
+	UE_API void ComboBoxSelectionChanged( TSharedPtr<int32> NewSelection, ESelectInfo::Type SelectInfo );
 
 	/**
 	 * Returns the friendly name of the enum at index EnumIndex
 	 *
 	 * @param EnumIndex	- The index of the enum to return the friendly name for
 	 */
-	FText OnGetFriendlyName(int32 EnumIndex);
+	UE_API FText OnGetFriendlyName(int32 EnumIndex);
 
 	/**
 	 * Returns the tooltip of the enum at index EnumIndex
 	 *
 	 * @param EnumIndex	- The index of the enum to return the tooltip for
 	 */
-	FText OnGetTooltip(int32 EnumIndex);
+	UE_API FText OnGetTooltip(int32 EnumIndex);
 
 	TSharedPtr<class SPinComboBox>	ComboBox;
 };
+
+#undef UE_API
diff --git a/Engine/Source/Editor/GraphEditor/Public/KismetPins/SGraphPinExec.h b/Engine/Source/Editor/GraphEditor/Public/KismetPins/SGraphPinExec.h
--- a/Engine/Source/Editor/GraphEditor/Public/KismetPins/SGraphPinExec.h
+++ b/Engine/Source/Editor/GraphEditor/Public/KismetPins/SGraphPinExec.h
@@ -8,26 +8,30 @@
 #include "Widgets/DeclarativeSyntaxSupport.h"
 #include "Widgets/SWidget.h"
 
+#define UE_API GRAPHEDITOR_API
+
 class SWidget;
 class UEdGraphPin;
 struct FSlateBrush;
 
-class GRAPHEDITOR_API SGraphPinExec : public SGraphPin
+class SGraphPinExec : public SGraphPin
 {
 public:
 	SLATE_BEGIN_ARGS(SGraphPinExec)	{}
 	SLATE_END_ARGS()
 
-	void Construct(const FArguments& InArgs, UEdGraphPin* InPin);
+	UE_API void Construct(const FArguments& InArgs, UEdGraphPin* InPin);
 
 protected:
 	//~ Begin SGraphPin Interface
-	virtual TSharedRef<SWidget>	GetDefaultValueWidget() override;
-	virtual const FSlateBrush* GetPinIcon() const override;
+	UE_API virtual TSharedRef<SWidget>	GetDefaultValueWidget() override;
+	UE_API virtual const FSlateBrush* GetPinIcon() const override;
 	//~ End SGraphPin Interface
 
-	void CachePinIcons();
+	UE_API void CachePinIcons();
 
 	const FSlateBrush* CachedImg_Pin_ConnectedHovered;
 	const FSlateBrush* CachedImg_Pin_DisconnectedHovered;
 };
+
+#undef UE_API
diff --git a/Engine/Source/Editor/GraphEditor/Public/KismetPins/SGraphPinInteger.h b/Engine/Source/Editor/GraphEditor/Public/KismetPins/SGraphPinInteger.h
--- a/Engine/Source/Editor/GraphEditor/Public/KismetPins/SGraphPinInteger.h
+++ b/Engine/Source/Editor/GraphEditor/Public/KismetPins/SGraphPinInteger.h
@@ -13,19 +13,23 @@
 #include "Widgets/DeclarativeSyntaxSupport.h"
 #include "Widgets/SWidget.h"
 
+#define UE_API GRAPHEDITOR_API
+
 class SWidget;
 class UEdGraphPin;
 
-class GRAPHEDITOR_API SGraphPinInteger : public SGraphPinNum<int32>
+class SGraphPinInteger : public SGraphPinNum<int32>
 {
 public:
 	SLATE_BEGIN_ARGS(SGraphPinInteger) {}
 	SLATE_END_ARGS()
 
-	void Construct(const FArguments& InArgs, UEdGraphPin* InGraphPinObj);
+	UE_API void Construct(const FArguments& InArgs, UEdGraphPin* InGraphPinObj);
 
 protected:
 	//~ Begin SGraphPinString Interface
-	virtual TSharedRef<SWidget>	GetDefaultValueWidget() override;
+	UE_API virtual TSharedRef<SWidget>	GetDefaultValueWidget() override;
 	//~ End SGraphPinString Interface
 };
+
+#undef UE_API
diff --git a/Engine/Source/Editor/GraphEditor/Public/KismetPins/SGraphPinIntegerSlider.h b/Engine/Source/Editor/GraphEditor/Public/KismetPins/SGraphPinIntegerSlider.h
--- a/Engine/Source/Editor/GraphEditor/Public/KismetPins/SGraphPinIntegerSlider.h
+++ b/Engine/Source/Editor/GraphEditor/Public/KismetPins/SGraphPinIntegerSlider.h
@@ -13,19 +13,23 @@
 #include "Widgets/DeclarativeSyntaxSupport.h"
 #include "Widgets/SWidget.h"
 
+#define UE_API GRAPHEDITOR_API
+
 class SWidget;
 class UEdGraphPin;
 
-class GRAPHEDITOR_API SGraphPinIntegerSlider : public SGraphPinNumSlider<int32>
+class SGraphPinIntegerSlider : public SGraphPinNumSlider<int32>
 {
 public:
 	SLATE_BEGIN_ARGS(SGraphPinIntegerSlider) {}
 	SLATE_END_ARGS()
 
-	void Construct(const FArguments& InArgs, UEdGraphPin* InGraphPinObj, FProperty* InProperty);
+	UE_API void Construct(const FArguments& InArgs, UEdGraphPin* InGraphPinObj, FProperty* InProperty);
 
 protected:
 	//~ Begin SGraphPinString Interface
-	virtual TSharedRef<SWidget>	GetDefaultValueWidget() override;
+	UE_API virtual TSharedRef<SWidget>	GetDefaultValueWidget() override;
 	//~ End SGraphPinString Interface
 };
+
+#undef UE_API
diff --git a/Engine/Source/Editor/GraphEditor/Public/KismetPins/SGraphPinObject.h b/Engine/Source/Editor/GraphEditor/Public/KismetPins/SGraphPinObject.h
--- a/Engine/Source/Editor/GraphEditor/Public/KismetPins/SGraphPinObject.h
+++ b/Engine/Source/Editor/GraphEditor/Public/KismetPins/SGraphPinObject.h
@@ -16,65 +16,67 @@
 #include "Widgets/Input/SMenuAnchor.h"
 #include "Widgets/SWidget.h"
 
+#define UE_API GRAPHEDITOR_API
+
 class SButton;
 class SWidget;
 class UEdGraphPin;
 
-class GRAPHEDITOR_API SGraphPinObject : public SGraphPin
+class SGraphPinObject : public SGraphPin
 {
 public:
 	SLATE_BEGIN_ARGS(SGraphPinObject) {}
 	SLATE_END_ARGS()
 
-	void Construct(const FArguments& InArgs, UEdGraphPin* InGraphPinObj);
+	UE_API void Construct(const FArguments& InArgs, UEdGraphPin* InGraphPinObj);
 
 protected:
 	//~ Begin SGraphPin Interface
-	virtual TSharedRef<SWidget>	GetDefaultValueWidget() override;
+	UE_API virtual TSharedRef<SWidget>	GetDefaultValueWidget() override;
 	virtual bool DoesWidgetHandleSettingEditingEnabled() const override { return true; }
 	//~ End SGraphPin Interface
 
 	/** Delegate to be called when the use current selected item in asset browser button is clicked */
-	virtual FOnClicked GetOnUseButtonDelegate();
+	UE_API virtual FOnClicked GetOnUseButtonDelegate();
 	/** Delegate to be called when the browse for item button is clicked */
-	virtual FOnClicked GetOnBrowseButtonDelegate();
+	UE_API virtual FOnClicked GetOnBrowseButtonDelegate();
 
 	/** Clicked Use button */
-	virtual FReply OnClickUse();
+	UE_API virtual FReply OnClickUse();
 	/** Clicked Browse button */
-	virtual FReply OnClickBrowse();
+	UE_API virtual FReply OnClickBrowse();
 	/** Get text tooltip for object */
-	FText GetObjectToolTip() const;
+	UE_API FText GetObjectToolTip() const;
 	/** Get text tooltip for object */
-	FString GetObjectToolTipAsString() const;
+	UE_API FString GetObjectToolTipAsString() const;
 	/** Get string value for object */
-	FText GetValue() const;
+	UE_API FText GetValue() const;
 	/** Get name of the object*/
-	FText GetObjectName() const;
+	UE_API FText GetObjectName() const;
 	/** Get default text for the picker combo */
-	virtual FText GetDefaultComboText() const;
+	UE_API virtual FText GetDefaultComboText() const;
 	/** Allow self pin widget */
 	virtual bool AllowSelfPinWidget() const { return true; }
 	/** True if this specific pin should be treated as a self pin */
-	virtual bool ShouldDisplayAsSelfPin() const;
+	UE_API virtual bool ShouldDisplayAsSelfPin() const;
 	/** Generate asset picker window */
-	virtual TSharedRef<SWidget> GenerateAssetPicker();
+	UE_API virtual TSharedRef<SWidget> GenerateAssetPicker();
 	/** Called to validate selection from picker window */
-	virtual void OnAssetSelectedFromPicker(const struct FAssetData& AssetData);
+	UE_API virtual void OnAssetSelectedFromPicker(const struct FAssetData& AssetData);
 	/** Called when enter is pressed when items are selected in the picker window */
-	void OnAssetEnterPressedInPicker(const TArray<FAssetData>& InSelectedAssets);
+	UE_API void OnAssetEnterPressedInPicker(const TArray<FAssetData>& InSelectedAssets);
 
 	/** Used to update the combo button text */
-	FText OnGetComboTextValue() const;
+	UE_API FText OnGetComboTextValue() const;
 	/** Combo Button Color and Opacity delegate */
-	FSlateColor OnGetComboForeground() const;
+	UE_API FSlateColor OnGetComboForeground() const;
 	/** Button Color and Opacity delegate */
-	FSlateColor OnGetWidgetForeground() const;
+	UE_API FSlateColor OnGetWidgetForeground() const;
 	/** Button Color and Opacity delegate */
-	FSlateColor OnGetWidgetBackground() const;
+	UE_API FSlateColor OnGetWidgetBackground() const;
 
 	/** Returns asset data of currently selected object, if bRuntimePath is true this will include _C for blueprint classes, for false it will point to UBlueprint instead */
-	virtual const FAssetData& GetAssetData(bool bRuntimePath) const;
+	UE_API virtual const FAssetData& GetAssetData(bool bRuntimePath) const;
 
 protected:
 	/** Object manipulator buttons. */
@@ -87,3 +89,5 @@ protected:
 	/** Cached AssetData of object selected */
 	mutable FAssetData CachedAssetData;
 };
+
+#undef UE_API
diff --git a/Engine/Source/Editor/GraphEditor/Public/KismetPins/SGraphPinString.h b/Engine/Source/Editor/GraphEditor/Public/KismetPins/SGraphPinString.h
--- a/Engine/Source/Editor/GraphEditor/Public/KismetPins/SGraphPinString.h
+++ b/Engine/Source/Editor/GraphEditor/Public/KismetPins/SGraphPinString.h
@@ -10,32 +10,36 @@
 #include "Widgets/DeclarativeSyntaxSupport.h"
 #include "Widgets/SWidget.h"
 
+#define UE_API GRAPHEDITOR_API
+
 class SWidget;
 class UEdGraphPin;
 
-class GRAPHEDITOR_API SGraphPinString : public SGraphPin
+class SGraphPinString : public SGraphPin
 {
 public:
 	SLATE_BEGIN_ARGS(SGraphPinString) {}
 	SLATE_END_ARGS()
 
-		void Construct(const FArguments& InArgs, UEdGraphPin* InGraphPinObj);
+		UE_API void Construct(const FArguments& InArgs, UEdGraphPin* InGraphPinObj);
 
 protected:
 	//~ Begin SGraphPin Interface
-	virtual TSharedRef<SWidget>	GetDefaultValueWidget() override;
+	UE_API virtual TSharedRef<SWidget>	GetDefaultValueWidget() override;
 	//~ End SGraphPin Interface
 
-	FText GetTypeInValue() const;
-	virtual void SetTypeInValue(const FText& NewTypeInValue, ETextCommit::Type CommitInfo);
+	UE_API FText GetTypeInValue() const;
+	UE_API virtual void SetTypeInValue(const FText& NewTypeInValue, ETextCommit::Type CommitInfo);
 
 	/** @return True if the pin default value field is read-only */
-	bool GetDefaultValueIsReadOnly() const;
+	UE_API bool GetDefaultValueIsReadOnly() const;
 
 private:
-	TSharedRef<SWidget> GenerateComboBoxEntry(TSharedPtr<FString> Value);
-	void HandleComboBoxSelectionChanged(TSharedPtr<FString> Value, ESelectInfo::Type InSelectInfo);
-	TSharedPtr<SWidget> TryBuildComboBoxWidget();
+	UE_API TSharedRef<SWidget> GenerateComboBoxEntry(TSharedPtr<FString> Value);
+	UE_API void HandleComboBoxSelectionChanged(TSharedPtr<FString> Value, ESelectInfo::Type InSelectInfo);
+	UE_API TSharedPtr<SWidget> TryBuildComboBoxWidget();
 
 	TArray<TSharedPtr<FString>> ComboBoxOptions;
 };
+
+#undef UE_API
diff --git a/Engine/Source/Editor/GraphEditor/Public/KismetPins/SGraphPinStructInstance.h b/Engine/Source/Editor/GraphEditor/Public/KismetPins/SGraphPinStructInstance.h
--- a/Engine/Source/Editor/GraphEditor/Public/KismetPins/SGraphPinStructInstance.h
+++ b/Engine/Source/Editor/GraphEditor/Public/KismetPins/SGraphPinStructInstance.h
@@ -15,6 +15,8 @@
 
 #include "SGraphPinStructInstance.generated.h"
 
+#define UE_API GRAPHEDITOR_API
+
 class FStructOnScope;
 class IStructureDetailsView;
 class SComboButton;
@@ -29,7 +31,7 @@ struct FPropertyChangedEvent;
  * Then, pass in StructName::StaticStruct() when creating SGraphPinStructInstance from a pin factory
  */
 USTRUCT()
-struct GRAPHEDITOR_API FPinStructEditWrapper
+struct FPinStructEditWrapper
 {
 	GENERATED_BODY()
 
@@ -69,7 +71,7 @@ struct FDataTypeEditWrapper : public FPinStructEditWrapper
  * It can be used directly by passing in a StructEditWrapper parameter when creating from a pin factory
  * Or, it can be subclassed for other types of struct display
  */
-class GRAPHEDITOR_API SGraphPinStructInstance : public SGraphPin
+class SGraphPinStructInstance : public SGraphPin
 {
 public:
 	SLATE_BEGIN_ARGS(SGraphPinStructInstance) 
@@ -78,37 +80,37 @@ public:
 		SLATE_ARGUMENT(const UScriptStruct*, StructEditWrapper)
 	SLATE_END_ARGS()
 
-	void Construct(const FArguments& InArgs, UEdGraphPin* InGraphPinObj);
+	UE_API void Construct(const FArguments& InArgs, UEdGraphPin* InGraphPinObj);
 
 protected:
 
 	//~ Begin SGraphPin Interface
-	virtual TSharedRef<SWidget>	GetDefaultValueWidget() override;
+	UE_API virtual TSharedRef<SWidget>	GetDefaultValueWidget() override;
 	//~ End SGraphPin Interface
 
 	/** Parses the struct ata from the pin to fill in the struct instance */
-	virtual void ParseDefaultValueData();
+	UE_API virtual void ParseDefaultValueData();
 
 	/** Call to push changes from edit instance to pin */
-	virtual void SaveDefaultValueData();
+	UE_API virtual void SaveDefaultValueData();
 
 	/** Refreshes cached description and edit data after an edit change */
-	virtual void RefreshCachedData();
+	UE_API virtual void RefreshCachedData();
 
 	/** Creates widget used to edit the struct instance */
-	virtual TSharedRef<SWidget> GetEditContent();
+	UE_API virtual TSharedRef<SWidget> GetEditContent();
 
 	/** Creates widget for displaying preview on the pin */
-	virtual TSharedRef<SWidget> GetDescriptionContent();
+	UE_API virtual TSharedRef<SWidget> GetDescriptionContent();
 
 	/** Slate accessor to shows cached description value */
-	virtual FText GetCachedDescriptionText() const;
+	UE_API virtual FText GetCachedDescriptionText() const;
 
 	/** Returns the base instance inside EditStruct, if null this is assumed to be a subclass that overrides other functions */
-	virtual FPinStructEditWrapper* GetEditWrapper() const;
+	UE_API virtual FPinStructEditWrapper* GetEditWrapper() const;
 
 	/** Called when struct is modified by the details view */
-	virtual void PropertyValueChanged(const FPropertyChangedEvent& PropertyChangedEvent);
+	UE_API virtual void PropertyValueChanged(const FPropertyChangedEvent& PropertyChangedEvent);
 
 
 	/** Combo Button widget used to show edit content */
@@ -123,3 +125,5 @@ protected:
 	/** Cached description text */
 	FText CachedDescription;
 };
+
+#undef UE_API
diff --git a/Engine/Source/Editor/GraphEditor/Public/NodeFactory.h b/Engine/Source/Editor/GraphEditor/Public/NodeFactory.h
--- a/Engine/Source/Editor/GraphEditor/Public/NodeFactory.h
+++ b/Engine/Source/Editor/GraphEditor/Public/NodeFactory.h
@@ -6,6 +6,8 @@
 #include "HAL/Platform.h"
 #include "Templates/SharedPointer.h"
 
+#define UE_API GRAPHEDITOR_API
+
 class FConnectionDrawingPolicy;
 class FSlateRect;
 class FSlateWindowElementList;
@@ -17,32 +19,34 @@ class UEdGraphPin;
 class UEdGraphSchema;
 
 /** Class that decides which widget type to create for a given data object */
-class GRAPHEDITOR_API FNodeFactory
+class FNodeFactory
 {
 public:
 	/** Create a widget for the supplied node */
-	static TSharedPtr<SGraphNode> CreateNodeWidget(UEdGraphNode* InNode);
+	static UE_API TSharedPtr<SGraphNode> CreateNodeWidget(UEdGraphNode* InNode);
 
 	/** Create a widget for the supplied pin */
-	static TSharedPtr<SGraphPin> CreatePinWidget(UEdGraphPin* InPin);
+	static UE_API TSharedPtr<SGraphPin> CreatePinWidget(UEdGraphPin* InPin);
 
 	/** Create a K2 schema default widget for the supplied pin */
-	static TSharedPtr<SGraphPin> CreateK2PinWidget(UEdGraphPin* InPin);
+	static UE_API TSharedPtr<SGraphPin> CreateK2PinWidget(UEdGraphPin* InPin);
 
     /** Create a pin connection factory for the supplied schema */
-    static FConnectionDrawingPolicy* CreateConnectionPolicy(const UEdGraphSchema* Schema, int32 InBackLayerID, int32 InFrontLayerID, float ZoomFactor, const FSlateRect& InClippingRect, FSlateWindowElementList& InDrawElements, UEdGraph* InGraphObj);
+    static UE_API FConnectionDrawingPolicy* CreateConnectionPolicy(const UEdGraphSchema* Schema, int32 InBackLayerID, int32 InFrontLayerID, float ZoomFactor, const FSlateRect& InClippingRect, FSlateWindowElementList& InDrawElements, UEdGraph* InGraphObj);
 };
 
 
-class GRAPHEDITOR_API FGraphNodeFactory
+class FGraphNodeFactory
 {
 public:
 	/** Create a widget for the supplied node */
-	virtual TSharedPtr<SGraphNode> CreateNodeWidget(UEdGraphNode* InNode);
+	UE_API virtual TSharedPtr<SGraphNode> CreateNodeWidget(UEdGraphNode* InNode);
 
 	/** Create a widget for the supplied pin */
-	virtual TSharedPtr<SGraphPin> CreatePinWidget(UEdGraphPin* InPin);
+	UE_API virtual TSharedPtr<SGraphPin> CreatePinWidget(UEdGraphPin* InPin);
 
 	/** Create a pin connection factory for the supplied schema */
-	virtual FConnectionDrawingPolicy* CreateConnectionPolicy(const UEdGraphSchema* Schema, int32 InBackLayerID, int32 InFrontLayerID, float ZoomFactor, const FSlateRect& InClippingRect, FSlateWindowElementList& InDrawElements, UEdGraph* InGraphObj);
-};
\ No newline at end of file
+	UE_API virtual FConnectionDrawingPolicy* CreateConnectionPolicy(const UEdGraphSchema* Schema, int32 InBackLayerID, int32 InFrontLayerID, float ZoomFactor, const FSlateRect& InClippingRect, FSlateWindowElementList& InDrawElements, UEdGraph* InGraphObj);
+};
+
+#undef UE_API
diff --git a/Engine/Source/Editor/GraphEditor/Public/SCommentBubble.h b/Engine/Source/Editor/GraphEditor/Public/SCommentBubble.h
--- a/Engine/Source/Editor/GraphEditor/Public/SCommentBubble.h
+++ b/Engine/Source/Editor/GraphEditor/Public/SCommentBubble.h
@@ -21,6 +21,8 @@
 #include "Widgets/DeclarativeSyntaxSupport.h"
 #include "Widgets/SCompoundWidget.h"
 
+#define UE_API GRAPHEDITOR_API
+
 class SMultiLineEditableTextBox;
 class UEdGraphNode;
 struct FGeometry;
@@ -29,7 +31,7 @@ struct FPointerEvent;
 DECLARE_DELEGATE_RetVal( bool, FIsGraphNodeHovered );
 DECLARE_DELEGATE_OneParam( FOnCommentBubbleToggled, bool );
 
-class GRAPHEDITOR_API SCommentBubble : public SCompoundWidget
+class SCommentBubble : public SCompoundWidget
 {
 	SLATE_BEGIN_ARGS( SCommentBubble )
 		: _GraphNode( NULL )
@@ -87,85 +89,85 @@ class GRAPHEDITOR_API SCommentBubble : public SCompoundWidget
 	 *
 	 * @param	InArgs				Declaration used by the SNew() macro to construct this widget
 	 */
-	void Construct( const FArguments& InArgs );
+	UE_API void Construct( const FArguments& InArgs );
 
 
 	//~ Begin SWidget Interface
-	virtual FCursorReply OnCursorQuery( const FGeometry& MyGeometry, const FPointerEvent& CursorEvent ) const override;
-	virtual void Tick( const FGeometry& AllottedGeometry, const double InCurrentTime, const float InDeltaTime ) override;
+	UE_API virtual FCursorReply OnCursorQuery( const FGeometry& MyGeometry, const FPointerEvent& CursorEvent ) const override;
+	UE_API virtual void Tick( const FGeometry& AllottedGeometry, const double InCurrentTime, const float InDeltaTime ) override;
 	//~ End SWidget Interface
 
 	/** Updates the comment Visibility */
-	void TickVisibility(const double InCurrentTime, const float InDeltaTime);
+	UE_API void TickVisibility(const double InCurrentTime, const float InDeltaTime);
 
 	/** Returns the offset from the SNode center slot */
 	UE_DEPRECATED(5.6, "Use the version of the function accepting FVector2f; this Slate API no longer interfaces directly with double-precision scalars and vectors.")
-	FVector2D GetOffset() const;
-	FVector2f GetOffset2f() const;
+	UE_API FVector2D GetOffset() const;
+	UE_API FVector2f GetOffset2f() const;
 
 	/** Returns the offset to the arrow center accounting for zoom on either the comment bubble or the title bar button based on current state */
-	float GetArrowCenterOffset() const;
+	UE_API float GetArrowCenterOffset() const;
 
 	/** Returns the bubble size */
 	UE_DEPRECATED(5.6, "Use the version of the function accepting FVector2f; this Slate API no longer interfaces directly with double-precision scalars and vectors.")
-	FVector2D GetSize() const;
-	FVector2f GetSize2f() const;
+	UE_API FVector2D GetSize() const;
+	UE_API FVector2f GetSize2f() const;
 
 	/** Returns if comment bubble is visible */
-	bool IsBubbleVisible() const;
+	UE_API bool IsBubbleVisible() const;
 
 	/** Returns if graph scaling can be applied to this bubble */
-	bool IsScalingAllowed() const;
+	UE_API bool IsScalingAllowed() const;
 
 	/** Called to display/hide the comment bubble */
-	void OnCommentBubbleToggle( ECheckBoxState State );
+	UE_API void OnCommentBubbleToggle( ECheckBoxState State );
 
 	/** Directly sets the bubble's visibility, without assuming it was from a user action (no undo transaction, or OnToggled callback)*/
-	void SetCommentBubbleVisibility(bool bVisible);
+	UE_API void SetCommentBubbleVisibility(bool bVisible);
 
 	/** Called when a node's comment bubble pinned state is changed */
-	void OnPinStateToggle( ECheckBoxState State ) const;
+	UE_API void OnPinStateToggle( ECheckBoxState State ) const;
 
 	/** Called to update the bubble widget layout */
-	void UpdateBubble();
+	UE_API void UpdateBubble();
 
 	/** Returns if the text block currently has keyboard focus*/
-	bool TextBlockHasKeyboardFocus() const;
+	UE_API bool TextBlockHasKeyboardFocus() const;
 
 protected:
 
 	/** Returns the current scale button tooltip */
-	FText GetScaleButtonTooltip() const;
+	UE_API FText GetScaleButtonTooltip() const;
 
 	/** Called to determine if the comment bubble's current visibility */
-	EVisibility GetBubbleVisibility() const;
+	UE_API EVisibility GetBubbleVisibility() const;
 
 	/** Called to determine if the toggle button's current visibility */
-	EVisibility GetToggleButtonVisibility() const;
+	UE_API EVisibility GetToggleButtonVisibility() const;
 
 	/** Returns the color for the toggle bubble including the opacity value */
-	FSlateColor GetToggleButtonColor() const;
+	UE_API FSlateColor GetToggleButtonColor() const;
 
 	/** Returns the color of the main bubble */
-	FSlateColor GetBubbleColor() const;
+	UE_API FSlateColor GetBubbleColor() const;
 
 	/** Returns the foreground color for the text and buttons, taking into account the bubble color */
 	FSlateColor GetForegroundColor() const override { return CalculatedForegroundColor; }
 
 	/** Called when the comment text is committed */
-	void OnCommentTextCommitted( const FText& NewText, ETextCommit::Type CommitInfo );
+	UE_API void OnCommentTextCommitted( const FText& NewText, ETextCommit::Type CommitInfo );
 
-	FSlateColor GetTextBackgroundColor() const;
-	FSlateColor GetTextForegroundColor() const;
-	FSlateColor GetReadOnlyTextForegroundColor() const;
+	UE_API FSlateColor GetTextBackgroundColor() const;
+	UE_API FSlateColor GetTextForegroundColor() const;
+	UE_API FSlateColor GetReadOnlyTextForegroundColor() const;
 	/** Returns bubble toggle check state */
-	ECheckBoxState GetToggleButtonCheck() const;
+	UE_API ECheckBoxState GetToggleButtonCheck() const;
 
 	/** Returns pinned check state */
-	ECheckBoxState GetPinnedButtonCheck() const;
+	UE_API ECheckBoxState GetPinnedButtonCheck() const;
 
 	/** Called to determine if the comment bubble is readonly */
-	bool IsReadOnly() const;
+	UE_API bool IsReadOnly() const;
 
 protected:
 
@@ -212,3 +214,5 @@ protected:
 	FText CachedCommentText;
 
 };
+
+#undef UE_API
diff --git a/Engine/Source/Editor/GraphEditor/Public/SGraphActionMenu.h b/Engine/Source/Editor/GraphEditor/Public/SGraphActionMenu.h
--- a/Engine/Source/Editor/GraphEditor/Public/SGraphActionMenu.h
+++ b/Engine/Source/Editor/GraphEditor/Public/SGraphActionMenu.h
@@ -33,6 +33,8 @@
 #include "Widgets/Views/STableViewBase.h"
 #include "Widgets/Views/STreeView.h"
 
+#define UE_API GRAPHEDITOR_API
+
 class FReferenceCollector;
 class ITableRow;
 class IToolTip;
@@ -56,14 +58,14 @@ DECLARE_DELEGATE( FOnRenameRequestActionNode );
 DECLARE_DELEGATE_RetVal_OneParam( bool, FCreateWidgetMouseButtonDown, TWeakPtr<FEdGraphSchemaAction> );
 
 /** Default widget for GraphActionMenu */
-class GRAPHEDITOR_API SDefaultGraphActionWidget : public SCompoundWidget
+class SDefaultGraphActionWidget : public SCompoundWidget
 {
 	SLATE_BEGIN_ARGS( SDefaultGraphActionWidget ) {}
 		SLATE_ATTRIBUTE(FText, HighlightText)
 	SLATE_END_ARGS()
 
-	void Construct(const FArguments& InArgs, const FCreateWidgetForActionData* InCreateData);
-	virtual FReply OnMouseButtonDown( const FGeometry& MyGeometry, const FPointerEvent& MouseEvent ) override;
+	UE_API void Construct(const FArguments& InArgs, const FCreateWidgetForActionData* InCreateData);
+	UE_API virtual FReply OnMouseButtonDown( const FGeometry& MyGeometry, const FPointerEvent& MouseEvent ) override;
 
 	/** The item that we want to display with this widget */
 	TWeakPtr<FEdGraphSchemaAction> ActionPtr;
@@ -109,7 +111,7 @@ struct FCustomExpanderData
 };
 
 /** Class that displays a list of graph actions and them to be searched and selected */
-class GRAPHEDITOR_API SGraphActionMenu : public SCompoundWidget, public FGCObject
+class SGraphActionMenu : public SCompoundWidget, public FGCObject
 {
 public:
 	/** Delegate that can be used to create a widget for a particular action */
@@ -187,24 +189,24 @@ public:
 
 	SLATE_END_ARGS()
 
-	void Construct( const FArguments& InArgs, bool bIsReadOnly = true );
+	UE_API void Construct( const FArguments& InArgs, bool bIsReadOnly = true );
 
 	// FGCObject override
-	virtual void AddReferencedObjects( FReferenceCollector& Collector ) override;
-	virtual FString GetReferencerName() const override;
+	UE_API virtual void AddReferencedObjects( FReferenceCollector& Collector ) override;
+	UE_API virtual FString GetReferencerName() const override;
 
 	/**
 	 * Refreshes the actions that this widget should display
 	 *
 	 * @param bPreserveExpansion		TRUE if the expansion state of the tree should be preserved
 	 * @param bHandleOnSelectionEvent		TRUE if the item should be selected and any actions that occur with selection will be handled. FALSE and only selection will occur */
-	void RefreshAllActions(bool bPreserveExpansion, bool bHandleOnSelectionEvent = true);
+	UE_API void RefreshAllActions(bool bPreserveExpansion, bool bHandleOnSelectionEvent = true);
 
 	/** Returns a map of all top level sections and their current expansion state. */
-	void GetSectionExpansion(TMap<int32, bool>& SectionExpansion) const;
+	UE_API void GetSectionExpansion(TMap<int32, bool>& SectionExpansion) const;
 
 	/** Sets the sections to be expanded of all top level sections. */
-	void SetSectionExpansion(const TMap<int32, bool>& SectionExpansion);
+	UE_API void SetSectionExpansion(const TMap<int32, bool>& SectionExpansion);
 
 protected:
 	/** Tree view for showing actions */
@@ -286,29 +288,29 @@ protected:
 
 public:
 	// SWidget interface
-	virtual FReply OnKeyDown(const FGeometry& MyGeometry, const FKeyEvent& KeyEvent) override;
+	UE_API virtual FReply OnKeyDown(const FGeometry& MyGeometry, const FKeyEvent& KeyEvent) override;
 	// End of SWidget interface
 
 	/** Get filter text box widget */
-	TSharedRef<SEditableTextBox> GetFilterTextBox();
+	UE_API TSharedRef<SEditableTextBox> GetFilterTextBox();
 
 	/** Get action that is currently selected */
-	void GetSelectedActions(TArray< TSharedPtr<FEdGraphSchemaAction> >& OutSelectedActions) const;
+	UE_API void GetSelectedActions(TArray< TSharedPtr<FEdGraphSchemaAction> >& OutSelectedActions) const;
 
 	/** Initiates a rename on the selected action node, if possible */
-	void OnRequestRenameOnActionNode();
+	UE_API void OnRequestRenameOnActionNode();
 
 	/** Queries if a rename on the selected action node is possible */
-	bool CanRequestRenameOnActionNode() const;
+	UE_API bool CanRequestRenameOnActionNode() const;
 
 	/** Get category that is currently selected */
-	FString GetSelectedCategoryName() const;
+	UE_API FString GetSelectedCategoryName() const;
 	
 	/** Get category child actions that is currently selected */
-	void GetSelectedCategorySubActions(TArray<TSharedPtr<FEdGraphSchemaAction>>& OutActions) const;
+	UE_API void GetSelectedCategorySubActions(TArray<TSharedPtr<FEdGraphSchemaAction>>& OutActions) const;
 
 	/** Get category child actions for the passed in action */
-	void GetCategorySubActions(TWeakPtr<FGraphActionNode> InAction, TArray<TSharedPtr<FEdGraphSchemaAction>>& OutActions) const;
+	UE_API void GetCategorySubActions(TWeakPtr<FGraphActionNode> InAction, TArray<TSharedPtr<FEdGraphSchemaAction>>& OutActions) const;
 
 	/**
 	 * Selects an non-creation item in the list, searching by FName, deselects if name is none 
@@ -319,48 +321,48 @@ public:
 	 * @param	bIsCategory		TRUE if the selection is a category, categories obey different rules and it's hard to re-select properly without this knowledge
 	 * @return					TRUE if the item was successfully selected or the tree cleared, FALSE if unsuccessful
 	 */
-	bool SelectItemByName(const FName& ItemName, ESelectInfo::Type SelectInfo = ESelectInfo::Direct, int32 SectionId = INDEX_NONE, bool bIsCategory = false );
+	UE_API bool SelectItemByName(const FName& ItemName, ESelectInfo::Type SelectInfo = ESelectInfo::Direct, int32 SectionId = INDEX_NONE, bool bIsCategory = false );
 
 	/** Expands any category with the associated name */
-	void ExpandCategory(const FText& CategoryName);
+	UE_API void ExpandCategory(const FText& CategoryName);
 
 	/* Handler for mouse button going down */
-	bool OnMouseButtonDownEvent( TWeakPtr<FEdGraphSchemaAction> InAction );
+	UE_API bool OnMouseButtonDownEvent( TWeakPtr<FEdGraphSchemaAction> InAction );
 
 	/** Updates the displayed list starting from IdxStart, useful for async building the display list of actions */
-	void UpdateForNewActions(int32 IdxStart);
+	UE_API void UpdateForNewActions(int32 IdxStart);
 	/** Regenerated filtered results (FilteredRootAction) based on filter text  */ 
-	void GenerateFilteredItems(bool bPreserveExpansion);
+	UE_API void GenerateFilteredItems(bool bPreserveExpansion);
 
 	/** The last typed action within the graph action menu */
-	static FString LastUsedFilterText;
+	static UE_API FString LastUsedFilterText;
 
 protected:
 	/** Get current filter text */
-	FText GetFilterText() const;
+	UE_API FText GetFilterText() const;
 
 	/** Change the selection to reflect the active suggestion */
-	void MarkActiveSuggestion();
+	UE_API void MarkActiveSuggestion();
 
 	/** Try to spawn the node reflected by the active suggestion */
-	bool TryToSpawnActiveSuggestion();
+	UE_API bool TryToSpawnActiveSuggestion();
 
 	/** Returns true if the tree should be autoexpanded */
-	bool ShouldExpandNodes() const;
+	UE_API bool ShouldExpandNodes() const;
 
 	/** Checks if the passed in node is safe for renaming */
-	bool CanRenameNode(TWeakPtr<FGraphActionNode> InNode) const;
+	UE_API bool CanRenameNode(TWeakPtr<FGraphActionNode> InNode) const;
 
 	// Delegates
 
 	/** Called when filter text changes */
-	void OnFilterTextChanged( const FText& InFilterText );
+	UE_API void OnFilterTextChanged( const FText& InFilterText );
 	/** Called when enter is hit in search box */
-	void OnFilterTextCommitted(const FText& InText, ETextCommit::Type CommitInfo);
+	UE_API void OnFilterTextCommitted(const FText& InText, ETextCommit::Type CommitInfo);
 	/** Get children  */
-	void OnGetChildrenForCategory( TSharedPtr<FGraphActionNode> InItem, TArray< TSharedPtr<FGraphActionNode> >& OutChildren );
+	UE_API void OnGetChildrenForCategory( TSharedPtr<FGraphActionNode> InItem, TArray< TSharedPtr<FGraphActionNode> >& OutChildren );
 	/** Create widget for the supplied node */
-	TSharedRef<ITableRow> MakeWidget( TSharedPtr<FGraphActionNode> InItem, const TSharedRef<STableViewBase>& OwnerTable, bool bIsReadOnly );
+	UE_API TSharedRef<ITableRow> MakeWidget( TSharedPtr<FGraphActionNode> InItem, const TSharedRef<STableViewBase>& OwnerTable, bool bIsReadOnly );
 
 	/**
 	 * Called when tree item is selected 
@@ -368,7 +370,7 @@ protected:
 	 * @param	InSelectedItem	The action node that is being selected
 	 * @param	SelectInfo		Selection type - Only OnMouseClick and OnKeyPress will trigger a call to HandleSelection
 	 */
-	void OnItemSelected( TSharedPtr< FGraphActionNode > InSelectedItem, ESelectInfo::Type SelectInfo );
+	UE_API void OnItemSelected( TSharedPtr< FGraphActionNode > InSelectedItem, ESelectInfo::Type SelectInfo );
 
 	/**
 	 * Executes the selection delegate providing it has been bound, and the provided action node given is valid and is an action node
@@ -377,30 +379,30 @@ protected:
 	 *
 	 * @return true if item selection delegate was executed
 	 */
-	bool HandleSelection( TSharedPtr< FGraphActionNode > &InSelectedItem, ESelectInfo::Type InSelectionType );
+	UE_API bool HandleSelection( TSharedPtr< FGraphActionNode > &InSelectedItem, ESelectInfo::Type InSelectionType );
 
 	/** Called when tree item is double clicked */
-	void OnItemDoubleClicked( TSharedPtr< FGraphActionNode > InClickedItem );
+	UE_API void OnItemDoubleClicked( TSharedPtr< FGraphActionNode > InClickedItem );
 	/** Called when tree item dragged */
-	FReply OnItemDragDetected( const FGeometry& MyGeometry, const FPointerEvent& MouseEvent );
+	UE_API FReply OnItemDragDetected( const FGeometry& MyGeometry, const FPointerEvent& MouseEvent );
 	/** Callback when rename text is committed */
-	void OnNameTextCommitted(const FText& NewText, ETextCommit::Type InTextCommit, TWeakPtr< FGraphActionNode > InAction);
+	UE_API void OnNameTextCommitted(const FText& NewText, ETextCommit::Type InTextCommit, TWeakPtr< FGraphActionNode > InAction);
 	/** Handler for when an item has scrolled into view after having been requested to do so */
-	void OnItemScrolledIntoView( TSharedPtr<FGraphActionNode> InActionNode, const TSharedPtr<ITableRow>& InWidget );
+	UE_API void OnItemScrolledIntoView( TSharedPtr<FGraphActionNode> InActionNode, const TSharedPtr<ITableRow>& InWidget );
 	/** Callback for expanding tree items recursively */
-	void OnSetExpansionRecursive(TSharedPtr<FGraphActionNode> InTreeNode, bool bInIsItemExpanded);
+	UE_API void OnSetExpansionRecursive(TSharedPtr<FGraphActionNode> InTreeNode, bool bInIsItemExpanded);
 	/** Helper function for adding and scoring actions from our builder */
-	void ScoreAndAddActions(int32 StartingIndex = INDEX_NONE);
+	UE_API void ScoreAndAddActions(int32 StartingIndex = INDEX_NONE);
 
 	/** Helper functions for keyboard interaction */
-	void SelectPreviousAction(int32 Num = 1);
-	void SelectNextAction(int32 Num = 1);
-	void SelectFirstAction();
-	void SelectLastAction();
+	UE_API void SelectPreviousAction(int32 Num = 1);
+	UE_API void SelectNextAction(int32 Num = 1);
+	UE_API void SelectFirstAction();
+	UE_API void SelectLastAction();
 	/** General helper functions for accessing the filtered tree */
-	TSharedPtr<FGraphActionNode> GetFirstAction();
-	const TArray< TSharedPtr<FGraphActionNode> >& GetFilteredActionNodes(int32* OutSelectedIndex = nullptr);
-	int32 GetTotalLeafNodes() const;
+	UE_API TSharedPtr<FGraphActionNode> GetFirstAction();
+	UE_API const TArray< TSharedPtr<FGraphActionNode> >& GetFilteredActionNodes(int32* OutSelectedIndex = nullptr);
+	UE_API int32 GetTotalLeafNodes() const;
 private:
 	/** The pins that have been dragged off of to prompt the creation of this action menu. */
 	TArray<UEdGraphPin*> DraggedFromPins;
@@ -409,3 +411,4 @@ private:
 	UEdGraph* GraphObj;
 };
 
+#undef UE_API
diff --git a/Engine/Source/Editor/GraphEditor/Public/SGraphEditorActionMenu.h b/Engine/Source/Editor/GraphEditor/Public/SGraphEditorActionMenu.h
--- a/Engine/Source/Editor/GraphEditor/Public/SGraphEditorActionMenu.h
+++ b/Engine/Source/Editor/GraphEditor/Public/SGraphEditorActionMenu.h
@@ -26,7 +26,7 @@ struct FGraphActionListBuilderBase;
 
 /////////////////////////////////////////////////////////////////////////////////////////////////
 
-class GRAPHEDITOR_API SGraphEditorActionMenu : public SBorder
+class SGraphEditorActionMenu : public SBorder
 {
 public:
 	SLATE_BEGIN_ARGS( SGraphEditorActionMenu )
@@ -43,13 +43,13 @@ public:
 		SLATE_EVENT( SGraphActionMenu::FOnCreateWidgetForAction, OnCreateWidgetForAction )
 	SLATE_END_ARGS()
 
-	void Construct( const FArguments& InArgs );
+	GRAPHEDITOR_API void Construct( const FArguments& InArgs );
 
-	~SGraphEditorActionMenu();
+	GRAPHEDITOR_API ~SGraphEditorActionMenu();
 
-	void RefreshAllActions();
+	GRAPHEDITOR_API void RefreshAllActions();
 
-	TSharedRef<SEditableTextBox> GetFilterTextBox();
+	GRAPHEDITOR_API TSharedRef<SEditableTextBox> GetFilterTextBox();
 
 protected:
 	UEdGraph* GraphObj;
@@ -62,8 +62,8 @@ protected:
 
 	TSharedPtr<SGraphActionMenu> GraphActionMenu;
 
-	void OnActionSelected( const TArray< TSharedPtr<FEdGraphSchemaAction> >& SelectedAction, ESelectInfo::Type InSelectionType );
+	GRAPHEDITOR_API void OnActionSelected( const TArray< TSharedPtr<FEdGraphSchemaAction> >& SelectedAction, ESelectInfo::Type InSelectionType );
 
 	/** Callback used to populate all actions list in SGraphActionMenu */
-	void CollectAllActions(FGraphActionListBuilderBase& OutAllActions);
+	GRAPHEDITOR_API void CollectAllActions(FGraphActionListBuilderBase& OutAllActions);
 };
diff --git a/Engine/Source/Editor/GraphEditor/Public/SGraphNode.h b/Engine/Source/Editor/GraphEditor/Public/SGraphNode.h
--- a/Engine/Source/Editor/GraphEditor/Public/SGraphNode.h
+++ b/Engine/Source/Editor/GraphEditor/Public/SGraphNode.h
@@ -37,6 +37,8 @@
 #include "Widgets/SOverlay.h"
 #include "Widgets/SWidget.h"
 
+#define UE_API GRAPHEDITOR_API
+
 class FActorDragDropOp;
 class FDragDropEvent;
 class ISlateStyle;
@@ -58,7 +60,7 @@ struct FSlateBrush;
 /////////////////////////////////////////////////////
 // SNodeTitle
 
-class GRAPHEDITOR_API SNodeTitle : public SCompoundWidget
+class SNodeTitle : public SCompoundWidget
 {
 public:
 	SLATE_BEGIN_ARGS(SNodeTitle)
@@ -79,20 +81,20 @@ public:
 		SLATE_ATTRIBUTE(FText, Text)
 	SLATE_END_ARGS()
 
-	void Construct(const FArguments& InArgs, UEdGraphNode* InNode);
+	UE_API void Construct(const FArguments& InArgs, UEdGraphNode* InNode);
 
 	// SWidget interface
-	virtual void Tick( const FGeometry& AllottedGeometry, const double InCurrentTime, const float InDeltaTime ) override;
+	UE_API virtual void Tick( const FGeometry& AllottedGeometry, const double InCurrentTime, const float InDeltaTime ) override;
 	// End of SWidget interface
 
 	/** Returns the main title for the node */
-	FText GetHeadTitle() const;
+	UE_API FText GetHeadTitle() const;
 
 	/** Get the size of this title the last time it was drawn */
-	UE::Slate::FDeprecateVector2DResult GetTitleSize() const;
+	UE_API UE::Slate::FDeprecateVector2DResult GetTitleSize() const;
 
 	/** Mark the cached title dirty */
-	void MarkDirty();
+	UE_API void MarkDirty();
 
 protected:
 	TWeakObjectPtr<UEdGraphNode> GraphNode;
@@ -111,106 +113,106 @@ protected:
 
 protected:
 	// Gets the expected node title
-	FText GetNodeTitle() const;
+	UE_API FText GetNodeTitle() const;
 
 	// Rebuilds the widget if needed
-	void RebuildWidget();
+	UE_API void RebuildWidget();
 };
 
 /////////////////////////////////////////////////////
 // SGraphNode
 
-class GRAPHEDITOR_API SGraphNode : public SNodePanel::SNode
+class SGraphNode : public SNodePanel::SNode
 {
 public:
 	// SWidget interface
-	virtual void OnDragEnter( const FGeometry& MyGeometry, const FDragDropEvent& DragDropEvent ) override;
-	virtual void OnDragLeave( const FDragDropEvent& DragDropEvent ) override;
-	virtual FReply OnDragOver( const FGeometry& MyGeometry, const FDragDropEvent& DragDropEvent ) override;
-	virtual FReply OnDrop( const FGeometry& MyGeometry, const FDragDropEvent& DragDropEvent ) override;
+	UE_API virtual void OnDragEnter( const FGeometry& MyGeometry, const FDragDropEvent& DragDropEvent ) override;
+	UE_API virtual void OnDragLeave( const FDragDropEvent& DragDropEvent ) override;
+	UE_API virtual FReply OnDragOver( const FGeometry& MyGeometry, const FDragDropEvent& DragDropEvent ) override;
+	UE_API virtual FReply OnDrop( const FGeometry& MyGeometry, const FDragDropEvent& DragDropEvent ) override;
 
-	virtual FReply OnMouseButtonDown( const FGeometry& MyGeometry, const FPointerEvent& MouseEvent ) override;
-	virtual FReply OnMouseButtonUp( const FGeometry& MyGeometry, const FPointerEvent& MouseEvent ) override;
-	virtual FReply OnMouseButtonDoubleClick( const FGeometry& InMyGeometry, const FPointerEvent& InMouseEvent ) override;
+	UE_API virtual FReply OnMouseButtonDown( const FGeometry& MyGeometry, const FPointerEvent& MouseEvent ) override;
+	UE_API virtual FReply OnMouseButtonUp( const FGeometry& MyGeometry, const FPointerEvent& MouseEvent ) override;
+	UE_API virtual FReply OnMouseButtonDoubleClick( const FGeometry& InMyGeometry, const FPointerEvent& InMouseEvent ) override;
 
-	virtual TSharedPtr<IToolTip> GetToolTip() override;
-	virtual void OnToolTipClosing() override;
+	UE_API virtual TSharedPtr<IToolTip> GetToolTip() override;
+	UE_API virtual void OnToolTipClosing() override;
 
-	virtual void Tick( const FGeometry& AllottedGeometry, const double InCurrentTime, const float InDeltaTime ) override;
+	UE_API virtual void Tick( const FGeometry& AllottedGeometry, const double InCurrentTime, const float InDeltaTime ) override;
 	// End of SWidget interface
 
 	// SNodePanel::SNode interface
-	virtual void MoveTo( const FVector2f& NewPosition, FNodeSet& NodeFilter, bool bMarkDirty = true ) override;
-	virtual FVector2f GetPosition2f() const override;
-	virtual FString GetNodeComment() const override;
-	virtual UObject* GetObjectBeingDisplayed() const override;
+	UE_API virtual void MoveTo( const FVector2f& NewPosition, FNodeSet& NodeFilter, bool bMarkDirty = true ) override;
+	UE_API virtual FVector2f GetPosition2f() const override;
+	UE_API virtual FString GetNodeComment() const override;
+	UE_API virtual UObject* GetObjectBeingDisplayed() const override;
 	// End of SNodePanel::SNode interface
 
 	/** Set attribute for determining if widget is editable */
-	void SetIsEditable(TAttribute<bool> InIsEditable);
+	UE_API void SetIsEditable(TAttribute<bool> InIsEditable);
 
 	/** Returns if widget is editable, additionally considers if the owning graph is read only */
-	virtual bool IsNodeEditable() const;
+	UE_API virtual bool IsNodeEditable() const;
 
 	/** Set event when node is double clicked */
-	void SetDoubleClickEvent(FSingleNodeEvent InDoubleClickEvent);
+	UE_API void SetDoubleClickEvent(FSingleNodeEvent InDoubleClickEvent);
 
 	/** @param OwnerPanel  The GraphPanel that this node belongs to */
-	virtual void SetOwner( const TSharedRef<SGraphPanel>& OwnerPanel );
+	UE_API virtual void SetOwner( const TSharedRef<SGraphPanel>& OwnerPanel );
 
 	/** @return the editable title for a node */
-	FString GetEditableNodeTitle() const;
+	UE_API FString GetEditableNodeTitle() const;
 
 	/** @return the editable title for a node */
-	FText GetEditableNodeTitleAsText() const;
+	UE_API FText GetEditableNodeTitleAsText() const;
 
 	/** @return the tint for the node's title image */
-	FSlateColor GetNodeTitleColor() const;
+	UE_API FSlateColor GetNodeTitleColor() const;
 
 	/** @return the tint for the node's comment */
-	FSlateColor GetNodeCommentColor() const;
+	UE_API FSlateColor GetNodeCommentColor() const;
 
 	/** @return the tint for the node's main body */
-	FSlateColor GetNodeBodyColor() const;
+	UE_API FSlateColor GetNodeBodyColor() const;
 
 	/** @return the image brush to use for the node */
-	virtual const FSlateBrush * GetNodeBodyBrush() const;
+	UE_API virtual const FSlateBrush * GetNodeBodyBrush() const;
 
 	/** @return the tint for the node's title icon */
-	FSlateColor GetNodeTitleIconColor() const;
+	UE_API FSlateColor GetNodeTitleIconColor() const;
 
 	/** @return the tint for the node's title text */
-	FLinearColor GetNodeTitleTextColor() const;
+	UE_API FLinearColor GetNodeTitleTextColor() const;
 
 	/** @return the tooltip to display when over the node */
-	FText GetNodeTooltip() const;
+	UE_API FText GetNodeTooltip() const;
 	
 	/** @return the node being observed by this widget*/
-	UEdGraphNode* GetNodeObj() const;
+	UE_API UEdGraphNode* GetNodeObj() const;
 
 	/** @return the node under the mouse (either this node or one of its children) */
-	virtual TSharedRef<SGraphNode> GetNodeUnderMouse(const FGeometry& MyGeometry, const FPointerEvent& MouseEvent);
+	UE_API virtual TSharedRef<SGraphNode> GetNodeUnderMouse(const FGeometry& MyGeometry, const FPointerEvent& MouseEvent);
 
-	TSharedPtr<SGraphPanel> GetOwnerPanel() const;
+	UE_API TSharedPtr<SGraphPanel> GetOwnerPanel() const;
 
 	/**
 	 * Update this GraphNode to match the data that it is observing
 	 */
-	virtual void UpdateGraphNode();
+	UE_API virtual void UpdateGraphNode();
 	
 	/** Create the widgets for pins on the node */
-	virtual void CreatePinWidgets();
+	UE_API virtual void CreatePinWidgets();
 
 	/** Create a single pin widget */
-	virtual void CreateStandardPinWidget(UEdGraphPin* Pin);
+	UE_API virtual void CreateStandardPinWidget(UEdGraphPin* Pin);
 
 	/**
 	 * Get all the pins found on this node.
 	 *
 	 * @param AllPins  The set of pins found on this node.
 	 */
-	void GetPins( TSet< TSharedRef<SWidget> >& AllPins ) const;
-	void GetPins( TArray< TSharedRef<SWidget> >& AllPins ) const;
+	UE_API void GetPins( TSet< TSharedRef<SWidget> >& AllPins ) const;
+	UE_API void GetPins( TArray< TSharedRef<SWidget> >& AllPins ) const;
 
 	/** 
 	 * Find the pin that is hovered.
@@ -220,29 +222,29 @@ public:
 	 *
 	 * @return A pointer to the pin widget hovered in this node, or invalid pointer if none.
 	 */
-	virtual TSharedPtr<SGraphPin> GetHoveredPin( const FGeometry& MyGeometry, const FPointerEvent& MouseEvent ) const;
+	UE_API virtual TSharedPtr<SGraphPin> GetHoveredPin( const FGeometry& MyGeometry, const FPointerEvent& MouseEvent ) const;
 
-	TSharedPtr<SGraphPin> FindWidgetForPin( UEdGraphPin* ThePin ) const;
+	UE_API TSharedPtr<SGraphPin> FindWidgetForPin( UEdGraphPin* ThePin ) const;
 
-	void PlaySpawnEffect();
+	UE_API void PlaySpawnEffect();
 
 	/** Given a coordinate in SGraphNode space, return the same coordinate in graph space while taking zoom and panning of the parent graph into account */
-	UE::Slate::FDeprecateVector2DResult NodeCoordToGraphCoord( const UE::Slate::FDeprecateVector2DParameter& PanelSpaceCoordinate ) const;
+	UE_API UE::Slate::FDeprecateVector2DResult NodeCoordToGraphCoord( const UE::Slate::FDeprecateVector2DParameter& PanelSpaceCoordinate ) const;
 
-	UE::Slate::FDeprecateVector2DResult GetContentScale() const;
-	FLinearColor GetColorAndOpacity() const;
-	FLinearColor GetPinLabelColorAndOpacity() const;	
+	UE_API UE::Slate::FDeprecateVector2DResult GetContentScale() const;
+	UE_API FLinearColor GetColorAndOpacity() const;
+	UE_API FLinearColor GetPinLabelColorAndOpacity() const;	
 
 	/** Set event when text is committed on the node */
-	void SetVerifyTextCommitEvent(FOnNodeVerifyTextCommit InOnVerifyTextCommit);
+	UE_API void SetVerifyTextCommitEvent(FOnNodeVerifyTextCommit InOnVerifyTextCommit);
 	/** Set event when text is committed on the node */
-	void SetTextCommittedEvent(FOnNodeTextCommitted InDelegate);
+	UE_API void SetTextCommittedEvent(FOnNodeTextCommitted InDelegate);
 	/** Set event when the user generates a warning tooltip because a connection was invalid */
-	void SetDisallowedPinConnectionEvent(SGraphEditor::FOnDisallowedPinConnection InOnDisallowedPinConnection);
+	UE_API void SetDisallowedPinConnectionEvent(SGraphEditor::FOnDisallowedPinConnection InOnDisallowedPinConnection);
 	/** called to replace this nodes comment text */
-	virtual void OnCommentTextCommitted(const FText& NewComment, ETextCommit::Type CommitInfo);
+	UE_API virtual void OnCommentTextCommitted(const FText& NewComment, ETextCommit::Type CommitInfo);
 	/** called when the node's comment bubble is toggled */
-	virtual void OnCommentBubbleToggled(bool bInCommentBubbleVisible);
+	UE_API virtual void OnCommentBubbleToggled(bool bInCommentBubbleVisible);
 	/** returns true if a rename is pending on this node */
 	bool IsRenamePending() const { return bRenameIsPending; }
 
@@ -253,28 +255,28 @@ public:
 	}
 
 	/** flags node as rename pending if supported */
-	void RequestRename();
+	UE_API void RequestRename();
 
 	/** Sets node into rename state if supported */
-	void ApplyRename();
+	UE_API void ApplyRename();
 
 	/** return rect of the title area */
-	virtual FSlateRect GetTitleRect() const;
+	UE_API virtual FSlateRect GetTitleRect() const;
 
 	/** Called from drag drop code when a disallowed connection is hovered */
-	void NotifyDisallowedPinConnection(const class UEdGraphPin* PinA, const class UEdGraphPin* PinB) const;
+	UE_API void NotifyDisallowedPinConnection(const class UEdGraphPin* PinA, const class UEdGraphPin* PinB) const;
 
 	/** Gets the unscaled position of the node from the last tick */
 	UE::Slate::FDeprecateVector2DResult GetUnscaledPosition() const {return CachedUnscaledPosition;}
 
 	/** Returns the current Node LOD or Highest LOD if unable to query */
-	EGraphRenderingLOD::Type GetCurrentLOD() const;
+	UE_API EGraphRenderingLOD::Type GetCurrentLOD() const;
 
 	/** Called when GraphNode changes its error information, may be called when no change has actually occurred: */
-	virtual void RefreshErrorInfo();
+	UE_API virtual void RefreshErrorInfo();
 	
 	/** Called when ed graph data is cleared, indicating this widget can no longer safely access GraphNode - forwards call to owned pins: */
-	void InvalidateGraphData();
+	UE_API void InvalidateGraphData();
 
 	/** Returns true if the node is hiding its pins */
 	virtual bool IsHidingPinWidgets() const { return false; }
@@ -289,33 +291,33 @@ public:
 	virtual TOptional<FSlateColor> GetPinTextColor(const SGraphPin* InGraphPin) const { return TOptional<FSlateColor>(); }
 
 protected:
-	SGraphNode();
+	UE_API SGraphNode();
 
-	void PositionThisNodeBetweenOtherNodes(const TMap< UObject*, TSharedRef<SNode> >& NodeToWidgetLookup, UEdGraphNode* PreviousNode, UEdGraphNode* NextNode, float HeightAboveWire) const;
-	void PositionThisNodeBetweenOtherNodes(const TMap< UObject*, TSharedRef<SNode> >& NodeToWidgetLookup, TSet<UEdGraphNode*>& PreviousNodes, TSet<UEdGraphNode*>& NextNodes, float HeightAboveWire) const;
-	void PositionThisNodeBetweenOtherNodes(const FVector2f& PrevPos, const FVector2f& NextPos, float HeightAboveWire) const;
+	UE_API void PositionThisNodeBetweenOtherNodes(const TMap< UObject*, TSharedRef<SNode> >& NodeToWidgetLookup, UEdGraphNode* PreviousNode, UEdGraphNode* NextNode, float HeightAboveWire) const;
+	UE_API void PositionThisNodeBetweenOtherNodes(const TMap< UObject*, TSharedRef<SNode> >& NodeToWidgetLookup, TSet<UEdGraphNode*>& PreviousNodes, TSet<UEdGraphNode*>& NextNodes, float HeightAboveWire) const;
+	UE_API void PositionThisNodeBetweenOtherNodes(const FVector2f& PrevPos, const FVector2f& NextPos, float HeightAboveWire) const;
 
 	/**
 	 * Check whether drag and drop functionality is permitted on the given node 
 	 */
-	static bool CanAllowInteractionUsingDragDropOp( const UEdGraphNode* GraphNodePtr, const TSharedPtr<FActorDragDropOp>& DragDropOp );
+	static UE_API bool CanAllowInteractionUsingDragDropOp( const UEdGraphNode* GraphNodePtr, const TSharedPtr<FActorDragDropOp>& DragDropOp );
 
 	/**
 	 *	Function to get error description string
 	 *
 	 *	@return string to be displayed as tooltip.
 	 */
-	FText GetErrorMsgToolTip() const;
+	UE_API FText GetErrorMsgToolTip() const;
 
 	/**
 	 * Add a new pin to this graph node. The pin must be newly created.
 	 *
 	 * @param PinToAdd   A new pin to add to this GraphNode.
 	 */
-	virtual void AddPin( const TSharedRef<SGraphPin>& PinToAdd );
+	UE_API virtual void AddPin( const TSharedRef<SGraphPin>& PinToAdd );
 
 	/** Hook that allows derived classes to supply their own SGraphPin derivatives for any pin. */
-	virtual TSharedPtr<SGraphPin> CreatePinWidget(UEdGraphPin* Pin) const;
+	UE_API virtual TSharedPtr<SGraphPin> CreatePinWidget(UEdGraphPin* Pin) const;
 
 	/**
 	 * Override this to provide support for an 'expensive' tooltip widget that is only built on demand
@@ -332,37 +334,37 @@ protected:
 	virtual TOptional<ETextOverflowPolicy> GetNameOverflowPolicy() const { return {}; }
 
 	/* Helper function to check if node can be renamed */
-	virtual bool IsNameReadOnly () const;
+	UE_API virtual bool IsNameReadOnly () const;
 
 	/** Called when text is being committed to check for validity */
-	bool OnVerifyNameTextChanged ( const FText& InText, FText& OutErrorMessage );
+	UE_API bool OnVerifyNameTextChanged ( const FText& InText, FText& OutErrorMessage );
 
 	/* Called when text is committed on the node */
-	void OnNameTextCommited ( const FText& InText, ETextCommit::Type CommitInfo ) ;
+	UE_API void OnNameTextCommited ( const FText& InText, ETextCommit::Type CommitInfo ) ;
 
 	/* Helper function to set the error color for the node */
 	FSlateColor GetErrorColor() const	{return ErrorColor;}
 
 	/** Controls wether to show or not the visual warning message */
-	EVisibility VisualWarningVisibility() const;
+	UE_API EVisibility VisualWarningVisibility() const;
 
 	/** Function to get a visual warning description string(visual warnings does not break the build) */
-	FText GetVisualWarningMsgToolTip() const;
+	UE_API FText GetVisualWarningMsgToolTip() const;
 
 	/* Helper function to set the error color for the node */
-	FSlateColor GetVisualWarningColor() const;
+	UE_API FSlateColor GetVisualWarningColor() const;
 
 	/** Helper function to get any error text for the node */
 	FString GetErrorMessage() const {return ErrorMsg;}
 
 	/** Called to set error text on the node */
-	virtual void UpdateErrorInfo();
+	UE_API virtual void UpdateErrorInfo();
 
 	/** Set-up the error reporting widget for the node */
-	void SetupErrorReporting();
+	UE_API void SetupErrorReporting();
 
 	// Should we use low-detail node titles?
-	virtual bool UseLowDetailNodeTitles() const;
+	UE_API virtual bool UseLowDetailNodeTitles() const;
 
 public:
 	
@@ -377,31 +379,31 @@ protected:
 	///// ADVANCED VIEW FUNCTIONS /////
 
 	/** Create button to show/hide advanced pins */
-	virtual void CreateAdvancedViewArrow(TSharedPtr<SVerticalBox> MainBox);
+	UE_API virtual void CreateAdvancedViewArrow(TSharedPtr<SVerticalBox> MainBox);
 
 	/** Returns visibility of AdvancedViewButton */
-	virtual EVisibility AdvancedViewArrowVisibility() const;
+	UE_API virtual EVisibility AdvancedViewArrowVisibility() const;
 
 	/** Show/hide advanced view */
-	virtual void OnAdvancedViewChanged( const ECheckBoxState NewCheckedState );
+	UE_API virtual void OnAdvancedViewChanged( const ECheckBoxState NewCheckedState );
 
 	/** hidden == unchecked, shown == checked */
-	virtual ECheckBoxState IsAdvancedViewChecked() const;
+	UE_API virtual ECheckBoxState IsAdvancedViewChecked() const;
 
 	/** Up when shown, down when hidden */
-	virtual const FSlateBrush* GetAdvancedViewArrow() const;
+	UE_API virtual const FSlateBrush* GetAdvancedViewArrow() const;
 
 	/** Checks if the node is the only node selected */
-	bool IsSelectedExclusively() const;
+	UE_API bool IsSelectedExclusively() const;
 
 	virtual void SetDefaultTitleAreaWidget(TSharedRef<SOverlay> DefaultTitleAreaWidget) {}
-	virtual TSharedRef<SWidget> CreateTitleWidget(TSharedPtr<SNodeTitle> NodeTitle);
+	UE_API virtual TSharedRef<SWidget> CreateTitleWidget(TSharedPtr<SNodeTitle> NodeTitle);
 	
 	/** Optionally create a widget to dock to the right in the title bar of the node. */
-	virtual TSharedRef<SWidget> CreateTitleRightWidget();
+	UE_API virtual TSharedRef<SWidget> CreateTitleRightWidget();
 	
 	/** Create the inner node content area, including the left/right pin boxes */
-	virtual TSharedRef<SWidget> CreateNodeContentArea();
+	UE_API virtual TSharedRef<SWidget> CreateNodeContentArea();
 
 	///// ADD PIN BUTTON FUNCTIONS /////
 
@@ -412,22 +414,22 @@ protected:
 	virtual void CreateOutputSideAddButton(TSharedPtr<SVerticalBox> OutputBox) {};
 
 	/** Creates widget for an Add pin button, which can then be added to the node */
-	TSharedRef<SWidget> AddPinButtonContent(FText PinText, FText PinTooltipText, bool bRightSide = true, FString DocumentationExcerpt = FString(), TSharedPtr<SToolTip> CustomTooltip = NULL);
+	UE_API TSharedRef<SWidget> AddPinButtonContent(FText PinText, FText PinTooltipText, bool bRightSide = true, FString DocumentationExcerpt = FString(), TSharedPtr<SToolTip> CustomTooltip = NULL);
 
 	/** Checks whether Add pin button should currently be visible */
-	virtual EVisibility IsAddPinButtonVisible() const;
+	UE_API virtual EVisibility IsAddPinButtonVisible() const;
 
 	/** Callback function executed when Add pin button is clicked */
 	virtual FReply OnAddPin() {return FReply::Handled();}
 
 	/* Populate a meta data tag with information about this graph node */
-	virtual void PopulateMetaTag(class FGraphNodeMetaData* TagMeta) const;
+	UE_API virtual void PopulateMetaTag(class FGraphNodeMetaData* TagMeta) const;
 
 	/** Returns TRUE if the input pin should be hidden from view */
-	bool ShouldPinBeHidden(const UEdGraphPin* InPin) const;
+	UE_API bool ShouldPinBeHidden(const UEdGraphPin* InPin) const;
 
 	/** Returns the widget to use for the enabled state of the node */
-	TSharedPtr<SWidget> GetEnabledStateWidget();
+	UE_API TSharedPtr<SWidget> GetEnabledStateWidget();
 
 protected:
 	/** Input pin widgets on this node */
@@ -495,3 +497,5 @@ protected:
 private:
 	TSharedPtr<SLevelOfDetailBranchNode> TitleLODBranchNode;
 };
+
+#undef UE_API
diff --git a/Engine/Source/Editor/GraphEditor/Public/SGraphNodeComment.h b/Engine/Source/Editor/GraphEditor/Public/SGraphNodeComment.h
--- a/Engine/Source/Editor/GraphEditor/Public/SGraphNodeComment.h
+++ b/Engine/Source/Editor/GraphEditor/Public/SGraphNodeComment.h
@@ -16,6 +16,8 @@
 #include "Templates/SharedPointer.h"
 #include "Widgets/DeclarativeSyntaxSupport.h"
 
+#define UE_API GRAPHEDITOR_API
+
 class FDragDropEvent;
 class SBorder;
 class SCommentBubble;
@@ -25,35 +27,35 @@ struct FGeometry;
 struct FPointerEvent;
 struct FSlateBrush;
 
-class GRAPHEDITOR_API SGraphNodeComment : public SGraphNodeResizable
+class SGraphNodeComment : public SGraphNodeResizable
 {
 public:
 	SLATE_BEGIN_ARGS(SGraphNodeComment){}
 	SLATE_END_ARGS()
 
 	//~ Begin SWidget Interface
-	virtual FReply OnMouseButtonDoubleClick( const FGeometry& InMyGeometry, const FPointerEvent& InMouseEvent ) override;
-	virtual FReply OnMouseButtonUp( const FGeometry& MyGeometry, const FPointerEvent& MouseEvent ) override;
-	virtual void Tick(const FGeometry& AllottedGeometry, const double InCurrentTime, const float InDeltaTime) override;
-	virtual FReply OnDrop( const FGeometry& MyGeometry, const FDragDropEvent& DragDropEvent ) override;
-	virtual void OnDragEnter( const FGeometry& MyGeometry, const FDragDropEvent& DragDropEvent ) override;
+	UE_API virtual FReply OnMouseButtonDoubleClick( const FGeometry& InMyGeometry, const FPointerEvent& InMouseEvent ) override;
+	UE_API virtual FReply OnMouseButtonUp( const FGeometry& MyGeometry, const FPointerEvent& MouseEvent ) override;
+	UE_API virtual void Tick(const FGeometry& AllottedGeometry, const double InCurrentTime, const float InDeltaTime) override;
+	UE_API virtual FReply OnDrop( const FGeometry& MyGeometry, const FDragDropEvent& DragDropEvent ) override;
+	UE_API virtual void OnDragEnter( const FGeometry& MyGeometry, const FDragDropEvent& DragDropEvent ) override;
 	//~ End SWidget Interface
 
 	//~ Begin SNodePanel::SNode Interface
-	virtual const FSlateBrush* GetShadowBrush(bool bSelected) const override;
-	virtual void GetOverlayBrushes(bool bSelected, const FVector2f& WidgetSize, TArray<FOverlayBrushInfo>& Brushes) const override;
+	UE_API virtual const FSlateBrush* GetShadowBrush(bool bSelected) const override;
+	UE_API virtual void GetOverlayBrushes(bool bSelected, const FVector2f& WidgetSize, TArray<FOverlayBrushInfo>& Brushes) const override;
 	virtual bool ShouldAllowCulling() const override { return true; }
-	virtual int32 GetSortDepth() const override;
-	virtual void EndUserInteraction() const override;
-	virtual FString GetNodeComment() const override;
+	UE_API virtual int32 GetSortDepth() const override;
+	UE_API virtual void EndUserInteraction() const override;
+	UE_API virtual FString GetNodeComment() const override;
 	//~ End SNodePanel::SNode Interface
 
 	//~ Begin SPanel Interface
-	virtual FVector2D ComputeDesiredSize(float) const override;
+	UE_API virtual FVector2D ComputeDesiredSize(float) const override;
 	//~ End SPanel Interface
 
 	//~ Begin SGraphNode Interface
-	virtual bool IsNameReadOnly() const override;
+	UE_API virtual bool IsNameReadOnly() const override;
 	virtual FSlateColor GetCommentColor() const override { return GetCommentBodyColor(); }
 	/** Requests a rename when the node was initially spawned */
 	virtual void RequestRenameOnSpawn() 
@@ -63,54 +65,54 @@ public:
 	}
 	//~ End SGraphNode Interface
 
-	void Construct( const FArguments& InArgs, UEdGraphNode_Comment* InNode );
+	UE_API void Construct( const FArguments& InArgs, UEdGraphNode_Comment* InNode );
 
 	/** return if the node can be selected, by pointing given location */
-	virtual bool CanBeSelected( const FVector2f& MousePositionInNode ) const override;
+	UE_API virtual bool CanBeSelected( const FVector2f& MousePositionInNode ) const override;
 
 	/** return size of the title bar */
-	virtual FVector2f GetDesiredSizeForMarquee2f() const override;
+	UE_API virtual FVector2f GetDesiredSizeForMarquee2f() const override;
 
 	/** return rect of the title bar */
-	virtual FSlateRect GetTitleRect() const override;
+	UE_API virtual FSlateRect GetTitleRect() const override;
 
 protected:
 	//~ Begin SGraphNode Interface
-	virtual void UpdateGraphNode() override;
-	virtual void PopulateMetaTag(class FGraphNodeMetaData* TagMeta) const override;
+	UE_API virtual void UpdateGraphNode() override;
+	UE_API virtual void PopulateMetaTag(class FGraphNodeMetaData* TagMeta) const override;
 
 	/**
 	 * Helper method to update selection state of comment and any nodes 'contained' within it
 	 * @param bSelected	If true comment is being selected, false otherwise
 	 * @param bUpdateNodesUnderComment If true then force the rebuild of the list of nodes under the comment
 	 */
-	void HandleSelection(bool bIsSelected, bool bUpdateNodesUnderComment = false) const;
+	UE_API void HandleSelection(bool bIsSelected, bool bUpdateNodesUnderComment = false) const;
 
 	/** Helper function to determine if a node is under this comment widget or not */
-	virtual bool IsNodeUnderComment(UEdGraphNode_Comment* InCommentNode, const TSharedRef<SGraphNode> InNodeWidget) const;
+	UE_API virtual bool IsNodeUnderComment(UEdGraphNode_Comment* InCommentNode, const TSharedRef<SGraphNode> InNodeWidget) const;
 
 	/** called when user is moving the comment node */
-	virtual void MoveTo(const FVector2f& NewPosition, FNodeSet& NodeFilter, bool bMarkDirty = true) override;
+	UE_API virtual void MoveTo(const FVector2f& NewPosition, FNodeSet& NodeFilter, bool bMarkDirty = true) override;
 
 	//~ Begin SGraphNodeResizable Interface
-	virtual float GetTitleBarHeight() const override;
-	virtual FSlateRect GetHitTestingBorder() const override;
-	virtual FVector2f GetNodeMaximumSize2f() const override;
+	UE_API virtual float GetTitleBarHeight() const override;
+	UE_API virtual FSlateRect GetHitTestingBorder() const override;
+	UE_API virtual FVector2f GetNodeMaximumSize2f() const override;
 	//~ Begin SGraphNodeResizable Interface
 
 	/** @return the color to tint the comment body */
-	FSlateColor GetCommentBodyColor() const;
+	UE_API FSlateColor GetCommentBodyColor() const;
 
 	/** @return the color to tint the title bar */
-	FSlateColor GetCommentTitleBarColor() const;
+	UE_API FSlateColor GetCommentTitleBarColor() const;
 
 	/** @return the color to tint the comment bubble */
-	FSlateColor GetCommentBubbleColor() const;
+	UE_API FSlateColor GetCommentBubbleColor() const;
 
 private:
 	
 	/** Returns the width to wrap the text of the comment at */
-	float GetWrapAt() const;
+	UE_API float GetWrapAt() const;
 
 	/** The comment bubble widget (used when zoomed out) */
 	TSharedPtr<SCommentBubble> CommentBubble;
@@ -138,3 +140,5 @@ private:
 	/** Local copy of the comment style */
 	FInlineEditableTextBlockStyle CommentStyle;
 };
+
+#undef UE_API
diff --git a/Engine/Source/Editor/GraphEditor/Public/SGraphNodeDefault.h b/Engine/Source/Editor/GraphEditor/Public/SGraphNodeDefault.h
--- a/Engine/Source/Editor/GraphEditor/Public/SGraphNodeDefault.h
+++ b/Engine/Source/Editor/GraphEditor/Public/SGraphNodeDefault.h
@@ -7,9 +7,11 @@
 #include "SGraphNode.h"
 #include "Widgets/DeclarativeSyntaxSupport.h"
 
+#define UE_API GRAPHEDITOR_API
+
 class UEdGraphNode;
 
-class GRAPHEDITOR_API SGraphNodeDefault : public SGraphNode
+class SGraphNodeDefault : public SGraphNode
 {
 public:
 
@@ -20,5 +22,7 @@ public:
 		SLATE_ARGUMENT( UEdGraphNode*, GraphNodeObj )
 	SLATE_END_ARGS()
 
-	void Construct( const FArguments& InArgs );
+	UE_API void Construct( const FArguments& InArgs );
 };
+
+#undef UE_API
diff --git a/Engine/Source/Editor/GraphEditor/Public/SGraphNodeDocumentation.h b/Engine/Source/Editor/GraphEditor/Public/SGraphNodeDocumentation.h
--- a/Engine/Source/Editor/GraphEditor/Public/SGraphNodeDocumentation.h
+++ b/Engine/Source/Editor/GraphEditor/Public/SGraphNodeDocumentation.h
@@ -15,57 +15,59 @@
 #include "Types/SlateStructs.h"
 #include "Widgets/DeclarativeSyntaxSupport.h"
 
+#define UE_API GRAPHEDITOR_API
+
 class SWidget;
 class UEdGraphNode;
 struct FGeometry;
 struct FPointerEvent;
 
-class GRAPHEDITOR_API SGraphNodeDocumentation : public SGraphNodeResizable
+class SGraphNodeDocumentation : public SGraphNodeResizable
 {
 public:
 	SLATE_BEGIN_ARGS(SGraphNodeDocumentation){}
 	SLATE_END_ARGS()
 
 	//~ Begin SPanel Interface
-	virtual FVector2D ComputeDesiredSize(float) const override;
+	UE_API virtual FVector2D ComputeDesiredSize(float) const override;
 	//~ End SPanel Interface
 
 	//~ Begin SWidget Interface
-	void Construct( const FArguments& InArgs, UEdGraphNode* InNode );
-	virtual FReply OnMouseWheel( const FGeometry& MyGeometry, const FPointerEvent& MouseEvent ) override;
-	virtual void Tick( const FGeometry& AllottedGeometry, const double InCurrentTime, const float InDeltaTime ) override;
+	UE_API void Construct( const FArguments& InArgs, UEdGraphNode* InNode );
+	UE_API virtual FReply OnMouseWheel( const FGeometry& MyGeometry, const FPointerEvent& MouseEvent ) override;
+	UE_API virtual void Tick( const FGeometry& AllottedGeometry, const double InCurrentTime, const float InDeltaTime ) override;
 	//~ End SWidget Interface
 
 	//~ Begin SGraphNodeResizable Interface
-	virtual FVector2f GetNodeMinimumSize2f() const override;
-	virtual FVector2f GetNodeMaximumSize2f() const override;
-	virtual float GetTitleBarHeight() const override;
-	virtual FSlateRect GetHitTestingBorder() const override;
+	UE_API virtual FVector2f GetNodeMinimumSize2f() const override;
+	UE_API virtual FVector2f GetNodeMaximumSize2f() const override;
+	UE_API virtual float GetTitleBarHeight() const override;
+	UE_API virtual FSlateRect GetHitTestingBorder() const override;
 	//~ End SGraphNodeResizable Interface
 
 protected:
 
 	//~ Begin SGraphNode Interface
-	virtual void UpdateGraphNode() override;
+	UE_API virtual void UpdateGraphNode() override;
 	virtual bool IsNameReadOnly () const override { return false; }
 	//~ End SGraphNode Interface
 
 	/** Retrieves the current documentation title based on the chosen excerpt */
-	FText GetDocumentationTitle() const;
+	UE_API FText GetDocumentationTitle() const;
 	/** Create documentation page from link and excerpt */
-	TSharedPtr<SWidget> CreateDocumentationPage();
+	UE_API TSharedPtr<SWidget> CreateDocumentationPage();
 	/** Returns the width the documentation content must adhere to, used as a delegate in child widgets */
-	FOptionalSize GetContentWidth() const;
+	UE_API FOptionalSize GetContentWidth() const;
 	/** Returns the height the documentation content must adhere to, used as a delegate in child widgets */
-	FOptionalSize GetContentHeight() const;
+	UE_API FOptionalSize GetContentHeight() const;
 	/** Returns the child widget text wrapat size, used as a delegate during creation of documentation page */
-	float GetDocumentationWrapWidth() const;
+	UE_API float GetDocumentationWrapWidth() const;
 	/** Returns the current child widgets visibility for hit testing */
-	EVisibility GetWidgetVisibility() const;
+	UE_API EVisibility GetWidgetVisibility() const;
 	/** Color of the page gradient start */
-	FLinearColor GetPageGradientStartColor() const;
+	UE_API FLinearColor GetPageGradientStartColor() const;
 	/** Color of the page gradient end */
-	FLinearColor GetPageGradientEndColor() const;
+	UE_API FLinearColor GetPageGradientEndColor() const;
 
 private:
 
@@ -83,3 +85,5 @@ private:
 	FString CachedDocumentationExcerpt;
 
 };
+
+#undef UE_API
diff --git a/Engine/Source/Editor/GraphEditor/Public/SGraphNodeKnot.h b/Engine/Source/Editor/GraphEditor/Public/SGraphNodeKnot.h
--- a/Engine/Source/Editor/GraphEditor/Public/SGraphNodeKnot.h
+++ b/Engine/Source/Editor/GraphEditor/Public/SGraphNodeKnot.h
@@ -10,6 +10,8 @@
 #include "Types/SlateEnums.h"
 #include "Widgets/DeclarativeSyntaxSupport.h"
 
+#define UE_API GRAPHEDITOR_API
+
 class FText;
 class SCommentBubble;
 class SGraphPin;
@@ -19,32 +21,32 @@ struct FSlateBrush;
 /** The visual representation of a control point meant to adjust how connections are routed, also known as a Reroute node.
  * The input knot node should have properly implemented ShouldDrawNodeAsControlPointOnly to return true with valid indices for its pins.
  */
-class GRAPHEDITOR_API SGraphNodeKnot : public SGraphNodeDefault
+class SGraphNodeKnot : public SGraphNodeDefault
 {
 public:
 	SLATE_BEGIN_ARGS(SGraphNodeKnot) {}
 	SLATE_END_ARGS()
 
-	void Construct(const FArguments& InArgs, class UEdGraphNode* InKnot);
+	UE_API void Construct(const FArguments& InArgs, class UEdGraphNode* InKnot);
 
 	// SGraphNode interface
-	virtual void UpdateGraphNode() override;
-	virtual const FSlateBrush* GetShadowBrush(bool bSelected) const override;
-	virtual TSharedPtr<SGraphPin> CreatePinWidget(UEdGraphPin* Pin) const override;
-	virtual void AddPin(const TSharedRef<SGraphPin>& PinToAdd) override;
+	UE_API virtual void UpdateGraphNode() override;
+	UE_API virtual const FSlateBrush* GetShadowBrush(bool bSelected) const override;
+	UE_API virtual TSharedPtr<SGraphPin> CreatePinWidget(UEdGraphPin* Pin) const override;
+	UE_API virtual void AddPin(const TSharedRef<SGraphPin>& PinToAdd) override;
 	virtual void RequestRenameOnSpawn() override { }
 	// End of SGraphNode interface
 
 protected:
 	/** Returns Offset to center comment on the node's only pin */
-	FVector2f GetCommentOffset() const;
+	UE_API FVector2f GetCommentOffset() const;
 protected:
 
 	/** Toggles the hovered visibility state */
-	virtual void OnCommentBubbleToggled(bool bInCommentBubbleVisible) override;
+	UE_API virtual void OnCommentBubbleToggled(bool bInCommentBubbleVisible) override;
 
 	/** If bHoveredCommentVisibility is true, hides the comment bubble after a change is committed */
-	virtual void OnCommentTextCommitted(const FText& NewComment, ETextCommit::Type CommitInfo) override;
+	UE_API virtual void OnCommentTextCommitted(const FText& NewComment, ETextCommit::Type CommitInfo) override;
 
 	/** The hovered visibility state. If false, comment bubble will only appear on hover. */
 	bool bAlwaysShowCommentBubble;
@@ -56,24 +58,25 @@ protected:
 	const FSlateBrush* ShadowBrushSelected;
 };
 
-class GRAPHEDITOR_API SGraphPinKnot : public SGraphPin
+class SGraphPinKnot : public SGraphPin
 {
 public:
 	SLATE_BEGIN_ARGS(SGraphPinKnot) {}
 	SLATE_END_ARGS()
 
-	void Construct(const FArguments& InArgs, UEdGraphPin* InPin);
+	UE_API void Construct(const FArguments& InArgs, UEdGraphPin* InPin);
 
 	// SWidget interface
-	virtual void OnDragEnter(const FGeometry& MyGeometry, const FDragDropEvent& DragDropEvent) override;
+	UE_API virtual void OnDragEnter(const FGeometry& MyGeometry, const FDragDropEvent& DragDropEvent) override;
 	// End of SWidget interface
 
 protected:
 	// Begin SGraphPin interface
-	virtual TSharedRef<SWidget>	GetDefaultValueWidget() override;
-	virtual TSharedRef<FDragDropOperation> SpawnPinDragEvent(const TSharedRef<SGraphPanel>& InGraphPanel, const TArray< TSharedRef<SGraphPin> >& InStartingPins) override;
-	virtual FReply OnPinMouseDown(const FGeometry& SenderGeometry, const FPointerEvent& MouseEvent) override;
-	virtual FSlateColor GetPinColor() const override;
+	UE_API virtual TSharedRef<SWidget>	GetDefaultValueWidget() override;
+	UE_API virtual TSharedRef<FDragDropOperation> SpawnPinDragEvent(const TSharedRef<SGraphPanel>& InGraphPanel, const TArray< TSharedRef<SGraphPin> >& InStartingPins) override;
+	UE_API virtual FReply OnPinMouseDown(const FGeometry& SenderGeometry, const FPointerEvent& MouseEvent) override;
+	UE_API virtual FSlateColor GetPinColor() const override;
 	// End SGraphPin interface
 };
 
+#undef UE_API
diff --git a/Engine/Source/Editor/GraphEditor/Public/SGraphNodeResizable.h b/Engine/Source/Editor/GraphEditor/Public/SGraphNodeResizable.h
--- a/Engine/Source/Editor/GraphEditor/Public/SGraphNodeResizable.h
+++ b/Engine/Source/Editor/GraphEditor/Public/SGraphNodeResizable.h
@@ -10,11 +10,13 @@
 #include "SGraphNode.h"
 #include "Templates/SharedPointer.h"
 
+#define UE_API GRAPHEDITOR_API
+
 class FScopedTransaction;
 struct FGeometry;
 struct FPointerEvent;
 
-class GRAPHEDITOR_API SGraphNodeResizable : public SGraphNode
+class SGraphNodeResizable : public SGraphNode
 {
 public:
 
@@ -37,48 +39,48 @@ public:
 	};
 
 	//~ Begin SWidget Interface
-	virtual FReply OnMouseMove( const FGeometry& MyGeometry, const FPointerEvent& MouseEvent ) override;
-	virtual void OnMouseEnter( const FGeometry& MyGeometry, const FPointerEvent& MouseEvent ) override;
-	virtual void OnMouseLeave( const FPointerEvent& MouseEvent ) override;
-	virtual FReply OnMouseButtonDown( const FGeometry& MyGeometry, const FPointerEvent& MouseEvent ) override;
-	virtual FReply OnMouseButtonUp( const FGeometry& MyGeometry, const FPointerEvent& MouseEvent ) override;
-	virtual FCursorReply OnCursorQuery( const FGeometry& MyGeometry, const FPointerEvent& CursorEvent ) const override;
+	UE_API virtual FReply OnMouseMove( const FGeometry& MyGeometry, const FPointerEvent& MouseEvent ) override;
+	UE_API virtual void OnMouseEnter( const FGeometry& MyGeometry, const FPointerEvent& MouseEvent ) override;
+	UE_API virtual void OnMouseLeave( const FPointerEvent& MouseEvent ) override;
+	UE_API virtual FReply OnMouseButtonDown( const FGeometry& MyGeometry, const FPointerEvent& MouseEvent ) override;
+	UE_API virtual FReply OnMouseButtonUp( const FGeometry& MyGeometry, const FPointerEvent& MouseEvent ) override;
+	UE_API virtual FCursorReply OnCursorQuery( const FGeometry& MyGeometry, const FPointerEvent& CursorEvent ) const override;
 	//~ End SWidget Interface
 	
 protected:
 
 	/** Find the current window zone the mouse is in */
 	UE_DEPRECATED(5.6, "Use the version of the function accepting FVector2f; this Slate API no longer interfaces directly with double-precision scalars and vectors.")
-	virtual SGraphNodeResizable::EResizableWindowZone FindMouseZone(const FVector2D& LocalMouseCoordinates) const UE_SLATE_DEPRECATED_VECTOR_VIRTUAL_FUNCTION;
-	virtual SGraphNodeResizable::EResizableWindowZone FindMouseZone(const FVector2f& LocalMouseCoordinates) const;
+	UE_API virtual SGraphNodeResizable::EResizableWindowZone FindMouseZone(const FVector2D& LocalMouseCoordinates) const UE_SLATE_DEPRECATED_VECTOR_VIRTUAL_FUNCTION;
+	UE_API virtual SGraphNodeResizable::EResizableWindowZone FindMouseZone(const FVector2f& LocalMouseCoordinates) const;
 
 	/** @return true if the current window zone is considered a selection area */
 	bool InSelectionArea() const { return InSelectionArea(MouseZone); }
 
 	/** @return true if the passed zone is a selection area */
-	bool InSelectionArea(EResizableWindowZone InZone) const;
+	UE_API bool InSelectionArea(EResizableWindowZone InZone) const;
 
 	/** Function to store anchor point before resizing the node. The node will be anchored to this point when resizing happens*/
-	void InitNodeAnchorPoint();
+	UE_API void InitNodeAnchorPoint();
 
 	/** Function to fetch the corrected node position based on anchor point*/
-	UE::Slate::FDeprecateVector2DResult GetCorrectedNodePosition() const;
+	UE_API UE::Slate::FDeprecateVector2DResult GetCorrectedNodePosition() const;
 
 	/** Get the current titlebar size */
-	virtual float GetTitleBarHeight() const;
+	UE_API virtual float GetTitleBarHeight() const;
 
 	/** Return smallest desired node size */
 	UE_DEPRECATED(5.6, "Use the version of the function accepting FVector2f; this Slate API no longer interfaces directly with double-precision scalars and vectors.")
-	virtual FVector2D GetNodeMinimumSize() const UE_SLATE_DEPRECATED_VECTOR_VIRTUAL_FUNCTION;
-	virtual FVector2f GetNodeMinimumSize2f() const;
+	UE_API virtual FVector2D GetNodeMinimumSize() const UE_SLATE_DEPRECATED_VECTOR_VIRTUAL_FUNCTION;
+	UE_API virtual FVector2f GetNodeMinimumSize2f() const;
 
 	/** Return largest desired node size */
 	UE_DEPRECATED(5.6, "Use the version of the function accepting FVector2f; this Slate API no longer interfaces directly with double-precision scalars and vectors.")
-	virtual FVector2D GetNodeMaximumSize() const UE_SLATE_DEPRECATED_VECTOR_VIRTUAL_FUNCTION;
-	virtual FVector2f GetNodeMaximumSize2f() const;
+	UE_API virtual FVector2D GetNodeMaximumSize() const UE_SLATE_DEPRECATED_VECTOR_VIRTUAL_FUNCTION;
+	UE_API virtual FVector2f GetNodeMaximumSize2f() const;
 
 	//** Return slate rect border for hit testing */
-	virtual FSlateRect GetHitTestingBorder() const;
+	UE_API virtual FSlateRect GetHitTestingBorder() const;
 
 protected:
 
@@ -104,3 +106,5 @@ protected:
 	bool bUserIsDragging;
 
 };
+
+#undef UE_API
diff --git a/Engine/Source/Editor/GraphEditor/Public/SGraphPalette.h b/Engine/Source/Editor/GraphEditor/Public/SGraphPalette.h
--- a/Engine/Source/Editor/GraphEditor/Public/SGraphPalette.h
+++ b/Engine/Source/Editor/GraphEditor/Public/SGraphPalette.h
@@ -18,6 +18,8 @@
 #include "Widgets/SCompoundWidget.h"
 #include "Widgets/SWidget.h"
 
+#define UE_API GRAPHEDITOR_API
+
 class FDragDropEvent;
 class SInlineEditableTextBlock;
 class SWidget;
@@ -31,7 +33,7 @@ struct FSlateFontInfo;
 //////////////////////////////////////////////////////////////////////////
 
 /** Widget for displaying a single item  */
-class GRAPHEDITOR_API SGraphPaletteItem : public SCompoundWidget
+class SGraphPaletteItem : public SCompoundWidget
 {
 public:
 	/** Delegate executed when mouse button goes down */
@@ -45,20 +47,20 @@ public:
 	SLATE_BEGIN_ARGS( SGraphPaletteItem ) {};
 	SLATE_END_ARGS()
 
-	void Construct(const FArguments& InArgs, FCreateWidgetForActionData* const InCreateData);
+	UE_API void Construct(const FArguments& InArgs, FCreateWidgetForActionData* const InCreateData);
 
 	// SWidget interface
-	virtual FReply OnDrop( const FGeometry& MyGeometry, const FDragDropEvent& DragDropEvent ) override;
-	virtual void OnDragEnter( const FGeometry& MyGeometry, const FDragDropEvent& DragDropEvent ) override;
-	virtual void OnDragLeave( const FDragDropEvent& DragDropEvent ) override;
-	virtual FReply OnMouseButtonDown( const FGeometry& MyGeometry, const FPointerEvent& MouseEvent ) override;
+	UE_API virtual FReply OnDrop( const FGeometry& MyGeometry, const FDragDropEvent& DragDropEvent ) override;
+	UE_API virtual void OnDragEnter( const FGeometry& MyGeometry, const FDragDropEvent& DragDropEvent ) override;
+	UE_API virtual void OnDragLeave( const FDragDropEvent& DragDropEvent ) override;
+	UE_API virtual FReply OnMouseButtonDown( const FGeometry& MyGeometry, const FPointerEvent& MouseEvent ) override;
 	// End of SWidget interface
 
 protected:
 	/** Creates an icon type widget, including big tooltip link */
-	TSharedRef<SWidget> CreateIconWidget( const FText& IconToolTip, const FSlateBrush* IconBrush, const FSlateColor& IconColor, const FString& DocLink, const FString& DocExcerpt, const FSlateBrush* SecondaryIconBrush, const FSlateColor& SecondaryColor );
+	UE_API TSharedRef<SWidget> CreateIconWidget( const FText& IconToolTip, const FSlateBrush* IconBrush, const FSlateColor& IconColor, const FString& DocLink, const FString& DocExcerpt, const FSlateBrush* SecondaryIconBrush, const FSlateColor& SecondaryColor );
 	/** Create an icon type widget */
-	TSharedRef<SWidget> CreateIconWidget(const FText& IconToolTip, const FSlateBrush* IconBrush, const FSlateColor& IconColor);
+	UE_API TSharedRef<SWidget> CreateIconWidget(const FText& IconToolTip, const FSlateBrush* IconBrush, const FSlateColor& IconColor);
 
 	UE_DEPRECATED(5.0, "CreateTextSlotWidget that takes a font parameter is deprecated")
 	virtual TSharedRef<SWidget> CreateTextSlotWidget(const FSlateFontInfo& NameFont, FCreateWidgetForActionData* const InCreateData, TAttribute<bool> bIsReadOnly)
@@ -67,24 +69,24 @@ protected:
 	}
 
 	/* Create the text widget */
-	virtual TSharedRef<SWidget> CreateTextSlotWidget( FCreateWidgetForActionData* const InCreateData, TAttribute<bool> bIsReadOnly );
+	UE_API virtual TSharedRef<SWidget> CreateTextSlotWidget( FCreateWidgetForActionData* const InCreateData, TAttribute<bool> bIsReadOnly );
 
 	/** Callback when rename text is being verified on text changed */
-	virtual bool OnNameTextVerifyChanged(const FText& InNewText, FText& OutErrorMessage);
+	UE_API virtual bool OnNameTextVerifyChanged(const FText& InNewText, FText& OutErrorMessage);
 
 	/** Callback when rename text is committed */
-	virtual void OnNameTextCommitted(const FText& NewText, ETextCommit::Type InTextCommit);
+	UE_API virtual void OnNameTextCommitted(const FText& NewText, ETextCommit::Type InTextCommit);
 
 	/** Callback to get the display text */ 
-	virtual FText GetDisplayText() const;
+	UE_API virtual FText GetDisplayText() const;
 
 	/** Callback to get the tooltip */
-	virtual FText GetItemTooltip() const;
+	UE_API virtual FText GetItemTooltip() const;
 };
 
 //////////////////////////////////////////////////////////////////////////
 
-class GRAPHEDITOR_API SGraphPalette : public SCompoundWidget
+class SGraphPalette : public SCompoundWidget
 {
 public:
 	SLATE_BEGIN_ARGS( SGraphPalette )
@@ -94,15 +96,15 @@ public:
 		SLATE_ARGUMENT( bool, AutoExpandActionMenu )
 	SLATE_END_ARGS()
 
-	void Construct(const FArguments& InArgs);
+	UE_API void Construct(const FArguments& InArgs);
 	virtual ~SGraphPalette() {};
 
 protected:
-	virtual TSharedRef<SWidget> OnCreateWidgetForAction(FCreateWidgetForActionData* const InCreateData);
+	UE_API virtual TSharedRef<SWidget> OnCreateWidgetForAction(FCreateWidgetForActionData* const InCreateData);
 
-	virtual FReply OnActionDragged(const TArray< TSharedPtr<FEdGraphSchemaAction> >& InActions, const FPointerEvent& MouseEvent);
+	UE_API virtual FReply OnActionDragged(const TArray< TSharedPtr<FEdGraphSchemaAction> >& InActions, const FPointerEvent& MouseEvent);
 
-	virtual void RefreshActionsList(bool bPreserveExpansion);
+	UE_API virtual void RefreshActionsList(bool bPreserveExpansion);
 
 	/** Callback used to populate all actions list in SGraphActionMenu */
 	virtual void CollectAllActions(FGraphActionListBuilderBase& OutAllActions) = 0;
@@ -111,3 +113,4 @@ protected:
 	TSharedPtr<SGraphActionMenu> GraphActionMenu;
 };
 
+#undef UE_API
