Updated headers using UnrealCodeFixup to make sure dllstorage is on methods/staticvars instead of types. Part 5/n

diff --git a/Engine/Source/Editor/GraphEditor/Public/SGraphPanel.h b/Engine/Source/Editor/GraphEditor/Public/SGraphPanel.h
--- a/Engine/Source/Editor/GraphEditor/Public/SGraphPanel.h
+++ b/Engine/Source/Editor/GraphEditor/Public/SGraphPanel.h
@@ -32,6 +32,8 @@
 #include "Widgets/DeclarativeSyntaxSupport.h"
 #include "Widgets/SWidget.h"
 
+#define UE_API GRAPHEDITOR_API
+
 class FActiveTimerHandle;
 class FArrangedChildren;
 class FDragDropEvent;
@@ -72,7 +74,7 @@ struct FGraphContextMenuArguments
 };
 
 
-class GRAPHEDITOR_API SGraphPanel : public SNodePanel, public FGCObject
+class SGraphPanel : public SNodePanel, public FGCObject
 {
 public:
 	DECLARE_DELEGATE_RetVal_OneParam(FActionMenuContent, FOnGetContextMenuFor, const FGraphContextMenuArguments& /*SpawnInfo*/)
@@ -172,81 +174,81 @@ PRAGMA_ENABLE_DEPRECATION_WARNINGS
 	 *
 	 * @param InArgs    The declaration describing how the widgets should be constructed.
 	 */
-	void Construct( const FArguments& InArgs );
+	UE_API void Construct( const FArguments& InArgs );
 
 	// Destructor
-	~SGraphPanel();
+	UE_API ~SGraphPanel();
 public:
 	// SWidget interface
-	virtual FReply OnMouseButtonDown(const FGeometry& MyGeometry, const FPointerEvent& MouseEvent) override;
-	virtual FReply OnMouseButtonUp(const FGeometry& MyGeometry, const FPointerEvent& MouseEvent) override;
-	virtual FReply OnMouseButtonDoubleClick(const FGeometry& InMyGeometry, const FPointerEvent& InMouseEvent) override;
-	virtual FReply OnMouseMove(const FGeometry& MyGeometry, const FPointerEvent& MouseEvent) override;
-	virtual void OnDragEnter(const FGeometry& MyGeometry, const FDragDropEvent& DragDropEvent) override;
-	virtual void OnDragLeave( const FDragDropEvent& DragDropEvent ) override;
-	virtual FReply OnDragOver( const FGeometry& MyGeometry, const FDragDropEvent& DragDropEvent ) override;
-	virtual FReply OnDrop( const FGeometry& MyGeometry, const FDragDropEvent& DragDropEvent ) override;
-	virtual int32 OnPaint( const FPaintArgs& Args, const FGeometry& AllottedGeometry, const FSlateRect& MyCullingRect, FSlateWindowElementList& OutDrawElements, int32 LayerId, const FWidgetStyle& InWidgetStyle, bool bParentEnabled ) const override;
-	virtual FReply OnKeyDown( const FGeometry& MyGeometry, const FKeyEvent& InKeyEvent ) override;
-	virtual bool SupportsKeyboardFocus() const override;
-	virtual void OnArrangeChildren( const FGeometry& AllottedGeometry, FArrangedChildren& ArrangedChildren ) const override;
-	virtual bool CustomPrepass(float LayoutScaleMultiplier) override;
+	UE_API virtual FReply OnMouseButtonDown(const FGeometry& MyGeometry, const FPointerEvent& MouseEvent) override;
+	UE_API virtual FReply OnMouseButtonUp(const FGeometry& MyGeometry, const FPointerEvent& MouseEvent) override;
+	UE_API virtual FReply OnMouseButtonDoubleClick(const FGeometry& InMyGeometry, const FPointerEvent& InMouseEvent) override;
+	UE_API virtual FReply OnMouseMove(const FGeometry& MyGeometry, const FPointerEvent& MouseEvent) override;
+	UE_API virtual void OnDragEnter(const FGeometry& MyGeometry, const FDragDropEvent& DragDropEvent) override;
+	UE_API virtual void OnDragLeave( const FDragDropEvent& DragDropEvent ) override;
+	UE_API virtual FReply OnDragOver( const FGeometry& MyGeometry, const FDragDropEvent& DragDropEvent ) override;
+	UE_API virtual FReply OnDrop( const FGeometry& MyGeometry, const FDragDropEvent& DragDropEvent ) override;
+	UE_API virtual int32 OnPaint( const FPaintArgs& Args, const FGeometry& AllottedGeometry, const FSlateRect& MyCullingRect, FSlateWindowElementList& OutDrawElements, int32 LayerId, const FWidgetStyle& InWidgetStyle, bool bParentEnabled ) const override;
+	UE_API virtual FReply OnKeyDown( const FGeometry& MyGeometry, const FKeyEvent& InKeyEvent ) override;
+	UE_API virtual bool SupportsKeyboardFocus() const override;
+	UE_API virtual void OnArrangeChildren( const FGeometry& AllottedGeometry, FArrangedChildren& ArrangedChildren ) const override;
+	UE_API virtual bool CustomPrepass(float LayoutScaleMultiplier) override;
 	// End of SWidget interface
 
 	// SNodePanel interface
-	virtual TSharedPtr<SWidget> OnSummonContextMenu(const FGeometry& MyGeometry, const FPointerEvent& MouseEvent) override;
-	virtual bool OnHandleLeftMouseRelease(const FGeometry& MyGeometry, const FPointerEvent& MouseEvent) override;
-	virtual void AddGraphNode(const TSharedRef<SNode>& NodeToAdd) override;
-	virtual void RemoveAllNodes() override;
+	UE_API virtual TSharedPtr<SWidget> OnSummonContextMenu(const FGeometry& MyGeometry, const FPointerEvent& MouseEvent) override;
+	UE_API virtual bool OnHandleLeftMouseRelease(const FGeometry& MyGeometry, const FPointerEvent& MouseEvent) override;
+	UE_API virtual void AddGraphNode(const TSharedRef<SNode>& NodeToAdd) override;
+	UE_API virtual void RemoveAllNodes() override;
 	// End of SNodePanel interface
 
 	// FGCObject interface.
-	virtual void AddReferencedObjects( FReferenceCollector& Collector ) override;
-	virtual FString GetReferencerName() const override;
+	UE_API virtual void AddReferencedObjects( FReferenceCollector& Collector ) override;
+	UE_API virtual FString GetReferencerName() const override;
 	// End of FGCObject interface.
 
-	void ArrangeChildrenForContextMenuSummon(const FGeometry& AllottedGeometry, FArrangedChildren& ArrangedChildren) const;
-	TSharedPtr<SWidget> SummonContextMenu(const UE::Slate::FDeprecateVector2DParameter& WhereToSummon, const UE::Slate::FDeprecateVector2DParameter& WhereToAddNode, UEdGraphNode* ForNode, UEdGraphPin* ForPin, const TArray<UEdGraphPin*>& DragFromPins);
-	void SummonCreateNodeMenuFromUICommand(uint32 NumNodesAdded);
-	void DismissContextMenu();
+	UE_API void ArrangeChildrenForContextMenuSummon(const FGeometry& AllottedGeometry, FArrangedChildren& ArrangedChildren) const;
+	UE_API TSharedPtr<SWidget> SummonContextMenu(const UE::Slate::FDeprecateVector2DParameter& WhereToSummon, const UE::Slate::FDeprecateVector2DParameter& WhereToAddNode, UEdGraphNode* ForNode, UEdGraphPin* ForPin, const TArray<UEdGraphPin*>& DragFromPins);
+	UE_API void SummonCreateNodeMenuFromUICommand(uint32 NumNodesAdded);
+	UE_API void DismissContextMenu();
 
-	void OnBeginMakingConnection(UEdGraphPin* InOriginatingPin);
-	void OnBeginMakingConnection(FGraphPinHandle PinHandle);
-	void OnStopMakingConnection(bool bForceStop = false);
-	void PreservePinPreviewUntilForced();
+	UE_API void OnBeginMakingConnection(UEdGraphPin* InOriginatingPin);
+	UE_API void OnBeginMakingConnection(FGraphPinHandle PinHandle);
+	UE_API void OnStopMakingConnection(bool bForceStop = false);
+	UE_API void PreservePinPreviewUntilForced();
 
 	/** Indicate that the connection from the given start to the given end pins is being relinked. A preview connection is being drawn for the relinked connection. */
-	void OnBeginRelinkConnection(const FGraphPinHandle& InSourcePinHandle, const FGraphPinHandle& InTargetPinHandle);
+	UE_API void OnBeginRelinkConnection(const FGraphPinHandle& InSourcePinHandle, const FGraphPinHandle& InTargetPinHandle);
 
 	/** The relink connection operation either got cancelled or has successfully been executed. Preview connection won't be drawn anymore. */
-	void OnEndRelinkConnection(bool bForceStop = false);
+	UE_API void OnEndRelinkConnection(bool bForceStop = false);
 
 	/** True in case a connection is currently being relinked, false if not. */
-	bool IsRelinkingConnection() const;
+	UE_API bool IsRelinkingConnection() const;
 
 	/** Update this GraphPanel to match the data that it is observing. Expected to be called during ticking. */
-	void Update();
+	UE_API void Update();
 
 	/** Purges the existing visual representation (typically followed by an Update call in the next tick) */
-	void PurgeVisualRepresentation();
+	UE_API void PurgeVisualRepresentation();
 
 	/** Use to determine if a comment title is currently visible */
-	bool IsNodeTitleVisible(const class UEdGraphNode* Node, bool bRequestRename);
+	UE_API bool IsNodeTitleVisible(const class UEdGraphNode* Node, bool bRequestRename);
 
 	/** Use to determine if a rectangle is currently visible */
-	bool IsRectVisible(const UE::Slate::FDeprecateVector2DParameter &TopLeft, const UE::Slate::FDeprecateVector2DParameter &BottomRight);
+	UE_API bool IsRectVisible(const UE::Slate::FDeprecateVector2DParameter &TopLeft, const UE::Slate::FDeprecateVector2DParameter &BottomRight);
 
 	/** Focuses the view on rectangle, zooming if neccesary */
-	bool JumpToRect(const UE::Slate::FDeprecateVector2DParameter &BottomLeft, const UE::Slate::FDeprecateVector2DParameter &TopRight);
+	UE_API bool JumpToRect(const UE::Slate::FDeprecateVector2DParameter &BottomLeft, const UE::Slate::FDeprecateVector2DParameter &TopRight);
 
-	void JumpToNode(const class UEdGraphNode* JumpToMe, bool bRequestRename, bool bSelectNode);
+	UE_API void JumpToNode(const class UEdGraphNode* JumpToMe, bool bRequestRename, bool bSelectNode);
 
-	void JumpToPin(const class UEdGraphPin* JumptToMe);
+	UE_API void JumpToPin(const class UEdGraphPin* JumptToMe);
 
-	void GetAllPins(TSet< TSharedRef<SWidget> >& AllPins);
+	UE_API void GetAllPins(TSet< TSharedRef<SWidget> >& AllPins);
 
-	void AddPinToHoverSet(UEdGraphPin* HoveredPin);
-	void RemovePinFromHoverSet(UEdGraphPin* UnhoveredPin);
+	UE_API void AddPinToHoverSet(UEdGraphPin* HoveredPin);
+	UE_API void RemovePinFromHoverSet(UEdGraphPin* UnhoveredPin);
 
 	SGraphEditor::EPinVisibility GetPinVisibility() const { return PinVisibility; }
 	void SetPinVisibility(SGraphEditor::EPinVisibility InVisibility) { PinVisibility = InVisibility; }
@@ -254,36 +256,36 @@ public:
 	UEdGraph* GetGraphObj() const { return GraphObj; }
 
 	/** helper to attach graph events to sub node, which won't be placed directly on the graph */
-	void AttachGraphEvents(TSharedPtr<SGraphNode> CreatedSubNode);
+	UE_API void AttachGraphEvents(TSharedPtr<SGraphNode> CreatedSubNode);
 
 	/** Returns if this graph is editable */
 	bool IsGraphEditable() const { return IsEditable.Get(); }
 
 	/** Attempt to retrieve the bounds for the specified node */
 	UE_DEPRECATED(5.6, "Use the version of the function accepting FVector2f; this Slate API no longer interfaces directly with double-precision scalars and vectors.")
-	bool GetBoundsForNode(const UObject* InNode, FVector2D& MinCorner, FVector2D& MaxCorner, float Padding = 0.0f) const;
-	bool GetBoundsForNode(const UObject* InNode, FVector2f& MinCorner, FVector2f& MaxCorner, float Padding = 0.0f) const;
+	UE_API bool GetBoundsForNode(const UObject* InNode, FVector2D& MinCorner, FVector2D& MaxCorner, float Padding = 0.0f) const;
+	UE_API bool GetBoundsForNode(const UObject* InNode, FVector2f& MinCorner, FVector2f& MaxCorner, float Padding = 0.0f) const;
 
 	/** Straighten all connections between the selected nodes */
-	void StraightenConnections();
+	UE_API void StraightenConnections();
 	
 	/** Straighten any connections attached to the specified pin, optionally limiting to the specified pin to align */
-	void StraightenConnections(UEdGraphPin* SourcePin, UEdGraphPin* PinToAlign = nullptr);
+	UE_API void StraightenConnections(UEdGraphPin* SourcePin, UEdGraphPin* PinToAlign = nullptr);
 
 	/** Refresh the visual state of a single node */
-	void RefreshNode(UEdGraphNode& Node);
+	UE_API void RefreshNode(UEdGraphNode& Node);
 
 	/** When the graph panel needs to be dynamically refreshing for animations, this function is registered to tick and invalidate the UI. */
-	EActiveTimerReturnType InvalidatePerTick(double InCurrentTime, float InDeltaTime);
+	UE_API EActiveTimerReturnType InvalidatePerTick(double InCurrentTime, float InDeltaTime);
 
 	/** Sets the current widget factory. */
-	void SetNodeFactory(const TSharedRef<class FGraphNodeFactory>& NewNodeFactory);
+	UE_API void SetNodeFactory(const TSharedRef<class FGraphNodeFactory>& NewNodeFactory);
 
 protected:
 
-	void NotifyGraphChanged ( const struct FEdGraphEditAction& InAction);
+	UE_API void NotifyGraphChanged ( const struct FEdGraphEditAction& InAction);
 
-	const TSharedRef<SGraphNode> GetChild(int32 ChildIndex);
+	UE_API const TSharedRef<SGraphNode> GetChild(int32 ChildIndex);
 
 	/** Flag to control AddNode, more readable than a bool:*/
 	enum AddNodeBehavior
@@ -294,22 +296,22 @@ protected:
 	};
 
 	/** Helper method to add a new node to the panel */
-	void AddNode(class UEdGraphNode* Node, AddNodeBehavior Behavior);
+	UE_API void AddNode(class UEdGraphNode* Node, AddNodeBehavior Behavior);
 
 	/** Helper method to remove a node from the panel */
-	void RemoveNode(const UEdGraphNode* Node);
+	UE_API void RemoveNode(const UEdGraphNode* Node);
 
 	/** Helper method to remove all nodes from the panel holding garbage / invalid pointers */
-	void RemoveAllNodesWithInvalidPointers();
+	UE_API void RemoveAllNodesWithInvalidPointers();
 public:
 	/** Pin marked via shift-clicking */
 	TWeakPtr<SGraphPin> MarkedPin;
 
 	/** Get a graph node widget from the specified GUID, if it applies to any nodes in this graph */
-	TSharedPtr<SGraphNode> GetNodeWidgetFromGuid(FGuid Guid) const;
+	UE_API TSharedPtr<SGraphNode> GetNodeWidgetFromGuid(FGuid Guid) const;
 
 	/** Get a list of selected editor graph nodes from the selection manager. */
-	TArray<UEdGraphNode*> GetSelectedGraphNodes() const;
+	UE_API TArray<UEdGraphNode*> GetSelectedGraphNodes() const;
 
 	const FGraphSplineOverlapResult& GetPreviousFrameSplineOverlap() const { return PreviousFrameSplineOverlap; }
 
@@ -414,40 +416,40 @@ private:
 	FDelegateHandle            MyRegisteredGraphChangedDelegateHandle;
 private:
 	/** Called when PIE begins */
-	void OnBeginPIE( const bool bIsSimulating );
+	UE_API void OnBeginPIE( const bool bIsSimulating );
 
 	/** Called when PIE ends */
-	void OnEndPIE( const bool bIsSimulating );
+	UE_API void OnEndPIE( const bool bIsSimulating );
 
 	/** Called when watched graph changes */
-	void OnGraphChanged( const FEdGraphEditAction& InAction );
+	UE_API void OnGraphChanged( const FEdGraphEditAction& InAction );
 
 	/** Update all selected nodes position by provided vector2d */
-	void UpdateSelectedNodesPositions(const FVector2f& PositionIncrement);
+	UE_API void UpdateSelectedNodesPositions(const FVector2f& PositionIncrement);
 
 	/** Handle updating the spline hover state */
-	bool OnSplineHoverStateChanged(const FGraphSplineOverlapResult& NewSplineHoverState);
+	UE_API bool OnSplineHoverStateChanged(const FGraphSplineOverlapResult& NewSplineHoverState);
 
 	/** Returns the pin that we're considering as hovered if we are hovering over a spline; may be null */
-	class SGraphPin* GetBestPinFromHoveredSpline() const;
+	UE_API class SGraphPin* GetBestPinFromHoveredSpline() const;
 
 	/** Returns true if all assets in the ReferencedAssets array are allowed to be referenced by this graph according to the AssetReferenceFilter. OutFailureReason if supplied if it is not.  */
-	bool PassesAssetReferenceFilter(const TArray<FAssetData>& ReferencedAssets, FText* OutFailureReason = nullptr) const;
+	UE_API bool PassesAssetReferenceFilter(const TArray<FAssetData>& ReferencedAssets, FText* OutFailureReason = nullptr) const;
 
 	/** Returns the top most graph node under the mouse pointer. Returns nullptr if no node found */
-	TSharedPtr<SGraphNode> GetGraphNodeUnderMouse(const FGeometry& MyGeometry, const FPointerEvent& MouseEvent);
+	UE_API TSharedPtr<SGraphNode> GetGraphNodeUnderMouse(const FGeometry& MyGeometry, const FPointerEvent& MouseEvent);
 
 	/** Returns a pin that is under the mouse, given a specified node. Returns nullptr if no valid node is given or no pin found */
-	UEdGraphPin* GetPinUnderMouse(const FGeometry& MyGeometry, const FPointerEvent& MouseEvent, TSharedPtr<SGraphNode> GraphNode) const;
+	UE_API UEdGraphPin* GetPinUnderMouse(const FGeometry& MyGeometry, const FPointerEvent& MouseEvent, TSharedPtr<SGraphNode> GraphNode) const;
 
 	/** If the spawned nodes were auto-wired from any of the dragged pins, then this will try to make the newly connected pin end up at SpawnGraphPosition */
-	void AdjustNewlySpawnedNodePositions(TArrayView<UEdGraphNode* const> SpawnedNodes, TArrayView<UEdGraphPin*> DraggedFromPins, FVector2f SpawnGraphPosition);
+	UE_API void AdjustNewlySpawnedNodePositions(TArrayView<UEdGraphNode* const> SpawnedNodes, TArrayView<UEdGraphPin*> DraggedFromPins, FVector2f SpawnGraphPosition);
 
 	/** Will move a group of nodes by the amount needed for an anchor pin to be at a certain position */
-	void MoveNodesToAnchorPinAtGraphPosition(TArrayView<UEdGraphNode* const> NodesToMove, FGraphPinHandle PinToAnchor, FVector2f DesiredPinGraphPosition);
+	UE_API void MoveNodesToAnchorPinAtGraphPosition(TArrayView<UEdGraphNode* const> NodesToMove, FGraphPinHandle PinToAnchor, FVector2f DesiredPinGraphPosition);
 
 	/** Begins a drag operation that disconnects the pin at ClosestPin, but only from the ConnectedPin, rather than all connections */
-	FReply OnMouseDownDisconnectClosestPinFromWire(TSharedRef<class SGraphPin> ClosestPin, TSharedRef<class SGraphPin> ConnectedPin);
+	UE_API FReply OnMouseDownDisconnectClosestPinFromWire(TSharedRef<class SGraphPin> ClosestPin, TSharedRef<class SGraphPin> ConnectedPin);
 
 	/** Handle to timer callback that allows the UI to refresh it's arrangement each tick, allows animations to occur within the UI */
 	TWeakPtr<FActiveTimerHandle> ActiveTimerHandleInvalidatePerTick;
@@ -464,3 +466,5 @@ private:
 	/** A flag indicating that we need to check if we have any nodes which use invalid pointers */
 	bool bCheckNodeGraphObjValidity;
 };
+
+#undef UE_API
diff --git a/Engine/Source/Editor/GraphEditor/Public/SGraphPin.h b/Engine/Source/Editor/GraphEditor/Public/SGraphPin.h
--- a/Engine/Source/Editor/GraphEditor/Public/SGraphPin.h
+++ b/Engine/Source/Editor/GraphEditor/Public/SGraphPin.h
@@ -29,6 +29,8 @@
 #include "Widgets/Layout/SWrapBox.h"
 #include "Widgets/SWidget.h"
 
+#define UE_API GRAPHEDITOR_API
+
 class FDragDropEvent;
 class FDragDropOperation;
 class FPinValueInspectorTooltip;
@@ -51,7 +53,7 @@ struct FSlateBrush;
 // FGraphPinHandle
 
 /** A handle to a pin, defined by its owning node's GUID, and the pin's name. Used to reference a pin without referring to its widget */
-struct GRAPHEDITOR_API FGraphPinHandle
+struct FGraphPinHandle
 {
 	/** The GUID of the node to which this pin belongs */
 	FGuid NodeGuid;
@@ -66,13 +68,13 @@ struct GRAPHEDITOR_API FGraphPinHandle
 	FGraphPinHandle() = default;
 
 	/** Constructor */
-	FGraphPinHandle(UEdGraphPin* InPin);
+	UE_API FGraphPinHandle(UEdGraphPin* InPin);
 
 	/** */
-	UEdGraphPin* GetPinObj(const SGraphPanel& InPanel) const;
+	UE_API UEdGraphPin* GetPinObj(const SGraphPanel& InPanel) const;
 
 	/** Find a pin widget in the specified panel from this handle */
-	TSharedPtr<class SGraphPin> FindInGraphPanel(const class SGraphPanel& InPanel) const;
+	UE_API TSharedPtr<class SGraphPin> FindInGraphPanel(const class SGraphPanel& InPanel) const;
 
 	/** */
 	bool IsValid() const
@@ -99,7 +101,7 @@ struct GRAPHEDITOR_API FGraphPinHandle
 /**
  * Represents a pin on a node in the GraphEditor
  */
-class GRAPHEDITOR_API SGraphPin : public SBorder
+class SGraphPin : public SBorder
 {
 public:
 
@@ -118,14 +120,14 @@ public:
 	 *
 	 * @param	InPin	The pin to create a widget for
 	 */
-	void Construct(const FArguments& InArgs, UEdGraphPin* InPin);
+	UE_API void Construct(const FArguments& InArgs, UEdGraphPin* InPin);
 
 public:
-	SGraphPin();
-	virtual ~SGraphPin();
+	UE_API SGraphPin();
+	UE_API virtual ~SGraphPin();
 
 	/** Set attribute for determining if pin is editable */
-	void SetIsEditable(TAttribute<bool> InIsEditable);
+	UE_API void SetIsEditable(TAttribute<bool> InIsEditable);
 
 	/** Retrieves the full horizontal box that contains the pin's row content */
 	TWeakPtr<SHorizontalBox> GetFullPinHorizontalRowWidget()
@@ -134,59 +136,59 @@ public:
 	}
 
 	/** Handle clicking on the pin */
-	virtual FReply OnPinMouseDown(const FGeometry& SenderGeometry, const FPointerEvent& MouseEvent);
+	UE_API virtual FReply OnPinMouseDown(const FGeometry& SenderGeometry, const FPointerEvent& MouseEvent);
 
 	/** Handle clicking on the pin name */
-	FReply OnPinNameMouseDown(const FGeometry& SenderGeometry, const FPointerEvent& MouseEvent);
+	UE_API FReply OnPinNameMouseDown(const FGeometry& SenderGeometry, const FPointerEvent& MouseEvent);
 
 	// SWidget interface
-	virtual FReply OnMouseButtonUp(const FGeometry& MyGeometry, const FPointerEvent& MouseEvent) override;
-	virtual FReply OnMouseMove(const FGeometry& MyGeometry, const FPointerEvent& MouseEvent) override;
-	virtual void OnMouseEnter(const FGeometry& MyGeometry, const FPointerEvent& MouseEvent) override;
-	virtual void OnMouseLeave(const FPointerEvent& MouseEvent) override;
-	virtual void OnDragEnter(const FGeometry& MyGeometry, const FDragDropEvent& DragDropEvent) override;
-	virtual void OnDragLeave(const FDragDropEvent& DragDropEvent) override;
-	virtual FReply OnDragOver(const FGeometry& MyGeometry, const FDragDropEvent& DragDropEvent) override;
-	virtual FReply OnDrop(const FGeometry& MyGeometry, const FDragDropEvent& DragDropEvent) override;
-	virtual void Tick( const FGeometry& AllottedGeometry, const double InCurrentTime, const float InDeltaTime ) override;
-	virtual TSharedPtr<IToolTip> GetToolTip() override;
+	UE_API virtual FReply OnMouseButtonUp(const FGeometry& MyGeometry, const FPointerEvent& MouseEvent) override;
+	UE_API virtual FReply OnMouseMove(const FGeometry& MyGeometry, const FPointerEvent& MouseEvent) override;
+	UE_API virtual void OnMouseEnter(const FGeometry& MyGeometry, const FPointerEvent& MouseEvent) override;
+	UE_API virtual void OnMouseLeave(const FPointerEvent& MouseEvent) override;
+	UE_API virtual void OnDragEnter(const FGeometry& MyGeometry, const FDragDropEvent& DragDropEvent) override;
+	UE_API virtual void OnDragLeave(const FDragDropEvent& DragDropEvent) override;
+	UE_API virtual FReply OnDragOver(const FGeometry& MyGeometry, const FDragDropEvent& DragDropEvent) override;
+	UE_API virtual FReply OnDrop(const FGeometry& MyGeometry, const FDragDropEvent& DragDropEvent) override;
+	UE_API virtual void Tick( const FGeometry& AllottedGeometry, const double InCurrentTime, const float InDeltaTime ) override;
+	UE_API virtual TSharedPtr<IToolTip> GetToolTip() override;
 	// End of SWidget interface
 
 	/** Refresh the level of detail applied to the pin name */
-	void RefreshLOD();
+	UE_API void RefreshLOD();
 
 public:
-	UEdGraphPin* GetPinObj() const;
+	UE_API UEdGraphPin* GetPinObj() const;
 
 	/** @param OwnerNode  The SGraphNode that this pin belongs to */
-	void SetOwner( const TSharedRef<SGraphNode> OwnerNode );
+	UE_API void SetOwner( const TSharedRef<SGraphNode> OwnerNode );
 
 	/** @param PinObj The UEdGraphPin object this pin widget represents */
-	void SetPinObj(UEdGraphPin* PinObj);
+	UE_API void SetPinObj(UEdGraphPin* PinObj);
 
 	/** @return whether this pin is incoming or outgoing */
-	EEdGraphPinDirection GetDirection() const;
+	UE_API EEdGraphPinDirection GetDirection() const;
 
 	/** @return whether this pin is an array value */
-	bool IsArray() const;
+	UE_API bool IsArray() const;
 
 	/** @return whether this pin is a set value */
-	bool IsSet() const;
+	UE_API bool IsSet() const;
 
 	/** @return whether this pin is a map value */
-	bool IsMap() const;
+	UE_API bool IsMap() const;
 
 	/** @return whether this pin is passed by mutable ref */
-	bool IsByMutableRef() const;
+	UE_API bool IsByMutableRef() const;
 
 	/** @return whether this pin is passed by mutable ref */
-	bool IsDelegate() const;
+	UE_API bool IsDelegate() const;
 
 	/** @return whether this pin is connected to another pin */
-	bool IsConnected() const;
+	UE_API bool IsConnected() const;
 
 	/** @return whether to fade out the pin's connections */
-	bool AreConnectionsFaded() const;
+	UE_API bool AreConnectionsFaded() const;
 
 	/** Tries to handle making a connection to another pin, depending on the schema and the pins it may do:  
 		 - Nothing
@@ -194,13 +196,13 @@ public:
 		 - Just make the new one
 		 @return Whether a connection was actually made or not
      */
-	virtual bool TryHandlePinConnection(SGraphPin& OtherSPin);
+	UE_API virtual bool TryHandlePinConnection(SGraphPin& OtherSPin);
 
 	/** @return Visibility based on whether or not we should show the inline editable text field for this pin */
-	virtual EVisibility GetDefaultValueVisibility() const;
+	UE_API virtual EVisibility GetDefaultValueVisibility() const;
 
 	/** Control whether we draw the label on this pin */
-	void SetShowLabel(bool bNewDrawLabel);
+	UE_API void SetShowLabel(bool bNewDrawLabel);
 
 	/** Allows the connection drawing policy to control the pin color */
 	void SetPinColorModifier(FLinearColor InColor)
@@ -226,19 +228,19 @@ public:
 	}
 
 	/** Set this pin to only be used to display default value */
-	void SetOnlyShowDefaultValue(bool bNewOnlyShowDefaultValue);
+	UE_API void SetOnlyShowDefaultValue(bool bNewOnlyShowDefaultValue);
 
 	/** If pin in node is visible at all */
-	EVisibility IsPinVisibleAsAdvanced() const;
+	UE_API EVisibility IsPinVisibleAsAdvanced() const;
 
 	/** Gets Node Offset */
-	UE::Slate::FDeprecateVector2DResult GetNodeOffset() const;
+	UE_API UE::Slate::FDeprecateVector2DResult GetNodeOffset() const;
 
 	/** Returns whether or not this pin is currently connectable */
-	bool GetIsConnectable() const;
+	UE_API bool GetIsConnectable() const;
 
 	/** Build the widget we should put into the 'default value' space, shown when nothing connected */
-	virtual TSharedRef<SWidget>	GetDefaultValueWidget();
+	UE_API virtual TSharedRef<SWidget>	GetDefaultValueWidget();
 
 	/** Get the widget created by GetDefaultValueWidget() */
 	TSharedRef<SWidget>	GetValueWidget() const { return ValueWidget.ToSharedRef(); }
@@ -247,7 +249,7 @@ public:
 	TSharedRef<SWrapBox> GetLabelAndValue() const { return LabelAndValue.ToSharedRef(); }
 
 	/** True if pin can be edited */
-	bool IsEditingEnabled() const;
+	UE_API bool IsEditingEnabled() const;
 
 	/** True if the pin's default value can be edited, false if it is read only */
 	bool GetDefaultValueIsEditable() const { return !GraphPinObj->bDefaultValueIsReadOnly; }
@@ -262,10 +264,10 @@ public:
 	void InvalidateGraphData() { bGraphDataInvalid = true; }
 
 	/** Override the visual look of the pin by providing two custom brushes */
-	void SetCustomPinIcon(const FSlateBrush* InConnectedBrush, const FSlateBrush* InDisconnectedBrush);
+	UE_API void SetCustomPinIcon(const FSlateBrush* InConnectedBrush, const FSlateBrush* InDisconnectedBrush);
 
 	/** @returns true if we have a valid PinValueInspector tooltip */
-	bool HasInteractiveTooltip() const;
+	UE_API bool HasInteractiveTooltip() const;
 
 	/** Enables or disables drag and drop on the pin */
 	void EnableDragAndDrop(bool bEnable) { bDragAndDropEnabled = bEnable; }
@@ -275,7 +277,7 @@ protected:
 	/** If true the graph pin subclass is responsible for setting the IsEnabled delegates for the aspects it cares about. If false, the default value widget enabling is done by the base class */
 	virtual bool DoesWidgetHandleSettingEditingEnabled() const { return false; }
 
-	FText GetPinLabel() const;
+	UE_API FText GetPinLabel() const;
 
 	/** If we should show the label on this pin */
 	EVisibility GetPinLabelVisibility() const
@@ -284,65 +286,65 @@ protected:
 	}
 
 	/** Get the widget we should put in the label space, which displays the name of the pin.*/
-	virtual TSharedRef<SWidget> GetLabelWidget(const FName& InPinLabelStyle);
+	UE_API virtual TSharedRef<SWidget> GetLabelWidget(const FName& InPinLabelStyle);
 
 	/** @return The brush with which to paint this graph pin's incoming/outgoing bullet point */
-	virtual const FSlateBrush* GetPinIcon() const;
+	UE_API virtual const FSlateBrush* GetPinIcon() const;
 
 	/** @return the brush which is to paint the 'secondary icon' for this pin, (e.g. value type for Map pins */
-	const FSlateBrush* GetSecondaryPinIcon() const;
+	UE_API const FSlateBrush* GetSecondaryPinIcon() const;
 
-	const FSlateBrush* GetPinBorder() const;
+	UE_API const FSlateBrush* GetPinBorder() const;
 	
 	/** @return The status brush to show after/before the name (on an input/output) */
-	virtual const FSlateBrush* GetPinStatusIcon() const;
+	UE_API virtual const FSlateBrush* GetPinStatusIcon() const;
 
 	/** @return Should we show a status brush after/before the name (on an input/output) */
-	virtual EVisibility GetPinStatusIconVisibility() const;
+	UE_API virtual EVisibility GetPinStatusIconVisibility() const;
 
 	/** Toggle the watch pinning */
-	virtual FReply ClickedOnPinStatusIcon();
+	UE_API virtual FReply ClickedOnPinStatusIcon();
 
 	/** @return The color that we should use to draw this pin */
-	virtual FSlateColor GetPinColor() const;
+	UE_API virtual FSlateColor GetPinColor() const;
 
 	/** @return The color that we should use to draw the highlight for this pin */
-	virtual FSlateColor GetHighlightColor() const;
+	UE_API virtual FSlateColor GetHighlightColor() const;
 
-	virtual FSlateColor GetPinDiffColor() const;
+	UE_API virtual FSlateColor GetPinDiffColor() const;
 
 	/** @return The secondary color that we should use to draw this pin (e.g. value color for Map pins) */
-	FSlateColor GetSecondaryPinColor() const;
+	UE_API FSlateColor GetSecondaryPinColor() const;
 
 	/** @return The color that we should use to draw this pin's text */
-	virtual FSlateColor GetPinTextColor() const;
+	UE_API virtual FSlateColor GetPinTextColor() const;
 
 	/** @return The tooltip to display for this pin */
-	FText GetTooltipText() const;
+	UE_API FText GetTooltipText() const;
 
 	/** Gets the window location (in screen coords) for an interactive tooltip (e.g. pin value inspector) */
 	UE_DEPRECATED(5.6, "Use the version of the function accepting FVector2f; this Slate API no longer interfaces directly with double-precision scalars and vectors.")
-	void GetInteractiveTooltipLocation(FVector2D& InOutDesiredLocation) const;
-	void GetInteractiveTooltipLocation(FVector2f& InOutDesiredLocation) const;
+	UE_API void GetInteractiveTooltipLocation(FVector2D& InOutDesiredLocation) const;
+	UE_API void GetInteractiveTooltipLocation(FVector2f& InOutDesiredLocation) const;
 
-	TOptional<EMouseCursor::Type> GetPinCursor() const;
+	UE_API TOptional<EMouseCursor::Type> GetPinCursor() const;
 
 	/** Spawns a FDragConnection or similar class for the pin drag event */
-	virtual TSharedRef<FDragDropOperation> SpawnPinDragEvent(const TSharedRef<class SGraphPanel>& InGraphPanel, const TArray< TSharedRef<SGraphPin> >& InStartingPins);
+	UE_API virtual TSharedRef<FDragDropOperation> SpawnPinDragEvent(const TSharedRef<class SGraphPanel>& InGraphPanel, const TArray< TSharedRef<SGraphPin> >& InStartingPins);
 
 	// Should we use low-detail pin names?
-	virtual bool UseLowDetailPinNames() const;
+	UE_API virtual bool UseLowDetailPinNames() const;
 
 	/** Determines the pin's visibility based on the LOD factor, when it is low LOD, no hit test will occur */
-	EVisibility GetPinVisiblity() const;
+	UE_API EVisibility GetPinVisiblity() const;
 
 public:
 
 	/** Returns the current pin image widget that is being used */
-	TSharedPtr<SWidget> GetPinImageWidget() const;
+	UE_API TSharedPtr<SWidget> GetPinImageWidget() const;
 
 	/** Sets the pin widget for this pin */
-	void SetPinImageWidget(TSharedRef<SWidget> NewWidget);
+	UE_API void SetPinImageWidget(TSharedRef<SWidget> NewWidget);
 
 protected:
 	/** The GraphNode that owns this pin */
@@ -432,3 +434,5 @@ private:
 	TSharedPtr<SLevelOfDetailBranchNode> PinNameLODBranchNode;
 	friend class SGraphPanel;
 };
+
+#undef UE_API
diff --git a/Engine/Source/Editor/GraphEditor/Public/SGraphPinComboBox.h b/Engine/Source/Editor/GraphEditor/Public/SGraphPinComboBox.h
--- a/Engine/Source/Editor/GraphEditor/Public/SGraphPinComboBox.h
+++ b/Engine/Source/Editor/GraphEditor/Public/SGraphPinComboBox.h
@@ -10,6 +10,8 @@
 #include "Widgets/Views/STableViewBase.h"
 #include "Widgets/Views/STableRow.h"
 
+#define UE_API GRAPHEDITOR_API
+
 class SComboButton;
 template <typename ItemType> class SListView;
 
@@ -17,7 +19,7 @@ template <typename ItemType> class SListView;
 DECLARE_DELEGATE_RetVal_OneParam(FText, FGetComboItemDisplayString, int32);
 
 //Class implementation to create combo box
-class GRAPHEDITOR_API SPinComboBox : public SCompoundWidget
+class SPinComboBox : public SCompoundWidget
 {
 public:
 	typedef SListView< TSharedPtr<int32> > SComboList;
@@ -32,7 +34,7 @@ public:
 	SLATE_END_ARGS()
 
 	//Construct combo box using combo button and combo list
-	void Construct( const FArguments& InArgs );
+	UE_API void Construct( const FArguments& InArgs );
 
 	//Function to return currently selected string
 	const TSharedPtr<int32> GetSelectedItem()
@@ -40,12 +42,12 @@ public:
 		return CurrentSelection.Pin();
 	}
 
-	void RemoveItemByIndex(int32 InIndexToRemove);
+	UE_API void RemoveItemByIndex(int32 InIndexToRemove);
 	
 private:
 
 	//Function to set the newly selected item
-	void OnSelectionChangedInternal( TSharedPtr<int32> NewSelection, ESelectInfo::Type SelectInfo );
+	UE_API void OnSelectionChangedInternal( TSharedPtr<int32> NewSelection, ESelectInfo::Type SelectInfo );
 
 	//Get string to display
 	FText OnGetVisibleTextInternal() const
@@ -54,10 +56,10 @@ private:
 	}
 
 	// Function to create each row of the combo widget
-	TSharedRef<ITableRow> OnGenerateComboWidget( TSharedPtr<int32> InComboIndex, const TSharedRef<STableViewBase>& OwnerTable );
+	UE_API TSharedRef<ITableRow> OnGenerateComboWidget( TSharedPtr<int32> InComboIndex, const TSharedRef<STableViewBase>& OwnerTable );
 	
 	/** Returns the target width of the enum selection menu. */
-	FOptionalSize GetMenuDesiredWidth() const;
+	UE_API FOptionalSize GetMenuDesiredWidth() const;
 	
 	FText GetRowString(int32 RowIndex) const
 	{
@@ -85,3 +87,5 @@ private:
 	/** Delegate for finding the tooltip */
 	FGetComboItemDisplayString OnGetTooltip;
 };
+
+#undef UE_API
diff --git a/Engine/Source/Editor/GraphEditor/Public/SGraphPinDataTableRowName.h b/Engine/Source/Editor/GraphEditor/Public/SGraphPinDataTableRowName.h
--- a/Engine/Source/Editor/GraphEditor/Public/SGraphPinDataTableRowName.h
+++ b/Engine/Source/Editor/GraphEditor/Public/SGraphPinDataTableRowName.h
@@ -9,27 +9,31 @@
 #include "UObject/SoftObjectPtr.h"
 #include "Widgets/DeclarativeSyntaxSupport.h"
 
+#define UE_API GRAPHEDITOR_API
+
 class UDataTable;
 class UEdGraphPin;
 
-class GRAPHEDITOR_API SGraphPinDataTableRowName : public SGraphPinNameList, public FDataTableEditorUtils::INotifyOnDataTableChanged
+class SGraphPinDataTableRowName : public SGraphPinNameList, public FDataTableEditorUtils::INotifyOnDataTableChanged
 {
 public:
 	SLATE_BEGIN_ARGS(SGraphPinDataTableRowName) {}
 	SLATE_END_ARGS()
 
-	void Construct(const FArguments& InArgs, UEdGraphPin* InGraphPinObj, class UDataTable* InDataTable);
+	UE_API void Construct(const FArguments& InArgs, UEdGraphPin* InGraphPinObj, class UDataTable* InDataTable);
 
-	SGraphPinDataTableRowName();
-	virtual ~SGraphPinDataTableRowName();
+	UE_API SGraphPinDataTableRowName();
+	UE_API virtual ~SGraphPinDataTableRowName();
 
 	// FDataTableEditorUtils::INotifyOnDataTableChanged
-	virtual void PreChange(const UDataTable* Changed, FDataTableEditorUtils::EDataTableChangeInfo Info) override;
-	virtual void PostChange(const UDataTable* Changed, FDataTableEditorUtils::EDataTableChangeInfo Info) override;
+	UE_API virtual void PreChange(const UDataTable* Changed, FDataTableEditorUtils::EDataTableChangeInfo Info) override;
+	UE_API virtual void PostChange(const UDataTable* Changed, FDataTableEditorUtils::EDataTableChangeInfo Info) override;
 
 protected:
 
-	void RefreshNameList();
+	UE_API void RefreshNameList();
 
 	TSoftObjectPtr<class UDataTable> DataTable;
 };
+
+#undef UE_API
diff --git a/Engine/Source/Editor/GraphEditor/Public/SGraphPinNameList.h b/Engine/Source/Editor/GraphEditor/Public/SGraphPinNameList.h
--- a/Engine/Source/Editor/GraphEditor/Public/SGraphPinNameList.h
+++ b/Engine/Source/Editor/GraphEditor/Public/SGraphPinNameList.h
@@ -11,17 +11,19 @@
 #include "Widgets/DeclarativeSyntaxSupport.h"
 #include "Widgets/SWidget.h"
 
+#define UE_API GRAPHEDITOR_API
+
 class FName;
 class SWidget;
 class UEdGraphPin;
 
-class GRAPHEDITOR_API SGraphPinNameList : public SGraphPin
+class SGraphPinNameList : public SGraphPin
 {
 public:
 	SLATE_BEGIN_ARGS(SGraphPinNameList) {}
 	SLATE_END_ARGS()
 
-	void Construct(const FArguments& InArgs, UEdGraphPin* InGraphPinObj, const TArray<TSharedPtr<FName>>& InNameList);
+	UE_API void Construct(const FArguments& InArgs, UEdGraphPin* InGraphPinObj, const TArray<TSharedPtr<FName>>& InNameList);
 
 protected:
 
@@ -30,7 +32,7 @@ protected:
 	 *
 	 *	@return Reference to the newly created widget object
 	 */
-	virtual TSharedRef<SWidget>	GetDefaultValueWidget() override;
+	UE_API virtual TSharedRef<SWidget>	GetDefaultValueWidget() override;
 
 	/**
 	 *	Function to set the newly selected index
@@ -38,10 +40,12 @@ protected:
 	 * @param NewSelection The newly selected item in the combo box
 	 * @param SelectInfo Provides context on how the selection changed
 	 */
-	void ComboBoxSelectionChanged(TSharedPtr<FName> StringItem, ESelectInfo::Type SelectInfo);
+	UE_API void ComboBoxSelectionChanged(TSharedPtr<FName> StringItem, ESelectInfo::Type SelectInfo);
 
 	TSharedPtr<class SNameComboBox>	ComboBox;
 
 	/** The actual list of FName values to choose from */
 	TArray<TSharedPtr<FName>> NameList;
 };
+
+#undef UE_API
diff --git a/Engine/Source/Editor/GraphEditor/Public/SGraphPreviewer.h b/Engine/Source/Editor/GraphEditor/Public/SGraphPreviewer.h
--- a/Engine/Source/Editor/GraphEditor/Public/SGraphPreviewer.h
+++ b/Engine/Source/Editor/GraphEditor/Public/SGraphPreviewer.h
@@ -10,13 +10,15 @@
 #include "Widgets/DeclarativeSyntaxSupport.h"
 #include "Widgets/SCompoundWidget.h"
 
+#define UE_API GRAPHEDITOR_API
+
 class FText;
 class SGraphPanel;
 class SWidget;
 class UEdGraph;
 
 // This widget provides a fully-zoomed-out preview of a specified graph
-class GRAPHEDITOR_API SGraphPreviewer : public SCompoundWidget
+class SGraphPreviewer : public SCompoundWidget
 {
 public:
 	SLATE_BEGIN_ARGS(SGraphPreviewer)
@@ -28,15 +30,15 @@ public:
 		SLATE_ARGUMENT( TSharedPtr<SWidget>, TitleBar )
 	SLATE_END_ARGS()
 
-	void Construct(const FArguments& InArgs, UEdGraph* InGraphObj);
+	UE_API void Construct(const FArguments& InArgs, UEdGraph* InGraphObj);
 
 protected:
 	/** Delegate handler for graph panel updates */
-	void OnUpdateGraphPanel();
+	UE_API void OnUpdateGraphPanel();
 
 private:
 	/** Helper function used to refresh the graph */
-	EActiveTimerReturnType RefreshGraphTimer(const double InCurrentTime, const float InDeltaTime);
+	UE_API EActiveTimerReturnType RefreshGraphTimer(const double InCurrentTime, const float InDeltaTime);
 
 protected:
 	/// The Graph we are currently viewing
@@ -48,3 +50,5 @@ protected:
 	// The underlying graph panel
 	TSharedPtr<SGraphPanel> GraphPanel;
 };
+
+#undef UE_API
diff --git a/Engine/Source/Editor/GraphEditor/Public/SNameComboBox.h b/Engine/Source/Editor/GraphEditor/Public/SNameComboBox.h
--- a/Engine/Source/Editor/GraphEditor/Public/SNameComboBox.h
+++ b/Engine/Source/Editor/GraphEditor/Public/SNameComboBox.h
@@ -31,6 +31,8 @@
 #include "Widgets/SCompoundWidget.h"
 #include "Widgets/SWidget.h"
 
+#define UE_API GRAPHEDITOR_API
+
 class SWidget;
 
 //#include "UnrealString.h"
@@ -38,7 +40,7 @@ class SWidget;
 /**
  * A combo box that shows FName content.
  */
-class GRAPHEDITOR_API SNameComboBox : public SCompoundWidget
+class SNameComboBox : public SCompoundWidget
 {
 public:
 
@@ -79,12 +81,12 @@ public:
 		SLATE_EVENT( FGetNameComboLabel, OnGetNameLabelForItem ) 
 	SLATE_END_ARGS()
 
-	void Construct( const FArguments& InArgs );
+	UE_API void Construct( const FArguments& InArgs );
 
 	/** Called to create a widget for each FName */
-	TSharedRef<SWidget> MakeItemWidget( TSharedPtr<FName> StringItem );
+	UE_API TSharedRef<SWidget> MakeItemWidget( TSharedPtr<FName> StringItem );
 
-	void SetSelectedItem (TSharedPtr<FName> NewSelection);
+	UE_API void SetSelectedItem (TSharedPtr<FName> NewSelection);
 
 	/** Returns the currently selected FName */
 	TSharedPtr<FName> GetSelectedItem()
@@ -93,21 +95,21 @@ public:
 	}
 
 	/** Request to reload the name options in the combobox from the OptionsSource attribute */
-	void RefreshOptions();
+	UE_API void RefreshOptions();
 
 	/** Clears the selected item in the name combo */
-	void ClearSelection();
+	UE_API void ClearSelection();
 
 private:
 	TSharedPtr<FName> OnGetSelection() const {return SelectedItem;}
 
 	/** Called when selection changes in the combo pop-up */
-	void OnSelectionChanged(TSharedPtr<FName> Selection, ESelectInfo::Type SelectInfo);
+	UE_API void OnSelectionChanged(TSharedPtr<FName> Selection, ESelectInfo::Type SelectInfo);
 
 	/** Helper method to get the text for a given item in the combo box */
-	FText GetSelectedNameLabel() const;
+	UE_API FText GetSelectedNameLabel() const;
 
-	FText GetItemNameLabel(TSharedPtr<FName> StringItem) const;
+	UE_API FText GetItemNameLabel(TSharedPtr<FName> StringItem) const;
 
 private:
 
@@ -129,3 +131,5 @@ private:
 	/** Sets the font used to draw the text */
 	TAttribute< FSlateFontInfo > Font;
 };
+
+#undef UE_API
diff --git a/Engine/Source/Editor/GraphEditor/Public/SNodePanel.h b/Engine/Source/Editor/GraphEditor/Public/SNodePanel.h
--- a/Engine/Source/Editor/GraphEditor/Public/SNodePanel.h
+++ b/Engine/Source/Editor/GraphEditor/Public/SNodePanel.h
@@ -49,6 +49,8 @@
 #include "Widgets/SWidget.h"
 #include "SNodePanel.generated.h"
 
+#define UE_API GRAPHEDITOR_API
+
 class FActiveTimerHandle;
 class FReferenceCollector;
 class FScopedTransaction;
@@ -231,7 +233,7 @@ struct FZoomLevelsContainer
 	virtual ~FZoomLevelsContainer( void ) {};
 };
 
-struct GRAPHEDITOR_API FGraphSelectionManager : public FGCObject
+struct FGraphSelectionManager : public FGCObject
 {
 	TSet<TObjectPtr<class UObject>> SelectedNodes;
 
@@ -239,13 +241,13 @@ struct GRAPHEDITOR_API FGraphSelectionManager : public FGCObject
 	SGraphEditor::FOnSelectionChanged OnSelectionChanged;
 public:
 	/** @return the set of selected nodes */
-	const FGraphPanelSelectionSet& GetSelectedNodes() const;
+	UE_API const FGraphPanelSelectionSet& GetSelectedNodes() const;
 
 	/** Select just the specified node */
-	void SelectSingleNode(SelectedItemType Node);
+	UE_API void SelectSingleNode(SelectedItemType Node);
 
 	/** Reset the selection state of all nodes */
-	void ClearSelectionSet();
+	UE_API void ClearSelectionSet();
 
 	/** Returns true if any nodes are selected */
 	bool AreAnyNodesSelected() const
@@ -254,7 +256,7 @@ public:
 	}
 
 	/** Changes the selection set to contain exactly all of the passed in nodes */
-	void SetSelectionSet(FGraphPanelSelectionSet& NewSet);
+	UE_API void SetSelectionSet(FGraphPanelSelectionSet& NewSet);
 
 	/**
 	 * Add or remove a node from the selection set
@@ -262,18 +264,18 @@ public:
 	 * @param Node      Node the affect.
 	 * @param bSelect   true to select the node; false to unselect.
 	 */
-	void SetNodeSelection(SelectedItemType Node, bool bSelect);
+	UE_API void SetNodeSelection(SelectedItemType Node, bool bSelect);
 
 	/** @return true if Node is selected; false otherwise */
-	bool IsNodeSelected(SelectedItemType Node) const;
+	UE_API bool IsNodeSelected(SelectedItemType Node) const;
 
 	// Handle the selection mechanics of starting to drag a node
-	void StartDraggingNode(SelectedItemType NodeBeingDragged, const FPointerEvent& MouseEvent);
+	UE_API void StartDraggingNode(SelectedItemType NodeBeingDragged, const FPointerEvent& MouseEvent);
 
 	// Handle the selection mechanics when a node is clicked on
-	void ClickedOnNode(SelectedItemType Node, const FPointerEvent& MouseEvent);
+	UE_API void ClickedOnNode(SelectedItemType Node, const FPointerEvent& MouseEvent);
 
-	void AddReferencedObjects(FReferenceCollector& Collector) override;
+	UE_API void AddReferencedObjects(FReferenceCollector& Collector) override;
 	virtual FString GetReferencerName() const override
 	{
 		return TEXT("FGraphSelectionManager");
@@ -309,7 +311,7 @@ namespace ENodeZone
 	};
 }
 
-class GRAPHEDITOR_API SNodePanel : public SPanel
+class SNodePanel : public SPanel
 {
 public:
 
@@ -318,7 +320,7 @@ public:
 	public:
 
 		/** A slot that support alignment of content and padding and z-order */
-		class GRAPHEDITOR_API FNodeSlot : public TSlotBase<FNodeSlot>, public TAlignmentWidgetSlotMixin<FNodeSlot>
+		class FNodeSlot : public TSlotBase<FNodeSlot>, public TAlignmentWidgetSlotMixin<FNodeSlot>
 		{
 		public:
 			friend SNode;
@@ -345,7 +347,7 @@ public:
 				SLATE_ATTRIBUTE(bool, AllowScaling)
 			SLATE_SLOT_END_ARGS()
 
-			void Construct(const FChildren& SlotOwner, FSlotArguments&& InArgs);
+			UE_API void Construct(const FChildren& SlotOwner, FSlotArguments&& InArgs);
 
 		public:
 			ENodeZone::Type GetZoneType() const
@@ -795,28 +797,28 @@ public:
 
 	};
 
-	SNodePanel();
+	UE_API SNodePanel();
 
 	// SPanel interface
-	virtual void OnArrangeChildren(const FGeometry& AllottedGeometry, FArrangedChildren& ArrangedChildren) const override;
-	virtual FVector2D ComputeDesiredSize(float) const override;
-	virtual FChildren* GetChildren() override;
+	UE_API virtual void OnArrangeChildren(const FGeometry& AllottedGeometry, FArrangedChildren& ArrangedChildren) const override;
+	UE_API virtual FVector2D ComputeDesiredSize(float) const override;
+	UE_API virtual FChildren* GetChildren() override;
 	// End of SPanel interface
 
 	// SWidget interface
-	virtual void Tick( const FGeometry& AllottedGeometry, const double InCurrentTime, const float InDeltaTime ) override;
-	virtual FReply OnMouseButtonDown(const FGeometry& MyGeometry, const FPointerEvent& MouseEvent) override;
-	virtual FReply OnMouseButtonUp(const FGeometry& MyGeometry, const FPointerEvent& MouseEvent) override;
-	virtual void OnMouseCaptureLost(const FCaptureLostEvent& CaptureLostEvent) override;
-	virtual FReply OnMouseMove(const FGeometry& MyGeometry, const FPointerEvent& MouseEvent) override;
-	virtual FReply OnMouseWheel(const FGeometry& MyGeometry, const FPointerEvent& MouseEvent) override;
-	virtual FCursorReply OnCursorQuery( const FGeometry& MyGeometry, const FPointerEvent& CursorEvent ) const override;
-	virtual FReply OnKeyDown( const FGeometry& MyGeometry, const FKeyEvent& InKeyEvent ) override;
-	virtual FReply OnKeyUp( const FGeometry& MyGeometry, const FKeyEvent& InKeyEvent ) override;
-	virtual void OnFocusLost( const FFocusEvent& InFocusEvent ) override;
-	virtual FReply OnTouchGesture( const FGeometry& MyGeometry, const FPointerEvent& GestureEvent ) override;
-	virtual FReply OnTouchEnded( const FGeometry& MyGeometry, const FPointerEvent& InTouchEvent ) override;
-	virtual float GetRelativeLayoutScale(int32 ChildIndex, float LayoutScaleMultiplier) const override;
+	UE_API virtual void Tick( const FGeometry& AllottedGeometry, const double InCurrentTime, const float InDeltaTime ) override;
+	UE_API virtual FReply OnMouseButtonDown(const FGeometry& MyGeometry, const FPointerEvent& MouseEvent) override;
+	UE_API virtual FReply OnMouseButtonUp(const FGeometry& MyGeometry, const FPointerEvent& MouseEvent) override;
+	UE_API virtual void OnMouseCaptureLost(const FCaptureLostEvent& CaptureLostEvent) override;
+	UE_API virtual FReply OnMouseMove(const FGeometry& MyGeometry, const FPointerEvent& MouseEvent) override;
+	UE_API virtual FReply OnMouseWheel(const FGeometry& MyGeometry, const FPointerEvent& MouseEvent) override;
+	UE_API virtual FCursorReply OnCursorQuery( const FGeometry& MyGeometry, const FPointerEvent& CursorEvent ) const override;
+	UE_API virtual FReply OnKeyDown( const FGeometry& MyGeometry, const FKeyEvent& InKeyEvent ) override;
+	UE_API virtual FReply OnKeyUp( const FGeometry& MyGeometry, const FKeyEvent& InKeyEvent ) override;
+	UE_API virtual void OnFocusLost( const FFocusEvent& InFocusEvent ) override;
+	UE_API virtual FReply OnTouchGesture( const FGeometry& MyGeometry, const FPointerEvent& GestureEvent ) override;
+	UE_API virtual FReply OnTouchEnded( const FGeometry& MyGeometry, const FPointerEvent& InTouchEvent ) override;
+	UE_API virtual float GetRelativeLayoutScale(int32 ChildIndex, float LayoutScaleMultiplier) const override;
 	// End of SWidget interface
 public:
 	/**
@@ -826,16 +828,16 @@ public:
 	 *
 	 * @return True if the node is being observed by some widget in this panel; false otherwise.
 	 */
-	bool Contains(UObject* Node) const;
+	UE_API bool Contains(UObject* Node) const;
 
 	/** @retun the zoom amount; e.g. a value of 0.25f results in quarter-sized nodes */
-	float GetZoomAmount() const;
+	UE_API float GetZoomAmount() const;
 	/** @return Zoom level as a pretty string */
-	FText GetZoomText() const;
-	FSlateColor GetZoomTextColorAndOpacity() const;
+	UE_API FText GetZoomText() const;
+	UE_API FSlateColor GetZoomTextColorAndOpacity() const;
 
 	/** @return the view offset in graph space */
-	UE::Slate::FDeprecateVector2DResult GetViewOffset() const;
+	UE_API UE::Slate::FDeprecateVector2DResult GetViewOffset() const;
 
 	/** 
 	 * when a panel is scrolling/zooming to a target, this can be called to get it's destination
@@ -844,77 +846,77 @@ public:
 	 * @return true if there's a scrolling/zooming target and false if there is no destination
 	 */
 	UE_DEPRECATED(5.6, "Use the version of the function accepting FVector2f; this Slate API no longer interfaces directly with double-precision scalars and vectors.")
-	bool GetZoomTargetRect(FVector2D& TopLeft, FVector2D& BottomRight) const;
-	bool GetZoomTargetRect(FVector2f& TopLeft, FVector2f& BottomRight) const;
+	UE_API bool GetZoomTargetRect(FVector2D& TopLeft, FVector2D& BottomRight) const;
+	UE_API bool GetZoomTargetRect(FVector2f& TopLeft, FVector2f& BottomRight) const;
 
 	/** @return the current view bookmark ID */
 	const FGuid& GetViewBookmarkId() const { return CurrentBookmarkGuid; }
 
 	/** Given a coordinate in panel space (i.e. panel widget space), return the same coordinate in graph space while taking zoom and panning into account */
-	UE::Slate::FDeprecateVector2DResult PanelCoordToGraphCoord(const UE::Slate::FDeprecateVector2DParameter& PanelSpaceCoordinate) const;
+	UE_API UE::Slate::FDeprecateVector2DResult PanelCoordToGraphCoord(const UE::Slate::FDeprecateVector2DParameter& PanelSpaceCoordinate) const;
 
 	/** Restore the graph panel to the supplied view offset/zoom. Also, optionally set the current bookmark ID. */
-	void RestoreViewSettings(const UE::Slate::FDeprecateVector2DParameter& InViewOffset, float InZoomAmount, const FGuid& InBookmarkGuid = FGuid());
+	UE_API void RestoreViewSettings(const UE::Slate::FDeprecateVector2DParameter& InViewOffset, float InZoomAmount, const FGuid& InBookmarkGuid = FGuid());
 
 	/** Get the grid snap size */
-	static uint32 GetSnapGridSize();
+	static UE_API uint32 GetSnapGridSize();
 
 	/** 
 	 * Zooms out to fit either all nodes or only the selected ones.
 	 * @param bOnlySelection Whether to zoom to fit around only the current selection (if false, will zoom to the extents of all nodes)
 	 */
-	void ZoomToFit(bool bOnlySelection);
+	UE_API void ZoomToFit(bool bOnlySelection);
 
 	/** Get the bounding area for the currently selected nodes 
 		@return false if nothing is selected					*/
-	bool GetBoundsForSelectedNodes(/*out*/ class FSlateRect& Rect, float Padding = 0.0f);
+	UE_API bool GetBoundsForSelectedNodes(/*out*/ class FSlateRect& Rect, float Padding = 0.0f);
 
 	/** @return the position where where nodes should be pasted (i.e. from the clipboard) */
-	UE::Slate::FDeprecateVector2DResult GetPastePosition() const;
+	UE_API UE::Slate::FDeprecateVector2DResult GetPastePosition() const;
 
 	/** Ask panel to scroll to location */
-	void RequestDeferredPan(const UE::Slate::FDeprecateVector2DParameter& TargetPosition);
+	UE_API void RequestDeferredPan(const UE::Slate::FDeprecateVector2DParameter& TargetPosition);
 
 	/** If it is focusing on a particular object */
-	bool HasDeferredObjectFocus() const;
+	UE_API bool HasDeferredObjectFocus() const;
 
 	/** Query whether this graph is about to start panning/zooming towards a destination */
-	bool HasDeferredZoomDestination() const;
+	UE_API bool HasDeferredZoomDestination() const;
 
 	/** Commit transactions for any node movements */
-	void FinalizeNodeMovements();
+	UE_API void FinalizeNodeMovements();
 
 	/** Returns the current LOD level of this panel, based on the zoom factor */
 	EGraphRenderingLOD::Type GetCurrentLOD() const { return CurrentLOD; }
 
 	/** Returns if the panel has been panned or zoomed since the last update */
-	bool HasMoved() const;
+	UE_API bool HasMoved() const;
 
 	/** Returns all the panel children rather than only visible */
-	FChildren* GetAllChildren();
+	UE_API FChildren* GetAllChildren();
 
 protected:
 	/** Initialize members */
-	void Construct();
+	UE_API void Construct();
 
 	/**
 	 * Zooms to the specified target rect
 	 * @param TopLeft The top left corner of the target rect
 	 * @param BottomRight The bottom right corner of the target rect
 	 */
-	void ZoomToTarget(const FVector2f& TopLeft, const FVector2f& BottomRight);
+	UE_API void ZoomToTarget(const FVector2f& TopLeft, const FVector2f& BottomRight);
 
 	/** Update the new view offset location  */
-	void UpdateViewOffset(const FGeometry& MyGeometry, const FVector2f& TargetPosition);
+	UE_API void UpdateViewOffset(const FGeometry& MyGeometry, const FVector2f& TargetPosition);
 
 	/** Compute much panel needs to change to pan to location  */
-	static FVector2f ComputeEdgePanAmount(const FGeometry& MyGeometry, const FVector2f& MouseEvent);
+	static UE_API FVector2f ComputeEdgePanAmount(const FGeometry& MyGeometry, const FVector2f& MouseEvent);
 
 	/** Given a coordinate in graph space (e.g. a node's position), return the same coordinate in widget space while taking zoom and panning into account */
-	FVector2f GraphCoordToPanelCoord(const UE::Slate::FDeprecateVector2DParameter& GraphSpaceCoordinate) const;
+	UE_API FVector2f GraphCoordToPanelCoord(const UE::Slate::FDeprecateVector2DParameter& GraphSpaceCoordinate) const;
 
 	/** Given a rectangle in panel space, return a rectangle in graph space. */
-	FSlateRect PanelRectToGraphRect(const FSlateRect& PanelSpaceRect) const;
+	UE_API FSlateRect PanelRectToGraphRect(const FSlateRect& PanelSpaceRect) const;
 
 	/** 
 	 * Lets the CanvasPanel know that the user is interacting with a node.
@@ -922,18 +924,18 @@ protected:
 	 * @param InNodeToDrag   The node that the user wants to drag
 	 * @param GrabOffset     Where within the node the user grabbed relative to split between inputs and outputs.
 	 */
-	virtual void OnBeginNodeInteraction(const TSharedRef<SNode>& InNodeToDrag, const FVector2D& GrabOffset) UE_SLATE_DEPRECATED_VECTOR_VIRTUAL_FUNCTION;
-	virtual void OnBeginNodeInteraction(const TSharedRef<SNode>& InNodeToDrag, const FVector2f& GrabOffset);
+	UE_API virtual void OnBeginNodeInteraction(const TSharedRef<SNode>& InNodeToDrag, const FVector2D& GrabOffset) UE_SLATE_DEPRECATED_VECTOR_VIRTUAL_FUNCTION;
+	UE_API virtual void OnBeginNodeInteraction(const TSharedRef<SNode>& InNodeToDrag, const FVector2f& GrabOffset);
 
 	/** 
 	 * Lets the CanvasPanel know that the user has ended interacting with a node.
 	 * 
 	 * @param InNodeToDrag   The node that the user was to dragging
 	 */
-	virtual void OnEndNodeInteraction(const TSharedRef<SNode>& InNodeToDrag);
+	UE_API virtual void OnEndNodeInteraction(const TSharedRef<SNode>& InNodeToDrag);
 
 	/** Figure out which nodes intersect the marquee rectangle */
-	void FindNodesAffectedByMarquee(FGraphPanelSelectionSet& OutAffectedNodes) const;
+	UE_API void FindNodesAffectedByMarquee(FGraphPanelSelectionSet& OutAffectedNodes) const;
 
 	/**
 	 * Apply the marquee operation to the current selection
@@ -942,7 +944,7 @@ protected:
 	 * @param CurrentSelection     The selection before the marquee operation.
 	 * @param OutNewSelection      The selection resulting from Marquee being applied to CurrentSelection.
 	 */
-	static void ApplyMarqueeSelection(const FMarqueeOperation& InMarquee, const FGraphPanelSelectionSet& CurrentSelection, TSet<TObjectPtr<UObject>>& OutNewSelection);
+	static UE_API void ApplyMarqueeSelection(const FMarqueeOperation& InMarquee, const FGraphPanelSelectionSet& CurrentSelection, TSet<TObjectPtr<UObject>>& OutNewSelection);
 
 	/**
 	 * On the next tick, centers and selects the widget associated with the object if it exists
@@ -950,44 +952,44 @@ protected:
 	 * @param ObjectToSelect	The object to select, and potentially center on
 	 * @param bCenter			Whether or not to center the graph node
 	 */
-	void SelectAndCenterObject(const UObject* ObjectToSelect, bool bCenter);
+	UE_API void SelectAndCenterObject(const UObject* ObjectToSelect, bool bCenter);
 
 	/**
 	 * On the next tick, centers the widget associated with the object if it exists
 	 *
 	 * @param ObjectToCenter	The object to center
 	 */
-	void CenterObject(const UObject* ObjectToCenter);
+	UE_API void CenterObject(const UObject* ObjectToCenter);
 
 	/** Add a slot to the CanvasPanel dynamically */
-	virtual void AddGraphNode(const TSharedRef<SNode>& NodeToAdd);
+	UE_API virtual void AddGraphNode(const TSharedRef<SNode>& NodeToAdd);
 
 	/** Remove all nodes from the panel */
-	virtual void RemoveAllNodes();
+	UE_API virtual void RemoveAllNodes();
 	
 	/** Populate visibile children array */
-	virtual void PopulateVisibleChildren(const FGeometry& AllottedGeometry);
+	UE_API virtual void PopulateVisibleChildren(const FGeometry& AllottedGeometry);
 
 	/** Arrange child nodes - allows derived classes to supply non-node children in OnArrangeChildren */
-	virtual void ArrangeChildNodes(const FGeometry& AllottedGeometry, FArrangedChildren& ArrangedChildren) const;
+	UE_API virtual void ArrangeChildNodes(const FGeometry& AllottedGeometry, FArrangedChildren& ArrangedChildren) const;
 
 	// Paint the background as lines
-	void PaintBackgroundAsLines(const FSlateBrush* BackgroundImage, const FGeometry& AllottedGeometry, const FSlateRect& MyCullingRect, FSlateWindowElementList& OutDrawElements, int32& DrawLayerId) const;
+	UE_API void PaintBackgroundAsLines(const FSlateBrush* BackgroundImage, const FGeometry& AllottedGeometry, const FSlateRect& MyCullingRect, FSlateWindowElementList& OutDrawElements, int32& DrawLayerId) const;
 
 	// Paint the well shadow (around the perimeter)
-	void PaintSurroundSunkenShadow(const FSlateBrush* ShadowImage, const FGeometry& AllottedGeometry, const FSlateRect& MyCullingRect, FSlateWindowElementList& OutDrawElements, int32 DrawLayerId) const;
+	UE_API void PaintSurroundSunkenShadow(const FSlateBrush* ShadowImage, const FGeometry& AllottedGeometry, const FSlateRect& MyCullingRect, FSlateWindowElementList& OutDrawElements, int32 DrawLayerId) const;
 
 	// Paint the marquee selection rectangle
-	void PaintMarquee(const FGeometry& AllottedGeometry, const FSlateRect& MyCullingRect, FSlateWindowElementList& OutDrawElements, int32 DrawLayerId) const;
+	UE_API void PaintMarquee(const FGeometry& AllottedGeometry, const FSlateRect& MyCullingRect, FSlateWindowElementList& OutDrawElements, int32 DrawLayerId) const;
 
 	// Paint the software mouse if necessary
-	void PaintSoftwareCursor(const FGeometry& AllottedGeometry, const FSlateRect& MyCullingRect, FSlateWindowElementList& OutDrawElements, int32 DrawLayerId) const;
+	UE_API void PaintSoftwareCursor(const FGeometry& AllottedGeometry, const FSlateRect& MyCullingRect, FSlateWindowElementList& OutDrawElements, int32 DrawLayerId) const;
 
 	// Paint a comment bubble
-	void PaintComment(const FString& CommentText, const FGeometry& AllottedGeometry, const FSlateRect& MyCullingRect, FSlateWindowElementList& OutDrawElements, int32 DrawLayerId, const FLinearColor& CommentTinting, float& HeightAboveNode, const FWidgetStyle& InWidgetStyle ) const;
+	UE_API void PaintComment(const FString& CommentText, const FGeometry& AllottedGeometry, const FSlateRect& MyCullingRect, FSlateWindowElementList& OutDrawElements, int32 DrawLayerId, const FLinearColor& CommentTinting, float& HeightAboveNode, const FWidgetStyle& InWidgetStyle ) const;
 
 	/** Determines if a specified node is not visually relevant. */
-	bool IsNodeCulled(const TSharedRef<SNode>& Node, const FGeometry& AllottedGeometry) const;
+	UE_API bool IsNodeCulled(const TSharedRef<SNode>& Node, const FGeometry& AllottedGeometry) const;
 
 protected:
 	///////////
@@ -1002,8 +1004,8 @@ protected:
 	 * @return True if successful
 	 */
 	UE_DEPRECATED(5.6, "Use the version of the function accepting FVector2f; this Slate API no longer interfaces directly with double-precision scalars and vectors.")
-	bool GetBoundsForNode(const UObject* InNode, /*out*/ FVector2D& MinCorner, /*out*/ FVector2D& MaxCorner, float Padding = 0.0f) const;
-	bool GetBoundsForNode(const UObject* InNode, /*out*/ FVector2f& MinCorner, /*out*/ FVector2f& MaxCorner, float Padding = 0.0f) const;
+	UE_API bool GetBoundsForNode(const UObject* InNode, /*out*/ FVector2D& MinCorner, /*out*/ FVector2D& MaxCorner, float Padding = 0.0f) const;
+	UE_API bool GetBoundsForNode(const UObject* InNode, /*out*/ FVector2f& MinCorner, /*out*/ FVector2f& MaxCorner, float Padding = 0.0f) const;
 
 	/**
 	 * Get the bounds of the selected nodes 
@@ -1011,34 +1013,34 @@ protected:
 	 * @return True if successful
 	 */
 	UE_DEPRECATED(5.6, "Use the version of the function accepting FVector2f; this Slate API no longer interfaces directly with double-precision scalars and vectors.")
-	bool GetBoundsForNodes(bool bSelectionSetOnly, /*out*/ FVector2D& MinCorner, /*out*/ FVector2D& MaxCorner, float Padding = 0.0f) const;
-	bool GetBoundsForNodes(bool bSelectionSetOnly, /*out*/ FVector2f& MinCorner, /*out*/ FVector2f& MaxCorner, float Padding = 0.0f) const;
+	UE_API bool GetBoundsForNodes(bool bSelectionSetOnly, /*out*/ FVector2D& MinCorner, /*out*/ FVector2D& MaxCorner, float Padding = 0.0f) const;
+	UE_API bool GetBoundsForNodes(bool bSelectionSetOnly, /*out*/ FVector2f& MinCorner, /*out*/ FVector2f& MaxCorner, float Padding = 0.0f) const;
 
 	/**
 	 * Scroll the view to the desired location 
 	 * @return true when the desired location is reached
 	 */
-	bool ScrollToLocation(const FGeometry& MyGeometry, FVector2f DesiredCenterPosition, const float InDeltaTime);
+	UE_API bool ScrollToLocation(const FGeometry& MyGeometry, FVector2f DesiredCenterPosition, const float InDeltaTime);
 
 	/**
 	 * Zoom to fit the desired size 
 	 * @return true when zoom fade has completed & fits the desired size
 	 */
-	bool ZoomToLocation(const FVector2f& CurrentSizeWithoutZoom, const FVector2f& DesiredSize, bool bDoneScrolling);
+	UE_API bool ZoomToLocation(const FVector2f& CurrentSizeWithoutZoom, const FVector2f& DesiredSize, bool bDoneScrolling);
 
 	/**
 	 * Change zoom level by the specified zoom level delta, about the specified origin.
 	 */
-	void ChangeZoomLevel(int32 ZoomLevelDelta, const FVector2f& WidgetSpaceZoomOrigin, bool bOverrideZoomLimiting);
+	UE_API void ChangeZoomLevel(int32 ZoomLevelDelta, const FVector2f& WidgetSpaceZoomOrigin, bool bOverrideZoomLimiting);
 
 	// Should be called whenever the zoom level has changed
-	void PostChangedZoom();
+	UE_API void PostChangedZoom();
 
 	// Fires up a per-tick function to zoom the graph to fit
-	void RequestZoomToFit();
+	UE_API void RequestZoomToFit();
 
 	// Cancels any active zoom-to-fit action
-	void CancelZoomToFit();
+	UE_API void CancelZoomToFit();
 
 public:
 	
@@ -1187,7 +1189,7 @@ protected:
 
 private:
 	/** Active timer that handles deferred zooming until the target zoom is reached */
-	EActiveTimerReturnType HandleZoomToFit(double InCurrentTime, float InDeltaTime);
+	UE_API EActiveTimerReturnType HandleZoomToFit(double InCurrentTime, float InDeltaTime);
 
 private:
 	/** The handle to the active timer */
@@ -1200,3 +1202,5 @@ private:
 	/** Current view bookmark info */
 	FGuid CurrentBookmarkGuid;
 };
+
+#undef UE_API
diff --git a/Engine/Source/Editor/HardwareTargeting/Public/HardwareTargetingSettings.h b/Engine/Source/Editor/HardwareTargeting/Public/HardwareTargetingSettings.h
--- a/Engine/Source/Editor/HardwareTargeting/Public/HardwareTargetingSettings.h
+++ b/Engine/Source/Editor/HardwareTargeting/Public/HardwareTargetingSettings.h
@@ -11,6 +11,8 @@
 
 #include "HardwareTargetingSettings.generated.h"
 
+#define UE_API HARDWARETARGETING_API
+
 /** Enum specifying a class of hardware */
 UENUM()
 enum class EHardwareClass : uint8
@@ -40,8 +42,8 @@ enum class EGraphicsPreset : uint8
 };
 
 /** Hardware targeting settings, stored in default config, per-project */
-UCLASS(config=Engine, defaultconfig)
-class HARDWARETARGETING_API UHardwareTargetingSettings : public UObject
+UCLASS(MinimalAPI, config=Engine, defaultconfig)
+class UHardwareTargetingSettings : public UObject
 {
 	GENERATED_UCLASS_BODY()
 
@@ -62,7 +64,7 @@ class HARDWARETARGETING_API UHardwareTargetingSettings : public UObject
 	EGraphicsPreset AppliedDefaultGraphicsPerformance;
 
 	/** Check if these settings have any pending changes that require action */
-	bool HasPendingChanges() const;
+	UE_API bool HasPendingChanges() const;
 
 #if WITH_EDITOR
 public:
@@ -72,10 +74,12 @@ public:
 
 protected:
 	/** Called when a property on this object is changed */
-	virtual void PostEditChangeProperty( struct FPropertyChangedEvent& PropertyChangedEvent ) override;
+	UE_API virtual void PostEditChangeProperty( struct FPropertyChangedEvent& PropertyChangedEvent ) override;
 
 private:
 	/** Holds an event delegate that is executed when a setting has changed. */
 	FSettingChangedEvent SettingChangedEvent;
 #endif
 };
+
+#undef UE_API
diff --git a/Engine/Source/Editor/InternationalizationSettings/Classes/InternationalizationSettingsModel.h b/Engine/Source/Editor/InternationalizationSettings/Classes/InternationalizationSettingsModel.h
--- a/Engine/Source/Editor/InternationalizationSettings/Classes/InternationalizationSettingsModel.h
+++ b/Engine/Source/Editor/InternationalizationSettings/Classes/InternationalizationSettingsModel.h
@@ -10,6 +10,8 @@
 
 #include "InternationalizationSettingsModel.generated.h"
 
+#define UE_API INTERNATIONALIZATIONSETTINGS_API
+
 UENUM()
 enum class ETimezoneSetting : uint8
 {
@@ -119,28 +121,28 @@ enum class ETimezoneSetting : uint8
 /**
  * Implements loading and saving of internationalization settings.
  */
-UCLASS(config=EditorSettings)
-class INTERNATIONALIZATIONSETTINGS_API UInternationalizationSettingsModel
+UCLASS(MinimalAPI, config=EditorSettings)
+class UInternationalizationSettingsModel
 	: public UObject
 {
 	GENERATED_UCLASS_BODY()
 
 public:
-	void ResetToDefault();
-	bool GetEditorLanguage(FString& OutEditorLanguage) const;
-	void SetEditorLanguage(const FString& InEditorLanguage);
-	bool GetEditorLocale(FString& OutEditorLocale) const;
-	void SetEditorLocale(const FString& InEditorLocale);
-	bool GetPreviewGameLanguage(FString& OutPreviewGameLanguage) const;
-	void SetPreviewGameLanguage(const FString& InPreviewGameLanguage);
-	bool ShouldUseLocalizedNumericInput() const;
-	void SetShouldUseLocalizedNumericInput(const bool Value);
-	bool ShouldUseLocalizedPropertyNames() const;
-	void SetShouldUseLocalizedPropertyNames(const bool Value);
-	bool ShouldUseLocalizedNodeAndPinNames() const;
-	void SetShouldUseLocalizedNodeAndPinNames(const bool Value);
+	UE_API void ResetToDefault();
+	UE_API bool GetEditorLanguage(FString& OutEditorLanguage) const;
+	UE_API void SetEditorLanguage(const FString& InEditorLanguage);
+	UE_API bool GetEditorLocale(FString& OutEditorLocale) const;
+	UE_API void SetEditorLocale(const FString& InEditorLocale);
+	UE_API bool GetPreviewGameLanguage(FString& OutPreviewGameLanguage) const;
+	UE_API void SetPreviewGameLanguage(const FString& InPreviewGameLanguage);
+	UE_API bool ShouldUseLocalizedNumericInput() const;
+	UE_API void SetShouldUseLocalizedNumericInput(const bool Value);
+	UE_API bool ShouldUseLocalizedPropertyNames() const;
+	UE_API void SetShouldUseLocalizedPropertyNames(const bool Value);
+	UE_API bool ShouldUseLocalizedNodeAndPinNames() const;
+	UE_API void SetShouldUseLocalizedNodeAndPinNames(const bool Value);
 
-	int32 GetTimezoneValue() const;
+	UE_API int32 GetTimezoneValue() const;
 
 public:
 	/**
@@ -150,3 +152,5 @@ public:
 	ETimezoneSetting DisplayTimezone;
 
 };
+
+#undef UE_API
diff --git a/Engine/Source/Editor/InternationalizationSettings/Public/InternationalizationSettingsModelDetails.h b/Engine/Source/Editor/InternationalizationSettings/Public/InternationalizationSettingsModelDetails.h
--- a/Engine/Source/Editor/InternationalizationSettings/Public/InternationalizationSettingsModelDetails.h
+++ b/Engine/Source/Editor/InternationalizationSettings/Public/InternationalizationSettingsModelDetails.h
@@ -6,14 +6,18 @@
 #include "IDetailCustomization.h"
 #include "Templates/SharedPointer.h"
 
+#define UE_API INTERNATIONALIZATIONSETTINGS_API
+
 class IDetailLayoutBuilder;
 
-class INTERNATIONALIZATIONSETTINGS_API FInternationalizationSettingsModelDetails : public IDetailCustomization
+class FInternationalizationSettingsModelDetails : public IDetailCustomization
 {
 public:
 	/** Makes a new instance of this detail layout class for a specific detail view requesting it */
-	static TSharedRef<IDetailCustomization> MakeInstance();
+	static UE_API TSharedRef<IDetailCustomization> MakeInstance();
 
 	/** IDetailCustomization interface */
-	virtual void CustomizeDetails( IDetailLayoutBuilder& DetailLayout ) override;
+	UE_API virtual void CustomizeDetails( IDetailLayoutBuilder& DetailLayout ) override;
 };
+
+#undef UE_API
diff --git a/Engine/Source/Editor/InternationalizationSettings/Public/SCulturePicker.h b/Engine/Source/Editor/InternationalizationSettings/Public/SCulturePicker.h
--- a/Engine/Source/Editor/InternationalizationSettings/Public/SCulturePicker.h
+++ b/Engine/Source/Editor/InternationalizationSettings/Public/SCulturePicker.h
@@ -24,6 +24,8 @@
 #include "Widgets/Views/STableViewBase.h"
 #include "Widgets/Views/STreeView.h"
 
+#define UE_API INTERNATIONALIZATIONSETTINGS_API
+
 class FText;
 class ITableRow;
 class SComboButton;
@@ -52,7 +54,7 @@ struct FCultureEntry
 	}
 };
 
-class INTERNATIONALIZATIONSETTINGS_API SCulturePicker : public SCompoundWidget
+class SCulturePicker : public SCompoundWidget
 {
 public:
 	/** A delegate type invoked when a selection changes somewhere. */
@@ -100,26 +102,26 @@ public:
 	, SuppressSelectionCallback(false)
 	{}
 
-	void Construct( const FArguments& InArgs );
-	void RequestTreeRefresh();
+	UE_API void Construct( const FArguments& InArgs );
+	UE_API void RequestTreeRefresh();
 
 private:
-	TSharedPtr<FCultureEntry> FindEntryForCulture(FCulturePtr Culture) const;
-	TSharedPtr<FCultureEntry> FindEntryForCultureImpl(FCulturePtr Culture, const TArray<TSharedPtr<FCultureEntry>>& Entries) const;
+	UE_API TSharedPtr<FCultureEntry> FindEntryForCulture(FCulturePtr Culture) const;
+	UE_API TSharedPtr<FCultureEntry> FindEntryForCultureImpl(FCulturePtr Culture, const TArray<TSharedPtr<FCultureEntry>>& Entries) const;
 
-	void AutoExpandEntries();
-	void AutoExpandEntriesImpl(const TArray<TSharedPtr<FCultureEntry>>& Entries);
+	UE_API void AutoExpandEntries();
+	UE_API void AutoExpandEntriesImpl(const TArray<TSharedPtr<FCultureEntry>>& Entries);
 
-	void BuildStockEntries();
-	void RebuildEntries();
+	UE_API void BuildStockEntries();
+	UE_API void RebuildEntries();
 
-	void OnFilterStringChanged(const FText& InFilterString);
-	TSharedRef<ITableRow> OnGenerateRow(TSharedPtr<FCultureEntry> Entry, const TSharedRef<STableViewBase>& Table);
-	void OnGetChildren(TSharedPtr<FCultureEntry> Entry, TArray< TSharedPtr<FCultureEntry> >& Children);
-	void OnSelectionChanged(TSharedPtr<FCultureEntry> Entry, ESelectInfo::Type SelectInfo);
+	UE_API void OnFilterStringChanged(const FText& InFilterString);
+	UE_API TSharedRef<ITableRow> OnGenerateRow(TSharedPtr<FCultureEntry> Entry, const TSharedRef<STableViewBase>& Table);
+	UE_API void OnGetChildren(TSharedPtr<FCultureEntry> Entry, TArray< TSharedPtr<FCultureEntry> >& Children);
+	UE_API void OnSelectionChanged(TSharedPtr<FCultureEntry> Entry, ESelectInfo::Type SelectInfo);
 
 private:
-	FString GetCultureDisplayName(const FCultureRef& Culture, const bool bIsRootItem) const;
+	UE_API FString GetCultureDisplayName(const FCultureRef& Culture, const bool bIsRootItem) const;
 
 	TSharedPtr< STreeView< TSharedPtr<FCultureEntry> > > TreeView;
 
@@ -151,7 +153,7 @@ private:
 	bool SuppressSelectionCallback;
 };
 
-class INTERNATIONALIZATIONSETTINGS_API SCulturePickerCombo : public SCompoundWidget
+class SCulturePickerCombo : public SCompoundWidget
 {
 public:
 	SLATE_BEGIN_ARGS(SCulturePickerCombo)
@@ -168,11 +170,11 @@ public:
 		SLATE_ARGUMENT(bool, CanSelectNone)
 	SLATE_END_ARGS()
 
-	void Construct(const FArguments& InArgs);
+	UE_API void Construct(const FArguments& InArgs);
 
 private:
-	FText GetSelectedCultureDisplayName() const;
-	void OnSelectedCultureChanged(FCulturePtr NewSelectedCulture, ESelectInfo::Type SelectInfo);
+	UE_API FText GetSelectedCultureDisplayName() const;
+	UE_API void OnSelectedCultureChanged(FCulturePtr NewSelectedCulture, ESelectInfo::Type SelectInfo);
 
 	TAttribute<FCulturePtr> SelectedCulture;
 	SCulturePicker::FOnSelectionChanged OnSelectionChanged;
@@ -180,3 +182,5 @@ private:
 
 	TSharedPtr<SComboButton> ComboButton;
 };
+
+#undef UE_API
diff --git a/Engine/Source/Editor/Kismet/Classes/BlueprintPaletteFavorites.h b/Engine/Source/Editor/Kismet/Classes/BlueprintPaletteFavorites.h
--- a/Engine/Source/Editor/Kismet/Classes/BlueprintPaletteFavorites.h
+++ b/Engine/Source/Editor/Kismet/Classes/BlueprintPaletteFavorites.h
@@ -15,6 +15,8 @@
 
 #include "BlueprintPaletteFavorites.generated.h"
 
+#define UE_API KISMET_API
+
 class UBlueprintNodeSpawner;
 struct FBlueprintActionInfo;
 struct FEdGraphSchemaAction;
@@ -24,12 +26,12 @@ struct FEdGraphSchemaAction;
 *******************************************************************************/
 
 USTRUCT()
-struct KISMET_API FFavoritedBlueprintPaletteItem
+struct FFavoritedBlueprintPaletteItem
 {
 	GENERATED_USTRUCT_BODY()
 
 public:
-	FFavoritedBlueprintPaletteItem();
+	UE_API FFavoritedBlueprintPaletteItem();
 
 	/**
 	 * Sometime favorites can be coming from user edited .ini files, so this 
@@ -38,7 +40,7 @@ public:
 	 * 
 	 * @param  SerializedAction		A string representing a serialized favorite.
 	 */
-	FFavoritedBlueprintPaletteItem(FString const& SerializedAction);
+	UE_API FFavoritedBlueprintPaletteItem(FString const& SerializedAction);
 
 	/**
 	 * Constructs a favorite from the specified palette action (some palette actions
@@ -46,10 +48,10 @@ public:
 	 * 
 	 * @param  PaletteAction	The action you wish to favorite.
 	 */
-	FFavoritedBlueprintPaletteItem(const FEdGraphSchemaAction& InPaletteAction);
+	UE_API FFavoritedBlueprintPaletteItem(const FEdGraphSchemaAction& InPaletteAction);
 
 	UE_DEPRECATED(5.5, "Provide a dereferenced FEdGraphSchemaAction, after validating it in the caller")
-	FFavoritedBlueprintPaletteItem(TSharedPtr<FEdGraphSchemaAction> PaletteAction);
+	UE_API FFavoritedBlueprintPaletteItem(TSharedPtr<FEdGraphSchemaAction> PaletteAction);
 
 	/**
 	 * 
@@ -57,7 +59,7 @@ public:
 	 * @param  BlueprintAction	
 	 * @return 
 	 */
-	FFavoritedBlueprintPaletteItem(UBlueprintNodeSpawner const* BlueprintAction);
+	UE_API FFavoritedBlueprintPaletteItem(UBlueprintNodeSpawner const* BlueprintAction);
 
 	/**
 	 * Sometimes we're not able to construct favorites from specified actions, 
@@ -65,7 +67,7 @@ public:
 	 * 
 	 * @return True if this favorite is valid (refers to a specific node), false if not.
 	 */
-	bool IsValid() const;
+	UE_API bool IsValid() const;
 
 	/**
 	 * Checks to see if this favorite matches the specified one.
@@ -73,7 +75,7 @@ public:
 	 * @param  Rhs	The other favorite you want to compare with.
 	 * @return True if this favorite matches the other one, false if not.
 	 */
-	bool operator==(FFavoritedBlueprintPaletteItem const& Rhs) const;
+	UE_API bool operator==(FFavoritedBlueprintPaletteItem const& Rhs) const;
 
 	/**
 	 * Checks to see if this favorite represents the supplied ed-graph action 
@@ -82,7 +84,7 @@ public:
 	 * @param  PaletteAction	The action you want to compare with.
 	 * @return True if this favorite represents the specified action, false if not.
 	 */
-	bool operator==(TSharedPtr<FEdGraphSchemaAction> PaletteAction) const;
+	UE_API bool operator==(TSharedPtr<FEdGraphSchemaAction> PaletteAction) const;
 
 	/**
 	 * We want to be able to specify some of these in .ini files, so let we have
@@ -90,7 +92,7 @@ public:
 	 * 
 	 * @return A string representation of this item.
 	 */
-	FString const& ToString() const;
+	UE_API FString const& ToString() const;
 
 private:
 	FBlueprintNodeSignature ActionSignature;
@@ -100,15 +102,15 @@ private:
 * FBlueprintPaletteFavorites
 *******************************************************************************/
 
-UCLASS(config=EditorPerProjectUserSettings)
-class KISMET_API UBlueprintPaletteFavorites : public UObject
+UCLASS(MinimalAPI, config=EditorPerProjectUserSettings)
+class UBlueprintPaletteFavorites : public UObject
 {
 	GENERATED_UCLASS_BODY()
 
 public:
 	//~ Begin UObject Interface
-	virtual void PostInitProperties() override;
-	virtual void PostEditChangeProperty(struct FPropertyChangedEvent& PropertyChangedEvent) override;
+	UE_API virtual void PostInitProperties() override;
+	UE_API virtual void PostEditChangeProperty(struct FPropertyChangedEvent& PropertyChangedEvent) override;
 	//~ End UObject Interface
 
 	/**
@@ -118,7 +120,7 @@ public:
 	 * @param  PaletteAction	The action you wish to check.
 	 * @return True if the action can be turned into a favorite, false if not.
 	 */
-	bool CanBeFavorited(TSharedPtr<FEdGraphSchemaAction> PaletteAction) const;
+	UE_API bool CanBeFavorited(TSharedPtr<FEdGraphSchemaAction> PaletteAction) const;
 
 	/**
 	 * This method can be used to see if a specified action is already favorited 
@@ -127,8 +129,8 @@ public:
 	 * @param  PaletteAction	The action you wish to check.
 	 * @return True if this action is already favorited, false if it not.
 	 */
-	bool IsFavorited(TSharedPtr<FEdGraphSchemaAction> PaletteAction) const;
-	bool IsFavorited(const FEdGraphSchemaAction& InCurrentAction) const;
+	UE_API bool IsFavorited(TSharedPtr<FEdGraphSchemaAction> PaletteAction) const;
+	UE_API bool IsFavorited(const FEdGraphSchemaAction& InCurrentAction) const;
 	
 	/**
 	 * 
@@ -136,7 +138,7 @@ public:
 	 * @param  BlueprintAction	
 	 * @return 
 	 */
-	bool IsFavorited(FBlueprintActionInfo& BlueprintAction) const;
+	UE_API bool IsFavorited(FBlueprintActionInfo& BlueprintAction) const;
 
 	/**
 	 * 
@@ -144,7 +146,7 @@ public:
 	 * @param  BlueprintAction	
 	 * @return 
 	 */
-	bool IsFavorited(UBlueprintNodeSpawner const* BlueprintAction) const;
+	UE_API bool IsFavorited(UBlueprintNodeSpawner const* BlueprintAction) const;
 
 	/**
 	 * Adds the specified action to the current favorites list (fails if the action
@@ -153,7 +155,7 @@ public:
 	 * 
 	 * @param  PaletteAction	The action you wish to add.
 	 */
-	void AddFavorite(TSharedPtr<FEdGraphSchemaAction> PaletteAction);
+	UE_API void AddFavorite(TSharedPtr<FEdGraphSchemaAction> PaletteAction);
 
 	/**
 	 * Adds the specified actions to the current favorites list. Will also 
@@ -161,7 +163,7 @@ public:
 	 * 
 	 * @param  PaletteActions	An array of action you wish to add.
 	 */
-	void AddFavorites(TArray< TSharedPtr<FEdGraphSchemaAction> > PaletteActions);
+	UE_API void AddFavorites(TArray< TSharedPtr<FEdGraphSchemaAction> > PaletteActions);
 
 	/**
 	 * Removes the specified action to the current favorites list (if it's 
@@ -170,7 +172,7 @@ public:
 	 * 
 	 * @param  PaletteAction	The action you wish to remove.
 	 */
-	void RemoveFavorite(TSharedPtr<FEdGraphSchemaAction> PaletteAction);
+	UE_API void RemoveFavorite(TSharedPtr<FEdGraphSchemaAction> PaletteAction);
 
 	/**
 	 * Remove the specified actions from the current favorites list. Will also 
@@ -178,7 +180,7 @@ public:
 	 * 
 	 * @param  PaletteActions	An array of action you wish to remove.
 	 */
-	void RemoveFavorites(TArray< TSharedPtr<FEdGraphSchemaAction> > PaletteActions);
+	UE_API void RemoveFavorites(TArray< TSharedPtr<FEdGraphSchemaAction> > PaletteActions);
 
 	/**
 	 * Throws out all current favorites and loads in ones for the specified 
@@ -186,7 +188,7 @@ public:
 	 * 
 	 * @param  ProfileName	The name of the profile you wish to load.
 	 */
-	void LoadProfile(FString const& ProfileName);
+	UE_API void LoadProfile(FString const& ProfileName);
 
 	/**
 	 * Provides any easy way to see if the user is currently using their own
@@ -194,7 +196,7 @@ public:
 	 * 
 	 * @return True if the user is using their own profile, false if it is a predefined one.
 	 */
-	bool IsUsingCustomProfile() const;
+	UE_API bool IsUsingCustomProfile() const;
 
 	/**
 	 * Gets the user's currently set profile. If the user hasn't manually set 
@@ -202,13 +204,13 @@ public:
 	 * 
 	 * @return A string name, representing the currently set profile (defined in the editor .ini file)
 	 */
-	FString const& GetCurrentProfile() const;
+	UE_API FString const& GetCurrentProfile() const;
 
 	/**
 	 * Removes every single favorite and sets the user's profile to a custom one
 	 * (if it isn't already).
 	 */
-	void ClearAllFavorites();
+	UE_API void ClearAllFavorites();
 
 	/** 
 	 * A event for users to hook into (specifically the UI), so they can be 
@@ -224,19 +226,19 @@ private:
 	 * 
 	 * @return A string name, representing the default profile (defined in the editor .ini file)
 	 */
-	FString const& GetDefaultProfileId() const;
+	UE_API FString const& GetDefaultProfileId() const;
 
 	/**
 	 * Throws out all current favorites and loads in ones for specified by 
 	 * CurrentProfile.
 	 */
-	void LoadSetProfile();
+	UE_API void LoadSetProfile();
 
 	/**
 	 * Fills the CurrentFavorites array with items that have been loaded into 
 	 * the CustomFavorites array.
 	 */
-	void LoadCustomFavorites();
+	UE_API void LoadCustomFavorites();
 
 	/**
 	 * Modifies the CurrentProfile member, but flags the property as edited so 
@@ -244,7 +246,7 @@ private:
 	 * 
 	 * @param  NewProfileName	The new profile that you want favorites saved with.
 	 */
-	void SetProfile(FString const& NewProfileName);
+	UE_API void SetProfile(FString const& NewProfileName);
 
 public:
 	/** 
@@ -272,3 +274,4 @@ public:
 	FString CurrentProfile;
 };
 
+#undef UE_API
diff --git a/Engine/Source/Editor/Kismet/Private/BPDelegateDragDropAction.h b/Engine/Source/Editor/Kismet/Private/BPDelegateDragDropAction.h
--- a/Engine/Source/Editor/Kismet/Private/BPDelegateDragDropAction.h
+++ b/Engine/Source/Editor/Kismet/Private/BPDelegateDragDropAction.h
@@ -20,12 +20,14 @@
 #include "UObject/WeakObjectPtr.h"
 #include "UObject/WeakObjectPtrTemplates.h"
 
+#define UE_API KISMET_API
+
 class SWidget;
 class UEdGraph;
 struct FEdGraphSchemaAction;
 
 /** DragDropAction class for dropping a Variable onto a graph */
-class KISMET_API FKismetDelegateDragDropAction : public FKismetVariableDragDropAction
+class FKismetDelegateDragDropAction : public FKismetVariableDragDropAction
 {
 public:
 	DRAG_DROP_OPERATOR_TYPE(FKismetDelegateDragDropAction, FKismetVariableDragDropAction)
@@ -33,12 +35,12 @@ public:
 	// FGraphEditorDragDropAction interface
 	virtual void HoverTargetChanged() override { FMyBlueprintItemDragDropAction::HoverTargetChanged(); }
 
-	virtual FReply DroppedOnPanel(const TSharedRef< SWidget >& Panel, const FVector2f& ScreenPosition, const FVector2f& GraphPosition, UEdGraph& Graph) override;
+	UE_API virtual FReply DroppedOnPanel(const TSharedRef< SWidget >& Panel, const FVector2f& ScreenPosition, const FVector2f& GraphPosition, UEdGraph& Graph) override;
 
-	virtual bool IsSupportedBySchema(const class UEdGraphSchema* Schema) const override;
+	UE_API virtual bool IsSupportedBySchema(const class UEdGraphSchema* Schema) const override;
 	// End of FGraphEditorDragDropAction
 
-	bool IsValid() const;
+	UE_API bool IsValid() const;
 	
 	static TSharedRef<FKismetDelegateDragDropAction> New(TSharedPtr<FEdGraphSchemaAction> InSourceAction, FName InVariableName, UStruct* InSource, FNodeCreationAnalytic AnalyticCallback)
 	{
@@ -78,11 +80,13 @@ public:
 	}
 
 	/** Create new custom event node from construction parameters */
-	static void MakeEvent(FNodeConstructionParams Params);
+	static UE_API void MakeEvent(FNodeConstructionParams Params);
 
 	/** Assign new delegate node from construction parameters */
-	static void AssignEvent(FNodeConstructionParams Params);
+	static UE_API void AssignEvent(FNodeConstructionParams Params);
 
 protected:
 	FKismetDelegateDragDropAction() {}
 };
+
+#undef UE_API
diff --git a/Engine/Source/Editor/Kismet/Private/BPFunctionDragDropAction.h b/Engine/Source/Editor/Kismet/Private/BPFunctionDragDropAction.h
--- a/Engine/Source/Editor/Kismet/Private/BPFunctionDragDropAction.h
+++ b/Engine/Source/Editor/Kismet/Private/BPFunctionDragDropAction.h
@@ -13,6 +13,8 @@
 #include "Templates/SharedPointer.h"
 #include "UObject/NameTypes.h"
 
+#define UE_API KISMET_API
+
 class FText;
 class UBlueprint;
 class UClass;
@@ -24,14 +26,14 @@ struct FEdGraphSchemaAction;
 * FKismetDragDropAction
 *******************************************************************************/
 
-class KISMET_API FKismetDragDropAction : public FMyBlueprintItemDragDropAction
+class FKismetDragDropAction : public FMyBlueprintItemDragDropAction
 {
 public:
 	DRAG_DROP_OPERATOR_TYPE(FKismetDragDropAction, FMyBlueprintItemDragDropAction)
 		
 	// FGraphEditorDragDropAction interface
-	virtual void HoverTargetChanged() override;
-	virtual FReply DroppedOnPanel( const TSharedRef< class SWidget >& Panel, const FVector2f& ScreenPosition, const FVector2f& GraphPosition, UEdGraph& Graph) override;
+	UE_API virtual void HoverTargetChanged() override;
+	UE_API virtual FReply DroppedOnPanel( const TSharedRef< class SWidget >& Panel, const FVector2f& ScreenPosition, const FVector2f& GraphPosition, UEdGraph& Graph) override;
 	// End of FGraphSchemaActionDragDropAction
 
 	DECLARE_DELEGATE_RetVal_ThreeParams(bool, FCanBeDroppedDelegate, TSharedPtr<FEdGraphSchemaAction> /*DropAction*/, UEdGraph* /*HoveredGraphIn*/, FText& /*ImpededReasonOut*/);
@@ -47,7 +49,7 @@ public:
 	}
 
 protected:
-	bool ActionWillShowExistingNode() const;
+	UE_API bool ActionWillShowExistingNode() const;
 
 	/** */
 	FCanBeDroppedDelegate CanBeDroppedDelegate;
@@ -57,19 +59,19 @@ protected:
 * FKismetFunctionDragDropAction
 *******************************************************************************/
 
-class KISMET_API FKismetFunctionDragDropAction : public FKismetDragDropAction
+class FKismetFunctionDragDropAction : public FKismetDragDropAction
 {
 public:
 	DRAG_DROP_OPERATOR_TYPE(FKismetFunctionDragDropAction, FKismetDragDropAction)
 
-	FKismetFunctionDragDropAction();
+	UE_API FKismetFunctionDragDropAction();
 
 	// FGraphEditorDragDropAction interface
-	virtual FReply DroppedOnPanel( const TSharedRef< class SWidget >& Panel, const FVector2f& ScreenPosition, const FVector2f& GraphPosition, UEdGraph& Graph) override;
-	virtual FReply DroppedOnPin(const FVector2f& ScreenPosition, const FVector2f& GraphPosition) override;
+	UE_API virtual FReply DroppedOnPanel( const TSharedRef< class SWidget >& Panel, const FVector2f& ScreenPosition, const FVector2f& GraphPosition, UEdGraph& Graph) override;
+	UE_API virtual FReply DroppedOnPin(const FVector2f& ScreenPosition, const FVector2f& GraphPosition) override;
 	// End of FGraphEditorDragDropAction
 
-	static TSharedRef<FKismetFunctionDragDropAction> New(TSharedPtr<FEdGraphSchemaAction> InActionNode, FName InFunctionName, UClass* InOwningClass, const FMemberReference& InCallOnMember, FNodeCreationAnalytic AnalyticCallback, FCanBeDroppedDelegate CanBeDroppedDelegate = FCanBeDroppedDelegate());
+	static UE_API TSharedRef<FKismetFunctionDragDropAction> New(TSharedPtr<FEdGraphSchemaAction> InActionNode, FName InFunctionName, UClass* InOwningClass, const FMemberReference& InCallOnMember, FNodeCreationAnalytic AnalyticCallback, FCanBeDroppedDelegate CanBeDroppedDelegate = FCanBeDroppedDelegate());
 
 protected:
 	/** Name of function being dragged */
@@ -80,28 +82,28 @@ protected:
 	FMemberReference CallOnMember;
 
 	/** Looks up the functions field on OwningClass using FunctionName */
-	UFunction const* GetFunctionProperty() const;
+	UE_API UFunction const* GetFunctionProperty() const;
 
 	/** Constructs an action to execute, placing a function call node for the associated function */
-	class UBlueprintFunctionNodeSpawner* GetDropAction(UEdGraph& Graph) const;
+	UE_API class UBlueprintFunctionNodeSpawner* GetDropAction(UEdGraph& Graph) const;
 };
 
 /*******************************************************************************
 * FKismetMacroDragDropAction
 *******************************************************************************/
 
-class KISMET_API FKismetMacroDragDropAction : public FKismetDragDropAction
+class FKismetMacroDragDropAction : public FKismetDragDropAction
 {
 public:
 	DRAG_DROP_OPERATOR_TYPE(FKismetMacroDragDropAction, FKismetDragDropAction)
 
-	FKismetMacroDragDropAction();
+	UE_API FKismetMacroDragDropAction();
 
 	// FGraphEditorDragDropAction interface
-	virtual FReply DroppedOnPanel( const TSharedRef< class SWidget >& Panel, const FVector2f& ScreenPosition, const FVector2f& GraphPosition, UEdGraph& Graph) override;
+	UE_API virtual FReply DroppedOnPanel( const TSharedRef< class SWidget >& Panel, const FVector2f& ScreenPosition, const FVector2f& GraphPosition, UEdGraph& Graph) override;
 	// End of FGraphEditorDragDropAction
 
-	static TSharedRef<FKismetMacroDragDropAction> New(TSharedPtr<FEdGraphSchemaAction> InActionNode, FName InMacroName, UBlueprint* InBlueprint, UEdGraph* InMacro, FNodeCreationAnalytic AnalyticCallback);
+	static UE_API TSharedRef<FKismetMacroDragDropAction> New(TSharedPtr<FEdGraphSchemaAction> InActionNode, FName InMacroName, UBlueprint* InBlueprint, UEdGraph* InMacro, FNodeCreationAnalytic AnalyticCallback);
 
 protected:
 	// FMyBlueprintItemDragDropAction interface
@@ -120,3 +122,4 @@ protected:
 	UBlueprint* Blueprint;
 };
 
+#undef UE_API
diff --git a/Engine/Source/Editor/Kismet/Private/DiffControl.h b/Engine/Source/Editor/Kismet/Private/DiffControl.h
--- a/Engine/Source/Editor/Kismet/Private/DiffControl.h
+++ b/Engine/Source/Editor/Kismet/Private/DiffControl.h
@@ -25,7 +25,7 @@ class SCheckBox;
 class FAsyncDetailViewDiff;
 
 /** Interface responsible for generating FBlueprintDifferenceTreeEntry's for visual diff tools */
-class KISMET_API IDiffControl
+class IDiffControl
 {
 public:
 	virtual ~IDiffControl() {}
@@ -35,18 +35,18 @@ public:
 	/** Adds widgets to the tree of differences to show */
 	virtual void GenerateTreeEntries(TArray< TSharedPtr<FBlueprintDifferenceTreeEntry> >& OutTreeEntries, TArray< TSharedPtr<FBlueprintDifferenceTreeEntry> >& OutRealDifferences) = 0;
 	
-	static FText RightRevision;
-	static TSharedRef<SWidget> GenerateSimpleDiffWidget(FText DiffText);
-	static TSharedRef<SWidget> GenerateObjectDiffWidget(FSingleObjectDiffEntry DiffEntry, FText ObjectName);
+	static KISMET_API FText RightRevision;
+	static KISMET_API TSharedRef<SWidget> GenerateSimpleDiffWidget(FText DiffText);
+	static KISMET_API TSharedRef<SWidget> GenerateObjectDiffWidget(FSingleObjectDiffEntry DiffEntry, FText ObjectName);
 
 	// to support comment posting, set this to the tree view that contains the comments.
-	void EnableComments(TWeakPtr<STreeView<TSharedPtr<FBlueprintDifferenceTreeEntry>>> TreeView, const UObject* OldObject, const UObject* NewObject);
-	void EnableComments(TWeakPtr<STreeView<TSharedPtr<FBlueprintDifferenceTreeEntry>>> TreeView, const TArray<const UObject*>& Objects);
+	KISMET_API void EnableComments(TWeakPtr<STreeView<TSharedPtr<FBlueprintDifferenceTreeEntry>>> TreeView, const UObject* OldObject, const UObject* NewObject);
+	KISMET_API void EnableComments(TWeakPtr<STreeView<TSharedPtr<FBlueprintDifferenceTreeEntry>>> TreeView, const TArray<const UObject*>& Objects);
 	virtual void EnableComments(TWeakPtr<STreeView<TSharedPtr<FBlueprintDifferenceTreeEntry>>> TreeView) {}
 	
 protected:
 	// to be called at the end of GenerateTreeEntries if you want this diff control to support review comments
-	void GenerateCategoryCommentTreeEntries(TArray< TSharedPtr<FBlueprintDifferenceTreeEntry> >& , TArray< TSharedPtr<FBlueprintDifferenceTreeEntry> >& OutRealDifferences, FString CategoryKey);
+	KISMET_API void GenerateCategoryCommentTreeEntries(TArray< TSharedPtr<FBlueprintDifferenceTreeEntry> >& , TArray< TSharedPtr<FBlueprintDifferenceTreeEntry> >& OutRealDifferences, FString CategoryKey);
 	
 	TSharedPtr<FReviewCommentsDiffControl> ReviewCommentsDiffControl;
 };
@@ -94,39 +94,39 @@ private:
 };
 
 /** Generic wrapper around a details view template parameter determines whether TreeEntries are populated */
-class KISMET_API FDetailsDiffControl : public TSharedFromThis<FDetailsDiffControl>, public IDiffControl
+class FDetailsDiffControl : public TSharedFromThis<FDetailsDiffControl>, public IDiffControl
 {
 public:
-	FDetailsDiffControl(const UObject* InOldObject, const UObject* InNewObject, FOnDiffEntryFocused InSelectionCallback, bool bPopulateOutTreeEntries);
-	~FDetailsDiffControl();
+	KISMET_API FDetailsDiffControl(const UObject* InOldObject, const UObject* InNewObject, FOnDiffEntryFocused InSelectionCallback, bool bPopulateOutTreeEntries);
+	KISMET_API ~FDetailsDiffControl();
 
-	virtual void Tick() override;
+	KISMET_API virtual void Tick() override;
 	
-	virtual void GenerateTreeEntries(TArray< TSharedPtr<FBlueprintDifferenceTreeEntry> >& OutTreeEntries, TArray< TSharedPtr<FBlueprintDifferenceTreeEntry> >& OutRealDifferences) override;
-	void GenerateTreeEntriesWithoutComments(TArray< TSharedPtr<FBlueprintDifferenceTreeEntry> >& OutTreeEntries, TArray< TSharedPtr<FBlueprintDifferenceTreeEntry> >& OutRealDifferences);
+	KISMET_API virtual void GenerateTreeEntries(TArray< TSharedPtr<FBlueprintDifferenceTreeEntry> >& OutTreeEntries, TArray< TSharedPtr<FBlueprintDifferenceTreeEntry> >& OutRealDifferences) override;
+	KISMET_API void GenerateTreeEntriesWithoutComments(TArray< TSharedPtr<FBlueprintDifferenceTreeEntry> >& OutTreeEntries, TArray< TSharedPtr<FBlueprintDifferenceTreeEntry> >& OutRealDifferences);
 
-	TSharedRef<IDetailsView> InsertObject(const UObject* Object, bool bScrollbarOnLeft = false, int32 Index = INDEX_NONE);
-	TSharedRef<IDetailsView> GetDetailsWidget(const UObject* Object) const;
-	TSharedPtr<IDetailsView> TryGetDetailsWidget(const UObject* Object) const;
-	TSharedPtr<FAsyncDetailViewDiff> GetDifferencesWithLeft(const UObject* Object) const;
-	TSharedPtr<FAsyncDetailViewDiff> GetDifferencesWithRight(const UObject* Object) const;
-	int32 IndexOfObject(const UObject* Object) const;
+	KISMET_API TSharedRef<IDetailsView> InsertObject(const UObject* Object, bool bScrollbarOnLeft = false, int32 Index = INDEX_NONE);
+	KISMET_API TSharedRef<IDetailsView> GetDetailsWidget(const UObject* Object) const;
+	KISMET_API TSharedPtr<IDetailsView> TryGetDetailsWidget(const UObject* Object) const;
+	KISMET_API TSharedPtr<FAsyncDetailViewDiff> GetDifferencesWithLeft(const UObject* Object) const;
+	KISMET_API TSharedPtr<FAsyncDetailViewDiff> GetDifferencesWithRight(const UObject* Object) const;
+	KISMET_API int32 IndexOfObject(const UObject* Object) const;
 
 	// to support comment posting, set this to the tree view that contains the comments.
-	virtual void EnableComments(TWeakPtr<STreeView<TSharedPtr<FBlueprintDifferenceTreeEntry>>> TreeView) override;
+	KISMET_API virtual void EnableComments(TWeakPtr<STreeView<TSharedPtr<FBlueprintDifferenceTreeEntry>>> TreeView) override;
 
 	DiffUtils::FOnGenerateCustomDiffEntries GenerateCustomEntriesCallback;
 	DiffUtils::FOnGenerateCustomDiffEntryWidget GenerateCustomEntryWidgetCallback;
 	DiffUtils::FOnOrganizeDiffEntries OrganizeEntriesCallback;
 protected:
-	virtual void OnSelectDiffEntry(FPropertySoftPath PropertyName);
-	TAttribute<TArray<FVector2f>> GetLinkedScrollRateAttribute(const TSharedRef<IDetailsView>& OldDetailsView, const TSharedRef<IDetailsView>& NewDetailsView);
+	KISMET_API virtual void OnSelectDiffEntry(FPropertySoftPath PropertyName);
+	KISMET_API TAttribute<TArray<FVector2f>> GetLinkedScrollRateAttribute(const TSharedRef<IDetailsView>& OldDetailsView, const TSharedRef<IDetailsView>& NewDetailsView);
 	
 
 	// helper function that analyzes two details views and determines the rate they should scroll relative to one another to be in sync
-	TArray<FVector2f> GetLinkedScrollRate(TSharedRef<IDetailsView> LeftDetailsView, TSharedRef<IDetailsView> RightDetailsView) const;
+	KISMET_API TArray<FVector2f> GetLinkedScrollRate(TSharedRef<IDetailsView> LeftDetailsView, TSharedRef<IDetailsView> RightDetailsView) const;
 private:
-	void HandlePropertiesChanged();
+	KISMET_API void HandlePropertiesChanged();
 
 	FDetailsDiffControl& operator=(const FDetailsDiffControl&) = delete;
 	FDetailsDiffControl( const FDetailsDiffControl& ) = delete;
@@ -415,4 +415,4 @@ protected:
 
 	// we only want one reply draft widget in existence so we can use a static to track it
 	static TWeakPtr<FCommentDraftTreeEntry> ReplyDraftEntry;
-};
\ No newline at end of file
+};
