[Graph Editor] Fix a regression where the offset would be incorrectly applied in deprecation path to the size parameter

diff --git a/Engine/Source/Editor/GraphEditor/Public/SNodePanel.h b/Engine/Source/Editor/GraphEditor/Public/SNodePanel.h
--- a/Engine/Source/Editor/GraphEditor/Public/SNodePanel.h
+++ b/Engine/Source/Editor/GraphEditor/Public/SNodePanel.h
@@ -391,7 +391,7 @@ public:
 			{
 			PRAGMA_DISABLE_DEPRECATION_WARNINGS
 				Size_Deprecated = MoveTemp(InSize);
-				Offset.BindLambda([this]()
+				Size.BindLambda([this]()
 				{
 					return UE::Slate::CastToVector2f(Size_Deprecated.Get());
 				});
