Fixed LOCTEXT literals using hex escape sequences, in favor of using Unicode escape sequences

#rnx

diff --git a/Engine/Source/Editor/BlueprintGraph/Private/K2Node_CastByteToEnum.cpp b/Engine/Source/Editor/BlueprintGraph/Private/K2Node_CastByteToEnum.cpp
--- a/Engine/Source/Editor/BlueprintGraph/Private/K2Node_CastByteToEnum.cpp
+++ b/Engine/Source/Editor/BlueprintGraph/Private/K2Node_CastByteToEnum.cpp
@@ -84,7 +84,7 @@ FSlateIcon UK2Node_CastByteToEnum::GetIconAndTint(FLinearColor& OutColor) const
 
 FText UK2Node_CastByteToEnum::GetCompactNodeTitle() const
 {
-	return NSLOCTEXT("K2Node", "CastSymbol", "\x2022");
+	return NSLOCTEXT("K2Node", "CastSymbol", "\u2022");
 }
 
 FName UK2Node_CastByteToEnum::GetFunctionName() const
diff --git a/Engine/Source/Editor/BlueprintGraph/Private/K2Node_GetEnumeratorName.cpp b/Engine/Source/Editor/BlueprintGraph/Private/K2Node_GetEnumeratorName.cpp
--- a/Engine/Source/Editor/BlueprintGraph/Private/K2Node_GetEnumeratorName.cpp
+++ b/Engine/Source/Editor/BlueprintGraph/Private/K2Node_GetEnumeratorName.cpp
@@ -53,7 +53,7 @@ FText UK2Node_GetEnumeratorName::GetNodeTitle(ENodeTitleType::Type TitleType) co
 
 FText UK2Node_GetEnumeratorName::GetCompactNodeTitle() const
 {
-	return NSLOCTEXT("K2Node", "CastSymbol", "\x2022");
+	return NSLOCTEXT("K2Node", "CastSymbol", "\u2022");
 }
 
 UEnum* UK2Node_GetEnumeratorName::GetConnectedEnum() const
diff --git a/Engine/Source/Editor/GraphEditor/Private/SGraphEditorImpl.cpp b/Engine/Source/Editor/GraphEditor/Private/SGraphEditorImpl.cpp
--- a/Engine/Source/Editor/GraphEditor/Private/SGraphEditorImpl.cpp
+++ b/Engine/Source/Editor/GraphEditor/Private/SGraphEditorImpl.cpp
@@ -1182,7 +1182,7 @@ void SGraphEditorImpl::RegisterContextMenu(const UEdGraphSchema* Schema, FToolMe
 							const UMaterialGraphNode_Root* RootNode = Cast<UMaterialGraphNode_Root>(ContextObject->Node);
 
 							TSharedRef<SButton> SubstrateConvertButton = SNew(SButton)
-								.Text(LOCTEXT("SubstrateConvertButtonLabel", "\x26A0 Convert To Substrate \x26A0"))
+								.Text(LOCTEXT("SubstrateConvertButtonLabel", "\u26A0 Convert To Substrate \u26A0"))
 								.ToolTipText(LOCTEXT("SubstrateConvertButtonToolTip", "Convert the root node to a Substrate graph. This material will then use Substrate more advanced features and topology processes.\nTo go back to the previous Non-Substrate material, you can use CTRL+Z or you will need to manually relink the graph."))
 								.ButtonStyle(FAppStyle::Get(), "Button")
 								.OnClicked_Lambda([RootNode]()
diff --git a/Engine/Source/Runtime/ApplicationCore/Private/Linux/LinuxPlatformSplash.cpp b/Engine/Source/Runtime/ApplicationCore/Private/Linux/LinuxPlatformSplash.cpp
--- a/Engine/Source/Runtime/ApplicationCore/Private/Linux/LinuxPlatformSplash.cpp
+++ b/Engine/Source/Runtime/ApplicationCore/Private/Linux/LinuxPlatformSplash.cpp
@@ -795,7 +795,7 @@ void FLinuxPlatformSplash::Show( )
 
 		// Display copyright information in editor splash screen
 		{
-			const FText CopyrightInfo = NSLOCTEXT( "UnrealEd", "SplashScreen_CopyrightInfo", "Copyright \x00a9   Epic Games, Inc.   All rights reserved." );
+			const FText CopyrightInfo = NSLOCTEXT( "UnrealEd", "SplashScreen_CopyrightInfo", "Copyright \u00a9   Epic Games, Inc.   All rights reserved." );
 			GSplashState->SetSplashText( SplashTextType::CopyrightInfo, CopyrightInfo );
 		}
 
diff --git a/Engine/Source/Runtime/ApplicationCore/Private/Mac/MacPlatformSplash.cpp b/Engine/Source/Runtime/ApplicationCore/Private/Mac/MacPlatformSplash.cpp
--- a/Engine/Source/Runtime/ApplicationCore/Private/Mac/MacPlatformSplash.cpp
+++ b/Engine/Source/Runtime/ApplicationCore/Private/Mac/MacPlatformSplash.cpp
@@ -230,7 +230,7 @@ void FMacPlatformSplash::Show()
 
 				// Display copyright information in editor splash screen
 				{
-					const FText CopyrightInfo = NSLOCTEXT( "UnrealEd", "SplashScreen_CopyrightInfo", "Copyright \x00a9   Epic Games, Inc.   All rights reserved." );
+					const FText CopyrightInfo = NSLOCTEXT( "UnrealEd", "SplashScreen_CopyrightInfo", "Copyright \u00a9   Epic Games, Inc.   All rights reserved." );
 					StartSetSplashText( SplashTextType::CopyrightInfo, CopyrightInfo );
 				}
 			}
diff --git a/Engine/Source/Runtime/ApplicationCore/Private/Windows/WindowsPlatformSplash.cpp b/Engine/Source/Runtime/ApplicationCore/Private/Windows/WindowsPlatformSplash.cpp
--- a/Engine/Source/Runtime/ApplicationCore/Private/Windows/WindowsPlatformSplash.cpp
+++ b/Engine/Source/Runtime/ApplicationCore/Private/Windows/WindowsPlatformSplash.cpp
@@ -688,7 +688,7 @@ void FWindowsPlatformSplash::Show()
 
 				// Display copyright information in editor splash screen
 				{
-					const FString CopyrightInfo = NSLOCTEXT( "UnrealEd", "SplashScreen_CopyrightInfo", "Copyright \x00a9   Epic Games, Inc.   All rights reserved." ).ToString();
+					const FString CopyrightInfo = NSLOCTEXT( "UnrealEd", "SplashScreen_CopyrightInfo", "Copyright \u00a9   Epic Games, Inc.   All rights reserved." ).ToString();
 					StartSetSplashText( SplashTextType::CopyrightInfo, *CopyrightInfo );
 				}
 			}
