Graph Editor: Added the ability for derived node widgets to override the icon widget

diff --git a/Engine/Source/Editor/GraphEditor/Private/SGraphNode.cpp b/Engine/Source/Editor/GraphEditor/Private/SGraphNode.cpp
--- a/Engine/Source/Editor/GraphEditor/Private/SGraphNode.cpp
+++ b/Engine/Source/Editor/GraphEditor/Private/SGraphNode.cpp
@@ -857,14 +857,6 @@ void SGraphNode::UpdateGraphNode()
 
 	TSharedPtr<SNodeTitle> NodeTitle = SNew(SNodeTitle, GraphNode);
 
-	// Get node icon
-	IconColor = FLinearColor::White;
-	const FSlateBrush* IconBrush = nullptr;
-	if (GraphNode != NULL && GraphNode->ShowPaletteIconOnNode())
-	{
-		IconBrush = GraphNode->GetIconAndTint(IconColor).GetOptionalIcon();
-	}
-
 	TSharedRef<SOverlay> DefaultTitleAreaWidget =
 		SNew(SOverlay)
 		+SOverlay::Slot()
@@ -894,9 +886,7 @@ void SGraphNode::UpdateGraphNode()
 					.Padding(FMargin(0.f, 0.f, 4.f, 0.f))
 					.AutoWidth()
 					[
-						SNew(SImage)
-						.Image(IconBrush)
-						.ColorAndOpacity(this, &SGraphNode::GetNodeTitleIconColor)
+						CreateIconWidget()
 					]
 					+ SHorizontalBox::Slot()
 					[
@@ -1749,3 +1739,18 @@ void SGraphNode::InvalidateGraphData()
 		InvalidateGraphPinDataLambda(Pin);
 	}
 }
+
+TSharedRef<SWidget> SGraphNode::CreateIconWidget()
+{
+	// Get node icon
+	IconColor = FLinearColor::White;
+	const FSlateBrush* IconBrush = nullptr;
+	if (GraphNode != nullptr && GraphNode->ShowPaletteIconOnNode())
+	{
+		IconBrush = GraphNode->GetIconAndTint(IconColor).GetOptionalIcon();
+	}
+	
+	return SNew(SImage)
+		.Image(IconBrush)
+		.ColorAndOpacity(this, &SGraphNode::GetNodeTitleIconColor);
+}
diff --git a/Engine/Source/Editor/GraphEditor/Public/SGraphNode.h b/Engine/Source/Editor/GraphEditor/Public/SGraphNode.h
--- a/Engine/Source/Editor/GraphEditor/Public/SGraphNode.h
+++ b/Engine/Source/Editor/GraphEditor/Public/SGraphNode.h
@@ -434,6 +434,9 @@ protected:
 	/** Returns the widget to use for the enabled state of the node */
 	UE_API TSharedPtr<SWidget> GetEnabledStateWidget();
 
+	/** Returns the widget to use for the corner icon of the node */
+	UE_API virtual TSharedRef<SWidget> CreateIconWidget();
+
 protected:
 	/** Input pin widgets on this node */
 	TArray< TSharedRef<SGraphPin> > InputPins;
