Beginnings of a prototype of graph editor nodes that can 'snap' or 'click' together, facilitating organization and readability. This is disabled by default unless the bp.EnableNodeSnapping cvar is enabled. Uses two new bits in the EdGraphPin flag mask to canonize snapping state.

#rnx

diff --git a/Engine/Source/Editor/GraphEditor/Private/GraphEditorSnapping/NodeSnappingManager.cpp b/Engine/Source/Editor/GraphEditor/Private/GraphEditorSnapping/NodeSnappingManager.cpp
new file mode 100644
--- /dev/null
+++ b/Engine/Source/Editor/GraphEditor/Private/GraphEditorSnapping/NodeSnappingManager.cpp
@@ -0,0 +1,963 @@
+// Copyright Epic Games, Inc. All Rights Reserved.
+
+#include "GraphEditorSnapping/NodeSnappingManager.h"
+
+#include "EdGraphSchema_K2.h"
+#include "EdGraph/EdGraph.h"
+#include "Framework/Application/SlateApplication.h"
+#include "ScopedTransaction.h"
+#include "SGraphPanel.h"
+
+namespace UE::Private
+{
+
+TAutoConsoleVariable<bool> CVarBPEnableNodeSnapping(
+	TEXT("bp.EnableNodeSnapping"),
+	0,
+	TEXT("if not 0, node snapping is enabled"),
+	ECVF_Default
+);
+
+template<typename Predicate>
+UEdGraphPin* GetOnlyVisiblePinWithDirection( UEdGraphNode* Node, EEdGraphPinDirection PinDirection, Predicate Pred )
+{
+	// if predicate matches multiple visible pins then we return nullptr:
+	const TArray<UEdGraphPin*>& Pins = Node->GetAllPins();
+	UEdGraphPin* Result = nullptr;
+	for(UEdGraphPin* Pin : Pins)
+	{
+		if(Pin->bHidden)
+		{
+			continue;
+		}
+
+		if(Pin->Direction == PinDirection && Pred(Pin))
+		{
+			if(Result)
+			{
+				Result = nullptr;
+				break;
+			}
+			else
+			{
+				Result = Pin;
+			}
+		}
+
+	}
+	return Result;
+}
+
+UEdGraphPin* GetFirstAndOnlyVisiblePin(UEdGraphNode* ForNode)
+{
+	const TArray<UEdGraphPin*>& Pins = ForNode->GetAllPins();
+	UEdGraphPin* Result = nullptr;
+	for(UEdGraphPin* Pin : Pins)
+	{
+		if(!Pin->bHidden)
+		{
+			if(Result)
+			{
+				return nullptr;
+			}
+			Result = Pin;
+		}
+	}
+	return Result;
+}
+
+UEdGraphPin* GetFirstAndOnlyVisibleExecInPin(UEdGraphNode* OnNode)
+{
+	return GetOnlyVisiblePinWithDirection(
+		OnNode, 
+		EEdGraphPinDirection::EGPD_Input, 
+		[](UEdGraphPin* Pin) { return UEdGraphSchema_K2::IsExecPin(*Pin); }
+	);
+}
+
+UEdGraphPin* GetFirstAndOnlyVisibleExecOutPin(UEdGraphNode* OnNode)
+{
+	return GetOnlyVisiblePinWithDirection(
+		OnNode, 
+		EEdGraphPinDirection::EGPD_Output, 
+		[](UEdGraphPin* Pin) { return UEdGraphSchema_K2::IsExecPin(*Pin); }
+	);
+}
+
+UEdGraphPin* FindPinConnectedTo(UEdGraphNode* Node, UEdGraphPin* Pin)
+{
+	check(Node && Pin);
+	for(UEdGraphPin* LinkedPin : Pin->LinkedTo)
+	{
+		if(LinkedPin->GetOwningNode() == Node)
+		{
+			return LinkedPin;
+		}
+	}
+	return nullptr;
+}
+UEdGraphPin* CanBeConnectedTo(UEdGraphNode* Node, UEdGraphPin* Pin)
+{
+	check(Node && Pin);
+	if(Pin->LinkedTo.Num() > 0)
+	{
+		return nullptr;
+	}
+
+	const UEdGraphSchema* Schema = Node->GetSchema();
+
+	for(UEdGraphPin* PinOnNode : Node->Pins)
+	{		
+		if(PinOnNode->bHidden)
+		{
+			continue;
+		}
+
+		const FPinConnectionResponse Response = Schema->CanCreateConnection(
+			PinOnNode, Pin);
+		if( Response.Response == ECanCreateConnectionResponse::CONNECT_RESPONSE_MAKE )
+		{
+			return PinOnNode;
+		}
+	}
+
+	return nullptr;
+}
+
+bool HasAnyVisibleExecPins(UEdGraphNode* Node)
+{
+	for(UEdGraphPin* Pin : Node->Pins)
+	{
+		if(Pin->bHidden)
+		{
+			continue;
+		}
+		
+		if(UEdGraphSchema_K2::IsExecPin(*Pin))
+		{
+			return true;
+		}
+	}
+	return false;
+}
+//Slate<->Object helpers:
+UEdGraphNode* GetEdGraphNode(TSharedRef<SGraphNode> FromSGraphNode)
+{
+	return CastChecked<UEdGraphNode>(FromSGraphNode->GetObjectBeingDisplayed());
+}
+bool AreCompatibleDirection(TSharedRef<SGraphPin> Pin, UEdGraphPin* OtherPin)
+{
+	return Pin->GetDirection() != OtherPin->Direction;
+}
+
+bool ConnectAndSnap(UEdGraphPin* PinA, UEdGraphPin* PinB)
+{
+	check(PinA && PinB);
+
+	// @todo: if the pins are snapped to anything unsnap them - do i need symmetrical 
+	// pin snaps?
+
+	// if the pins are not connected then connect them:
+	if(!PinA->LinkedTo.Contains(PinB))
+	{
+		GetDefault<UEdGraphSchema_K2>()->TryCreateConnection(PinA, PinB);
+	}
+	
+	if(!PinA->LinkedTo.Contains(PinB))
+	{
+		return false;
+	}
+
+	// execution pins should snap in execution order - data pins
+	// should snap with consumer owning the producer:
+	UEdGraphPin* FirstPin = PinA;
+	UEdGraphPin* SecondPin = PinB;
+	if(UEdGraphSchema_K2::IsExecPin(*FirstPin))
+	{
+		if(FirstPin->Direction == EEdGraphPinDirection::EGPD_Input)
+		{
+			::Swap(FirstPin, SecondPin);
+		}
+	}
+	else
+	{
+		// @todo: may need to prefer to snap 'within' pins that are on nodes
+		// with execution pins
+		if(FirstPin->Direction == EEdGraphPinDirection::EGPD_Output)
+		{
+			::Swap(FirstPin, SecondPin);
+		}
+	}
+
+	if(FirstPin->GetSnappedChildPin() == SecondPin)
+	{
+		return true;
+	}
+
+	FirstPin->SetSnappedChild(SecondPin);
+	return true;
+}
+}
+
+FNodeSnappingManager::FNodeSnappingManager(TWeakPtr<SGraphPanel> InOwningPanel)
+	: OwningPanel(InOwningPanel)
+	, TotalMouseDelta(0.f)
+	, bIsActivelyDragging(false)
+{
+	check(InOwningPanel.Pin());
+}
+
+bool FNodeSnappingManager::IsDragging() const
+{
+	return bIsActivelyDragging;
+}
+
+bool FNodeSnappingManager::BeginDragWithSnap(const FGeometry& MyGeometry, const FPointerEvent& PointerEvent)
+{	
+	TSharedPtr<SGraphPanel> OwningPanelPinned = OwningPanel.Pin();
+	if(!OwningPanelPinned)
+	{
+		return false;
+	}
+
+	int32 NodeUnderMouseIndex = INDEX_NONE;
+	FVector2f MousePositionInNode = FVector2f::ZeroVector;
+	TSharedPtr<SGraphNode> NodeStartingDrag = FindNodeForMouseInteraction(MyGeometry, PointerEvent, MousePositionInNode);
+		
+	if( !NodeStartingDrag || !NodeStartingDrag->CanBeSelected(MousePositionInNode) )
+	{
+		return false;
+	}
+
+	OwningPanelPinned->NodeGrabOffset = MousePositionInNode;
+
+	TotalMouseDelta = 0.f;
+	bIsActivelyDragging = true;
+	
+	UEdGraphNode* ClickedNode = CastChecked<UEdGraphNode>(NodeStartingDrag->GetObjectBeingDisplayed());
+	// don't update selection here, the selection manager is updated on mouse up/mouse move
+	// so we have to project the selection. If selection is updated it will toggle in the
+	// mouse up logic:
+	TSet<TObjectPtr<UObject>> SelectedNodes;
+	if(	PointerEvent.IsControlDown() || PointerEvent.IsShiftDown() || 
+		OwningPanelPinned->SelectionManager.SelectedNodes.Contains(ClickedNode))
+	{
+		SelectedNodes = OwningPanelPinned->SelectionManager.SelectedNodes;
+	}
+	SelectedNodes.Add(ClickedNode);
+	// prune any nodes from the drag set if their outer node is also selected, we can move
+	// them by simply moving their inner - also override the root node with its outermost
+	// selection:
+	{
+		TSet<TObjectPtr<UObject>> SelectedNodesPruned;
+		for(TObjectPtr<UObject> Object : SelectedNodes)
+		{
+			UEdGraphNode* AsNode = CastChecked<UEdGraphNode>(Object);
+			bool bShouldBePruned = AsNode->IsSnappedToAny(SelectedNodes);
+			if(bShouldBePruned)
+			{
+				// don't prune exec level snaps - those are peers and must remain selected:
+				bShouldBePruned = !UE::Private::HasAnyVisibleExecPins(AsNode);
+			}
+			if(!bShouldBePruned)
+			{
+				SelectedNodesPruned.Add(Object);
+			}
+		}
+		SelectedNodes = MoveTemp(SelectedNodesPruned);
+
+		if(!SelectedNodes.Contains(ClickedNode))
+		{
+			// use the root selected node as the node starting drag:
+			UEdGraphNode* SelectionRoot = ClickedNode->GetRootSnappedParent(SelectedNodes);
+			check(SelectionRoot);
+			
+			TSharedRef<SNodePanel::SNode>* Widget = OwningPanelPinned->NodeToWidgetLookup.Find(SelectionRoot);
+			check(Widget != nullptr);
+			NodeStartingDrag = StaticCastSharedRef<SGraphNode>(*Widget);
+		}
+	}
+	
+	MoveNodeTransaction = MakeUnique<FScopedTransaction>(
+		SelectedNodes.Num() > 1 
+		? NSLOCTEXT("GraphEditor", "MoveNodesAction", "Move Nodes")
+		: NSLOCTEXT("GraphEditor", "MoveNodeAction", "Move Node"));
+	for(TObjectPtr<UObject> Node : SelectedNodes)
+	{
+		UEdGraphNode* AsEdGraphNode = CastChecked<UEdGraphNode>(Node);
+		ModifiedNodes.Add(AsEdGraphNode);
+		MovingNodes.Add(AsEdGraphNode);
+		Node->Modify();
+	}
+
+	PendingSnaps.SetNum(SelectedNodes.Num());
+	int32 Idx = 0;
+	for(TObjectPtr<UObject> Object : SelectedNodes)
+	{
+		FPendingSnap& MoveData = PendingSnaps[Idx];
+		
+		TSharedRef<SNodePanel::SNode>* Widget = OwningPanelPinned->NodeToWidgetLookup.Find(Object);
+		check(Widget != nullptr);
+		MoveData.NodeBeingDragged = StaticCastSharedRef<SGraphNode>(*Widget);
+
+		++Idx;
+	}
+	RootDragNode = StaticCastSharedRef<SGraphNode>(NodeStartingDrag.ToSharedRef());
+	OwningPanelPinned->NodeUnderMousePtr = RootDragNode;
+
+	return true;
+}
+
+bool FNodeSnappingManager::FinalizeNodeMovements()
+{
+	if(!IsDragging())
+	{
+		return false; // no customization - fall back to marquee selection or pin commit or whatever
+	}
+
+	CommitDrag();
+	return true;
+}
+
+void FNodeSnappingManager::UpdateDrag(const FGeometry& MyGeometry, const FPointerEvent& MouseEvent)
+{
+	ensure(bIsActivelyDragging);
+	TSharedPtr<SGraphPanel> OwningPanelPinned = OwningPanel.Pin();
+	if(!OwningPanelPinned)
+	{
+		// believed unreachable, hard to confirm with reference counts:
+		ResetDragManager();
+		return;
+	}
+
+	TSharedPtr<SGraphNode> NodeBeingDragged = RootDragNode;
+	if(!ensure(NodeBeingDragged))
+	{
+		return;
+	}
+	
+	const FVector2f CursorDelta = MouseEvent.GetCursorDelta();
+	TotalMouseDelta += CursorDelta.Size();
+
+	// Update the amount to pan panel
+	OwningPanelPinned->UpdateViewOffset(MyGeometry, MouseEvent.GetScreenSpacePosition());
+
+	const bool bCursorInDeadZone = TotalMouseDelta <= FSlateApplication::Get().GetDragTriggerDistance();
+	
+	// reset snap data for the nodes we're moving, we want to recalculate it:
+	for(UEdGraphNode* Node : MovingNodes)
+	{
+		// reset snapped state - a moving node
+		// should unsnap itself from any unselected node:
+		for(UEdGraphPin* Pin : Node->Pins)
+		{
+			if(Pin->GetSnappedChildNode() && UEdGraphSchema_K2::IsExecPin(*Pin) && !MovingNodes.Contains(Pin->GetSnappedChildNode()))
+			{
+				Pin->ClearSnappedChild();
+			}
+			else if(Pin->LinkedTo.Num() > 0 )
+			{
+				UEdGraphPin* OwningPin = Pin->LinkedTo[0];
+				UEdGraphNode* SnappedParentNode = OwningPin->GetSnappedChildNode();
+				if(MovingNodes.Contains(SnappedParentNode) && !MovingNodes.Contains(OwningPin->GetOwningNode()))
+				{
+					Pin->LinkedTo[0]->ClearSnappedChild();
+				}
+			}
+		}
+	}
+	
+	for(FPendingSnap& ExistingPendingSnap : PendingSnaps)
+	{
+		if(TryToSnapDataflow(ExistingPendingSnap))
+		{
+			continue;
+		}
+		else
+		{
+			TryToSnapControlflow(ExistingPendingSnap);
+		}
+	}
+	
+	// move the nodes, this must be done after snapping or we'll get
+	// incoherence in attached node position!
+	if ( NodeBeingDragged.IsValid() )
+	{
+		if ( !bCursorInDeadZone )
+		{
+			// Note, NodeGrabOffset() comes from the node itself, so it's already scaled correctly.
+			FVector2f AnchorNodeNewPos = OwningPanelPinned->PanelCoordToGraphCoord( MyGeometry.AbsoluteToLocal( MouseEvent.GetScreenSpacePosition() ) ) - OwningPanelPinned->NodeGrabOffset;
+
+			// Snap to grid
+			const float SnapSize = static_cast<float>(OwningPanelPinned->GetSnapGridSize());
+			AnchorNodeNewPos.X = SnapSize * FMath::RoundToFloat( AnchorNodeNewPos.X / SnapSize );
+			AnchorNodeNewPos.Y = SnapSize * FMath::RoundToFloat( AnchorNodeNewPos.Y / SnapSize );
+
+			// Dragging an unselected node automatically selects it.
+			OwningPanelPinned->SelectionManager.StartDraggingNode(NodeBeingDragged->GetObjectBeingDisplayed(), MouseEvent);
+
+			// Move all the selected nodes.
+			const FVector2f AnchorNodeOldPos = NodeBeingDragged->GetPosition2f();
+			const FVector2f DeltaPos = AnchorNodeNewPos - AnchorNodeOldPos;
+			SNodePanel::SNode::FNodeSet NodeFilter;
+
+			for (const FPendingSnap& DraggingNode : PendingSnaps)
+			{
+				DraggingNode.NodeBeingDragged->MoveTo(
+					DraggingNode.NodeBeingDragged->GetPosition2f() + DeltaPos,
+					NodeFilter,
+					false);
+			}
+		}
+	}
+	
+	// first stage of refresh - move nodes back to root level:
+	for(UEdGraphNode* Node : ModifiedNodes)
+	{
+		TSharedRef<SNodePanel::SNode>* Widget = OwningPanelPinned->NodeToWidgetLookup.Find(Node);
+		check(Widget != nullptr);
+		TSharedRef<SGraphNode> NodeTyped = StaticCastSharedRef<SGraphNode>(*Widget);
+		ClearAllSnaps(NodeTyped);
+	}
+	
+	// with the data model updated, update snap state for all display nodes:
+	for(UEdGraphNode* Node : ModifiedNodes)
+	{
+		TSharedRef<SNodePanel::SNode>* Widget = OwningPanelPinned->NodeToWidgetLookup.Find(Node);
+		check(Widget != nullptr);
+		TSharedRef<SGraphNode> NodeTyped = StaticCastSharedRef<SGraphNode>(*Widget);
+		UpdateSnapsFromModel(NodeTyped);
+	}
+}
+
+void FNodeSnappingManager::HandleGraphChanged()
+{
+	TSharedRef<SGraphPanel> OwningPanelPinned = OwningPanel.Pin().ToSharedRef();
+	UEdGraph* DisplayedGraph = OwningPanelPinned->GraphObj;
+	if(!DisplayedGraph)
+	{
+		return;
+	}
+	
+	// first stage of refresh - move nodes back to root level
+	// this makes sure we are in a consistent state if a partial
+	// refresh happens
+	for(UEdGraphNode* Node : DisplayedGraph->Nodes)
+	{
+		TSharedRef<SNodePanel::SNode>* Widget = OwningPanelPinned->NodeToWidgetLookup.Find(Node);
+		check(Widget != nullptr);
+		TSharedRef<SGraphNode> NodeTyped = StaticCastSharedRef<SGraphNode>(*Widget);
+		ClearAllSnaps(NodeTyped);
+	}
+
+	for(UEdGraphNode* Node : DisplayedGraph->Nodes)
+	{
+		TSharedRef<SNodePanel::SNode>* Widget = OwningPanelPinned->NodeToWidgetLookup.Find(Node);
+		check(Widget != nullptr);
+		TSharedRef<SGraphNode> NodeTyped = StaticCastSharedRef<SGraphNode>(*Widget);
+		UpdateSnapsFromModel(NodeTyped);
+	}
+}
+
+void FNodeSnappingManager::ForEachNode(TFunctionRef<void(const TSharedRef<SNodePanel::SNode>&)> Ftor) const
+{
+	TSharedRef<SGraphPanel> OwningPanelPinned = OwningPanel.Pin().ToSharedRef();
+	
+	const auto SGraphNodeAdapter = [&Ftor](const TSharedRef<SGraphNode>& GraphNode)
+	{
+		Ftor(GraphNode);
+	};
+
+	for (int32 ChildIndex = 0; ChildIndex < OwningPanelPinned->Children.Num(); ++ChildIndex)
+	{
+		TSharedRef<SGraphNode> AsSgraphNode = StaticCastSharedRef<SGraphNode>(OwningPanelPinned->Children[ChildIndex]);
+		AsSgraphNode->ForEachNodeIncludingChildren(SGraphNodeAdapter);
+	}
+}
+
+bool FNodeSnappingManager::IsNodeSnappingEnabled()
+{
+	return UE::Private::CVarBPEnableNodeSnapping->GetBool();
+}
+
+void FNodeSnappingManager::CommitDrag()
+{
+	// no finalization work to do atm - we scoped our transaction
+	// correctly so we don't have to do the whacky 'restore to original 
+	// node positions/call modify' stuff that the normal code path did
+	ResetDragManager();
+}
+
+void FNodeSnappingManager::ResetDragManager()
+{
+	MoveNodeTransaction.Reset();
+	PendingSnaps.Reset();
+	RootDragNode = nullptr;
+	MovingNodes.Reset();
+	ModifiedNodes.Reset();
+	TotalMouseDelta = 0.f;
+	bIsActivelyDragging = false;
+}
+
+const UEdGraphSchema* FNodeSnappingManager::GetSchema() const
+{
+	TSharedPtr<SGraphPanel> OwningPanelPinned = OwningPanel.Pin();
+	check(OwningPanelPinned);
+	return OwningPanelPinned->GraphObj->GetSchema();
+}
+
+bool FNodeSnappingManager::TryToSnapDataflow(FPendingSnap& PendingSnap)
+{
+	TSharedPtr<SGraphNode> NodeBeingSnapped = PendingSnap.NodeBeingDragged;
+	TSharedPtr<SGraphPanel> OwningPanelPinned = OwningPanel.Pin();
+	check(OwningPanelPinned);
+	UEdGraphPin* PinToMatchForSnap = UE::Private::GetFirstAndOnlyVisiblePin(NodeBeingSnapped->GetNodeObj());
+	if(PinToMatchForSnap && PinToMatchForSnap->LinkedTo.Num() > 1)
+	{
+		PinToMatchForSnap = nullptr;
+	}
+
+	TSharedPtr<SGraphNode> DestGraphNode = FindNodeForSnap(NodeBeingSnapped);
+	if(!PinToMatchForSnap || !DestGraphNode)
+	{
+		return false;
+	}
+
+	UEdGraphPin* PinToSnapTo = nullptr;
+	if(PinToMatchForSnap->LinkedTo.Num() == 1)
+	{
+		PinToSnapTo = PinToMatchForSnap->LinkedTo[0];
+	}
+
+	TSharedPtr<SGraphPin> DestPin = DestGraphNode->GetHoveredPin(DestGraphNode->GetCachedGeometry(), OwningPanelPinned->LastPointerEvent);
+
+	bool bIsDestPinValid = DestPin.IsValid();
+	if(PinToSnapTo && bIsDestPinValid)
+	{
+		bIsDestPinValid = (DestPin->GetPinObj() == PinToSnapTo);
+	}
+
+	if(!bIsDestPinValid )
+	{
+		return false;
+	}
+	const FPinConnectionResponse Response = GetSchema()->CanCreateConnection(
+		DestPin->GetPinObj(), PinToMatchForSnap);
+	if( Response.Response == ECanCreateConnectionResponse::CONNECT_RESPONSE_DISALLOW )
+	{
+		return false;
+	}
+
+	return TrySnapTo(DestPin.ToSharedRef(), PendingSnap, PinToMatchForSnap);
+}
+
+bool FNodeSnappingManager::TryToSnapControlflow(FPendingSnap& PendingSnap)
+{
+	TSharedPtr<SGraphPanel> OwningPanelPinned = OwningPanel.Pin();
+	check(OwningPanelPinned);
+	TSharedPtr<SGraphNode> NodeBeingSnapped = PendingSnap.NodeBeingDragged;
+	check(NodeBeingSnapped);
+	UEdGraphNode* EdGraphNodeBeingSnapped = NodeBeingSnapped->GetNodeObj();
+
+	UEdGraphPin* ExecutionPinIn = UE::Private::GetFirstAndOnlyVisibleExecInPin(EdGraphNodeBeingSnapped);
+	UEdGraphPin* ExecutionPinOut = UE::Private::GetFirstAndOnlyVisibleExecOutPin(EdGraphNodeBeingSnapped);
+
+	if( !ExecutionPinOut && !ExecutionPinIn )
+	{
+		return false; // no control flow to snap - bail
+	}
+
+	TSharedPtr<SGraphPin> ExecPinOutToSnapTo = FindBestControlFlowForSnap(NodeBeingSnapped, ExecutionPinIn);
+	TSharedPtr<SGraphPin> ExecPinInToSnapTo = FindBestControlFlowForSnap(NodeBeingSnapped, ExecutionPinOut);
+	TSharedPtr<SGraphPin> PreferredPin = GetPreferredControlFlowSnapTarget(ExecPinOutToSnapTo, ExecPinInToSnapTo);
+	UEdGraphPin* PinOnDraggingNode = PreferredPin == ExecPinOutToSnapTo ? ExecutionPinIn : ExecutionPinOut;
+	if(!PreferredPin)
+	{
+		return false;
+	}
+
+	// this doesn't snap hierarchically, instead the pending snap is merely positioned
+	// and restyled such that it looks as if it is connected to its peer. If either
+	// node is moved the nodes will restyle:
+	SnapControlFlowPeer(PendingSnap, PreferredPin.ToSharedRef(), PinOnDraggingNode);
+	return true;
+}
+
+bool FNodeSnappingManager::TrySnapTo(TSharedRef<SGraphPin> PinToSnapTo, FPendingSnap& SnapState, UEdGraphPin* PinOnMovingNode)
+{
+	UEdGraphPin* OtherPin = PinToSnapTo->GetPinObj();
+	check(OtherPin && OtherPin != PinOnMovingNode);
+	ModifySafe(OtherPin);
+	UE::Private::ConnectAndSnap(PinOnMovingNode, OtherPin);
+	return true;
+}
+
+void FNodeSnappingManager::SnapControlFlowPeer(FPendingSnap& SnapState, TSharedRef<SGraphPin> PinToSnapTo, UEdGraphPin* PinOnMovingNode)
+{
+	UEdGraphPin* OtherPin = PinToSnapTo->GetPinObj();
+	check(OtherPin && OtherPin != PinOnMovingNode);
+
+	TSharedRef<SGraphNode> NodeBeingDragged = SnapState.NodeBeingDragged.ToSharedRef();
+	TSharedRef<SGraphNode> NodeDroppedOn = PinToSnapTo->OwnerNodePtr.Pin().ToSharedRef();
+
+	TSharedRef<SGraphNode> StartingNode = (PinToSnapTo->GetDirection() ==  EEdGraphPinDirection::EGPD_Input) ? 
+		NodeBeingDragged : NodeDroppedOn;
+		
+	if(!StartingNode->ControlFlowSequence)
+	{
+		return; // @todo: overcome SGraphNode copy paste - we can't snap to this
+	}
+	
+	// update the data model.. this is all we should really do here - but I have made some mistakes
+	ModifySafe(OtherPin);
+	UE::Private::ConnectAndSnap(PinOnMovingNode, OtherPin);
+
+	SnapState.PinSnappedTo = PinToSnapTo;
+}
+
+void FNodeSnappingManager::ClearAllSnaps(TSharedRef<SGraphNode> OnNode)
+{
+	ClearControlFlowSnap(OnNode);
+	ClearDataSnap(OnNode);
+}
+
+void FNodeSnappingManager::ClearControlFlowSnap(TSharedRef<SGraphNode> OnNode)
+{
+	const auto IsAlreadySnapped = [](TSharedRef<SGraphNode> NodeSnappingTo, TSharedRef<SGraphNode> NodeToSnap)
+	{
+		const FChildren* Children = NodeSnappingTo->ControlFlowSequence->GetChildren();
+
+		for (int32 I = 0; I < Children->Num(); ++I)
+		{
+			TSharedRef<const SWidget> ChildWidget = Children->GetChildAt(I);
+			if (ChildWidget == NodeToSnap)
+			{
+				return true;
+			}
+		}
+		return false;
+	};
+
+	TSharedPtr<SGraphNode> NodeParent = OnNode->SnapParent.Pin();
+	if(NodeParent)
+	{
+		TSharedRef<SGraphNode> NodeParentRef = NodeParent.ToSharedRef();
+		if(NodeParentRef->ControlFlowSequence && IsAlreadySnapped(NodeParentRef, OnNode))
+		{
+			NodeParentRef->ControlFlowSequence->RemoveSlot(OnNode);
+			NodeParentRef->GetOwnerPanel()->Children.Add(OnNode);
+			OnNode->SnapParent = nullptr;
+		}
+	}
+
+	if(OnNode->ControlFlowSequence)
+	{
+		// unsnap subsequent children:
+		FChildren* Children = OnNode->ControlFlowSequence->GetChildren();
+
+		for (int32 I = Children->Num() - 1; I >= 1 && Children->Num(); --I)
+		{
+			TSharedRef<SGraphNode> ChildRef = StaticCastSharedRef<SGraphNode>(Children->GetChildAt(I));
+			OnNode->GetOwnerPanel()->Children.Add(ChildRef);
+			OnNode->ControlFlowSequence->RemoveSlot(ChildRef);
+			ChildRef->SnapParent = nullptr;
+		}
+	}
+}
+
+void FNodeSnappingManager::ClearDataSnap(TSharedRef<SGraphNode> OnNode)
+{
+	// clear snap to our parent node:
+	TSharedRef<SGraphPanel> OwningPanelPinned = OnNode->GetOwnerPanel().ToSharedRef();
+	TSharedPtr<SGraphNode> NodeParent = OnNode->SnapParent.Pin();
+	if(NodeParent)
+	{
+		for(TSharedRef<SGraphPin> Pin : NodeParent->InputPins)
+		{
+			TSharedPtr<SGraphNode> SnappedNode = Pin->SnappedNode.Pin();
+			if(SnappedNode == OnNode)
+			{
+				Pin->ClearSnappedSGraphNode();
+			}
+		}
+		for(TSharedRef<SGraphPin> Pin : NodeParent->OutputPins)
+		{
+			TSharedPtr<SGraphNode> SnappedNode = Pin->SnappedNode.Pin();
+			if(SnappedNode== OnNode)
+			{
+				Pin->ClearSnappedSGraphNode();
+			}
+		}
+		OwningPanelPinned->Children.Add(OnNode);
+		OnNode->SnapParent = nullptr;
+	}
+
+	// clear snap to our children node:
+	for(TSharedRef<SGraphPin> Pin : OnNode->InputPins)
+	{
+		TSharedPtr<SGraphNode> SnappedNode = Pin->SnappedNode.Pin();
+		if(SnappedNode)
+		{
+			Pin->ClearSnappedSGraphNode();
+			OwningPanelPinned->Children.Add(SnappedNode.ToSharedRef());
+			SnappedNode->SnapParent = nullptr;
+		}
+	}
+	for(TSharedRef<SGraphPin> Pin : OnNode->OutputPins)
+	{
+		TSharedPtr<SGraphNode> SnappedNode = Pin->SnappedNode.Pin();
+		if(SnappedNode)
+		{
+			Pin->ClearSnappedSGraphNode();
+			OwningPanelPinned->Children.Add(SnappedNode.ToSharedRef());
+			SnappedNode->SnapParent = nullptr;
+		}
+	}
+}
+
+void FNodeSnappingManager::UpdateSnapsFromModel(TSharedRef<SGraphNode> OnNode)
+{
+	TSharedRef<SGraphPanel> OwningPanelPinned = OnNode->GetOwnerPanel().ToSharedRef();
+	const UEdGraphNode* NodeObj = OnNode->GetNodeObj();
+	// if we have a snapped input pin to after it:
+	for(UEdGraphPin* Pin : NodeObj->Pins)
+	{
+		// if no snap, nothing to do:
+		FSnappedEdGraphPin SnappedPins = Pin->GetSnapDescription();
+		if(SnappedPins == FSnappedEdGraphPin())
+		{
+			continue;
+		}
+		
+		// continue if the ui is already snapped - this handles cycles in the graph - 
+		// e.g. two nodes are being modified at once - multiple snaps reworking, etc:
+		TSharedRef<SGraphNode> ParentNode = GetSGraphNode(OwningPanelPinned, SnappedPins.ParentPin->GetOwningNode());
+		TSharedRef<SGraphNode> ChildNode = GetSGraphNode(OwningPanelPinned, SnappedPins.ChildPin->GetOwningNode());
+		if(ChildNode->SnapParent.Pin() != nullptr)
+		{
+			ensure(ChildNode->SnapParent == ParentNode);
+			continue;
+		}
+
+		const bool bPinIsParent = (SnappedPins.ParentPin == Pin);
+		const bool bPinIsChild = !bPinIsParent;
+
+		if(UEdGraphSchema_K2::IsExecPin(*Pin))
+		{
+			// control flow snap - @todo: move to a function
+			// if the node is being moved, snap it into position before the unmoving node:
+			if(MovingNodes.Contains(NodeObj) && !MovingNodes.Contains(Pin->GetSnappedChildNode()) && bPinIsParent)
+			{
+				// set the moving node into the correct position:
+				FVector2f NodePosition = Pin->GetSnappedChildNode()->GetPosition();
+				NodePosition -= FVector2f(0.f, OnNode->GetDesiredSizeForMarquee2f().Y); // does my marquee size grow? need to cache it? 😷 back to vsnap objects..
+				SNodePanel::SNode::FNodeSet NodeFilter;
+				OnNode->MoveTo(NodePosition, NodeFilter, false);
+			}
+			else if(MovingNodes.Contains(NodeObj) && bPinIsChild)
+			{
+				// moving node needs to be moved beneath parent, so that selection
+				// etc can occur in the corrrect position - subtle problems will
+				// be created if the node is snapped but not in the right position
+				// (e.g. failure to select pins because their notional position
+				// is different than their painted position):
+				FVector2f NodePosition = ParentNode->GetPosition2f();
+				NodePosition += FVector2f(0.f, ParentNode->GetDesiredSizeForMarquee2f().Y);
+				SNodePanel::SNode::FNodeSet NodeFilter;
+				ChildNode->MoveTo(NodePosition, NodeFilter, false);
+			}
+
+			ParentNode->ControlFlowSequence->AddSlot()
+				.AutoHeight()
+				[
+					ChildNode
+				];
+			ChildNode->SnapParent = ParentNode;
+			
+			OwningPanelPinned->Children.Remove(ChildNode);
+			OwningPanelPinned->VisibleChildren.Remove(ChildNode);
+		}
+		else
+		{
+			// data snap - @todo: move to a function
+			TSharedRef<SGraphPin> PinToSnapTo = ParentNode->FindWidgetForPin(SnappedPins.ParentPin).ToSharedRef();
+			if(PinToSnapTo->SnapSGraphNode(ChildNode))
+			{
+				ChildNode->SnapParent = ParentNode;
+				OwningPanelPinned->Children.Remove(ChildNode);
+				OwningPanelPinned->VisibleChildren.Remove(ChildNode);
+			}
+		}
+	}
+}
+
+
+TSharedPtr<SGraphNode> FNodeSnappingManager::FindNodeForMouseInteraction( const FGeometry& MyGeometry, const FPointerEvent& PointerEvent, FVector2f& OutPositionInNode ) const
+{
+	TSharedPtr<SGraphPanel> OwningPanelPinned = OwningPanel.Pin();
+	if(!OwningPanelPinned)
+	{
+		return nullptr;
+	}
+	FVector2f ArrangedSpacePosition = PointerEvent.GetScreenSpacePosition();
+	FArrangedChildren ArrangedChildren(EVisibility::Visible);
+	OwningPanelPinned->ArrangeChildNodes(MyGeometry, ArrangedChildren);
+
+	// loop over children, preferring nodes at leaves of snap hierarchy to roots:
+	const int32 NumChildren = ArrangedChildren.Num();
+	for( int32 ChildIndex=NumChildren-1; ChildIndex >= 0; --ChildIndex )
+	{
+		const FArrangedWidget& Candidate = ArrangedChildren[ChildIndex];
+		const bool bCandidateUnderCursor = 
+			// Candidate is physically under the cursor
+			Candidate.Geometry.IsUnderLocation( ArrangedSpacePosition );
+		
+		if (bCandidateUnderCursor)
+		{
+			OutPositionInNode = Candidate.Geometry.AbsoluteToLocal(ArrangedSpacePosition);
+			TSharedRef<SGraphNode> NodeWidgetUnderMouse = 
+				StaticCastSharedRef<SGraphNode>( Candidate.Widget );
+
+			const bool bIsSelected = (PointerEvent.IsControlDown() || PointerEvent.IsShiftDown()) && 
+				OwningPanelPinned->SelectionManager.SelectedNodes.Contains(NodeWidgetUnderMouse->GetObjectBeingDisplayed());
+
+			if(!bIsSelected)
+			{
+				TSharedPtr<SGraphNode> InnerNode = NodeWidgetUnderMouse->GetSnappedChildForPosition(
+					Candidate.Geometry, ArrangedSpacePosition, OutPositionInNode);
+				if(InnerNode)
+				{
+					NodeWidgetUnderMouse = InnerNode.ToSharedRef();
+				}
+			}
+
+			return NodeWidgetUnderMouse;
+		}
+	}
+
+	return nullptr;
+}
+
+TSharedPtr<SGraphNode> FNodeSnappingManager::FindNodeForSnap(TSharedPtr<SGraphNode> NodeBeingMoved)
+{
+	TSharedPtr<SGraphPanel> OwningPanelPinned = OwningPanel.Pin();
+	if(!OwningPanelPinned)
+	{
+		return nullptr;
+	}
+
+	const FGeometry& MyGeometry = OwningPanelPinned->LastPointerGeometry;
+	const FPointerEvent& MouseEvent = OwningPanelPinned->LastPointerEvent;
+
+	FArrangedChildren ArrangedNodes(EVisibility::Visible);
+	
+	for (int32 ChildIndex = 0; ChildIndex < OwningPanelPinned->Children.Num(); ++ChildIndex)
+	{
+		const TSharedRef<SNodePanel::SNode>& SomeChild = OwningPanelPinned->Children[ChildIndex];
+
+		FArrangedWidget ArrangedWidget = MyGeometry.MakeChild(
+			SomeChild, 
+			SomeChild->GetPosition2f() - OwningPanelPinned->ViewOffset, 
+			SomeChild->GetDesiredSize(), 
+			OwningPanelPinned->GetZoomAmount()
+		);
+		ArrangedNodes.AddWidget( ArrangedWidget );
+	}
+
+	const auto FindChildrenUnderMouse = [](const FArrangedChildren& Children, const FPointerEvent& PointerEvent) -> FArrangedChildren
+	{
+		FVector2f AbsoluteCursorLocation = PointerEvent.GetScreenSpacePosition();
+		FArrangedChildren WidgetsUnderMouse(EVisibility::Visible);
+		const int32 NumChildren = Children.Num();
+		for( int32 ChildIndex=NumChildren-1; ChildIndex >= 0; --ChildIndex )
+		{
+			const FArrangedWidget& Candidate = Children[ChildIndex];
+			const bool bCandidateUnderCursor = 
+				// Candidate is physically under the cursor
+				Candidate.Geometry.IsUnderLocation( AbsoluteCursorLocation );
+	
+			if (bCandidateUnderCursor)
+			{
+				WidgetsUnderMouse.AddWidget(Candidate);
+			}
+		}
+		return WidgetsUnderMouse;
+	};
+	
+	TSharedPtr<SGraphNode> GraphNode;
+	FArrangedChildren WidgetsUnderMouse = FindChildrenUnderMouse(ArrangedNodes, MouseEvent);
+	for(const FArrangedWidget& ArrangedWidget : WidgetsUnderMouse.GetInternalArray() )
+	{
+		TSharedPtr<SGraphNode> PotentialGraphNode = StaticCastSharedRef<SGraphNode>(ArrangedWidget.Widget);
+		FVector2f Discard;
+		TSharedPtr<SGraphNode> GraphSubNode =  PotentialGraphNode->GetSnappedChildForPosition(ArrangedWidget.Geometry, MouseEvent.GetScreenSpacePosition(), Discard);
+		PotentialGraphNode = GraphSubNode.IsValid() ? GraphSubNode.ToSharedRef() : PotentialGraphNode;
+		if(PotentialGraphNode != NodeBeingMoved)
+		{
+			GraphNode = PotentialGraphNode;
+			//break;
+		}
+	}
+
+	return GraphNode;
+}
+
+TSharedPtr<SGraphPin> FNodeSnappingManager::FindBestControlFlowForSnap(TSharedPtr<SGraphNode> NodeBeingMoved, UEdGraphPin* Pin)
+{
+	if(!Pin)
+	{
+		return nullptr;
+	}
+
+	TSharedPtr<SGraphNode> Node = FindNodeForSnap(NodeBeingMoved);
+	if(!Node)
+	{
+		return nullptr;
+	}
+
+	// Lets snap to the sgraph node if:
+	//  1. we're already attached
+	// or
+	//  2. we're not attached to anything and the node's execution pin also isn't: (@todo)
+	UEdGraphNode* EdGraphNodeBeingMoved = UE::Private::GetEdGraphNode(Node.ToSharedRef());
+	UEdGraphPin* PinToSnap = UE::Private::FindPinConnectedTo(EdGraphNodeBeingMoved, Pin);
+	if(!PinToSnap)
+	{
+		return nullptr;
+	}
+	
+	return Node->FindWidgetForPin(PinToSnap);
+}
+
+TSharedPtr<SGraphPin> FNodeSnappingManager::GetPreferredControlFlowSnapTarget(TSharedPtr<SGraphPin> A, TSharedPtr<SGraphPin> B) const
+{
+	if(!A)
+	{
+		return B;
+	}
+	else if(!B)
+	{
+		return A;
+	}
+
+	// @todo: tie break.. closest node?
+	return nullptr;
+}
+
+TSharedRef<SGraphNode> FNodeSnappingManager::GetSGraphNode(TSharedRef<SGraphPanel> FromPanel, const UEdGraphNode* Node)
+{
+	const TSharedRef<SNodePanel::SNode>* Widget = FromPanel->NodeToWidgetLookup.Find(Node);
+	check(Widget != nullptr);
+	return StaticCastSharedRef<SGraphNode>(*Widget);
+}
+
+void FNodeSnappingManager::ModifySafe(UEdGraphPin* Pin)
+{
+	UEdGraphNode* Node = Pin->GetOwningNode();
+	if(!ModifiedNodes.Contains(Node))
+	{
+		ModifiedNodes.Add(Node);
+		Node->Modify();
+	}
+}
+
diff --git a/Engine/Source/Editor/GraphEditor/Private/GraphEditorSnapping/NodeSnappingManager.h b/Engine/Source/Editor/GraphEditor/Private/GraphEditorSnapping/NodeSnappingManager.h
new file mode 100644
--- /dev/null
+++ b/Engine/Source/Editor/GraphEditor/Private/GraphEditorSnapping/NodeSnappingManager.h
@@ -0,0 +1,62 @@
+// Copyright Epic Games, Inc. All Rights Reserved.
+
+#pragma once
+
+#include "SGraphPin.h"
+#include "SGraphNode.h"
+
+struct FNodeSnappingManager
+{
+	FNodeSnappingManager(TWeakPtr<SGraphPanel> InOwningPanel);
+	bool IsDragging() const;
+	bool BeginDragWithSnap(const FGeometry& MyGeometry, const FPointerEvent& PointerEvent);
+	bool FinalizeNodeMovements();
+	void UpdateDrag(const FGeometry& MyGeometry, const FPointerEvent& MouseEvent);
+	void HandleGraphChanged();
+	void ForEachNode(TFunctionRef<void(const TSharedRef<SNodePanel::SNode>&)> Ftor) const;
+
+	static bool IsNodeSnappingEnabled();
+private:
+	struct FPendingSnap
+	{
+		TSharedPtr<SGraphNode> NodeBeingDragged;
+		TSharedPtr<SGraphPin> PinSnappedTo;
+
+		bool operator==(const FPendingSnap& RHS) const = default;
+	};
+
+	void CommitDrag();
+	void ResetDragManager();
+	const UEdGraphSchema* GetSchema() const;
+	bool TryToSnapDataflow(FPendingSnap& PendingSnap);
+	bool TryToSnapControlflow(FPendingSnap& PendingSnap);
+	bool TrySnapTo(TSharedRef<SGraphPin> PinToSnapTo, FPendingSnap& SnapState, UEdGraphPin* PinOnMovingNode);
+	void SnapControlFlowPeer(FPendingSnap& SnapState, TSharedRef<SGraphPin> PinToSnapTo, UEdGraphPin* PinOnMovingNode);
+	
+	void ClearAllSnaps(TSharedRef<SGraphNode> OnNode);
+	void ClearControlFlowSnap(TSharedRef<SGraphNode> OnNode);
+	void ClearDataSnap(TSharedRef<SGraphNode> OnNode);
+	void UpdateSnapsFromModel(TSharedRef<SGraphNode> OnNode);
+	static TArray<TSharedRef<SGraphNode>> GetControlFlowInputNodes(TSharedRef<SGraphNode> ForNode);
+	static TArray<TSharedRef<SGraphNode>> GetControlFlowOutputNodes(TSharedRef<SGraphNode> ForNode);
+	
+	TSharedPtr<SGraphNode> FindNodeForMouseInteraction( const FGeometry& MyGeometry, const FPointerEvent& PointerEvent, FVector2f& OutPositionInNode ) const;
+	TSharedPtr<SGraphNode> FindNodeForSnap(TSharedPtr<SGraphNode> NodeBeingMoved);
+	TSharedPtr<SGraphPin> FindBestControlFlowForSnap(TSharedPtr<SGraphNode> NodeBeingMoved, UEdGraphPin* Pin);
+	TSharedPtr<SGraphPin> GetPreferredControlFlowSnapTarget(TSharedPtr<SGraphPin> A, TSharedPtr<SGraphPin> B) const;
+	static TSharedRef<SGraphNode> GetSGraphNode(TSharedRef<SGraphPanel> FromPanel, const UEdGraphNode* Node);
+	void ModifySafe(UEdGraphPin* Pin);
+
+	TWeakPtr<SGraphPanel> OwningPanel; // must be weak, otherwise we'll have a cycle, never meant to be null
+
+	TUniquePtr<FScopedTransaction> MoveNodeTransaction;
+	// @todo: rename
+	TArray<FPendingSnap> PendingSnaps;
+	TSharedPtr<SGraphNode> RootDragNode; // actual node that is being dragged, used to calculate deltas
+	TSet<UEdGraphNode*> MovingNodes;
+	TSet<UEdGraphNode*> ModifiedNodes; // track modification to make sure we don't modify after mutating
+
+	float TotalMouseDelta;
+
+	bool bIsActivelyDragging;
+};
\ No newline at end of file
diff --git a/Engine/Source/Editor/GraphEditor/Private/SGraphNode.cpp b/Engine/Source/Editor/GraphEditor/Private/SGraphNode.cpp
--- a/Engine/Source/Editor/GraphEditor/Private/SGraphNode.cpp
+++ b/Engine/Source/Editor/GraphEditor/Private/SGraphNode.cpp
@@ -596,6 +596,12 @@ FVector2f SGraphNode::GetPosition2f() const
 	return GraphNode->GetPosition();
 }
 
+FVector2f SGraphNode::GetDesiredSizeForMarquee2f() const
+{
+	return ControlFlowSequence ? UE::Slate::CastToVector2f(ControlFlowSequence->GetChildren()->GetChildAt(0)->GetCachedGeometry().GetLocalSize())
+		: SNodePanel::SNode::GetDesiredSizeForMarquee2f();
+}
+
 FString SGraphNode::GetEditableNodeTitle() const
 {
 	if (GraphNode != nullptr)
@@ -853,7 +859,6 @@ void SGraphNode::UpdateGraphNode()
 	//            |       |      | T (>) |
 	//            |_______|______|_______|
 	//
-	TSharedPtr<SVerticalBox> MainVerticalBox;
 	SetupErrorReporting();
 
 	TSharedPtr<SNodeTitle> NodeTitle = SNew(SNodeTitle, GraphNode);
@@ -1005,10 +1010,10 @@ void SGraphNode::UpdateGraphNode()
 
 
 	this->GetOrAddSlot( ENodeZone::Center )
-		.HAlign(HAlign_Center)
+		.HAlign(HAlign_Fill)
 		.VAlign(VAlign_Center)
 		[
-			SAssignNew(MainVerticalBox, SVerticalBox)
+			SAssignNew(ControlFlowSequence, SVerticalBox)
 			+SVerticalBox::Slot()
 			.AutoHeight()
 			[
@@ -1062,7 +1067,7 @@ void SGraphNode::UpdateGraphNode()
 			];
 	}
 
-	CreateBelowWidgetControls(MainVerticalBox);
+	CreateBelowWidgetControls(ControlFlowSequence);
 	CreatePinWidgets();
 	CreateInputSideAddButton(LeftNodeBox);
 	CreateOutputSideAddButton(RightNodeBox);
@@ -1755,3 +1760,241 @@ TSharedRef<SWidget> SGraphNode::CreateIconWidget()
 		.Image(IconBrush)
 		.ColorAndOpacity(this, &SGraphNode::GetNodeTitleIconColor);
 }
+
+TArray<FArrangedWidget> SGraphNode::GetArrangedSnappedChildNodes(const FGeometry& Geometry)
+{
+	TArray<FArrangedWidget> FinalResult;
+	TSet< TSharedRef<SWidget> > MyPins;
+	{
+		GetPins( MyPins );
+	}
+	TMap<TSharedRef<SWidget>, FArrangedWidget> Result;
+	FindChildGeometries(Geometry, MyPins, Result);
+
+	// walk each pin looking for a snapped child node, if found add it to final result:
+	for(const TPair<TSharedRef<SWidget>, FArrangedWidget>& ArrangedPin : Result)
+	{
+		TSharedRef<SGraphPin> AsPin = StaticCastSharedRef<SGraphPin>(ArrangedPin.Key);
+		TOptional<FArrangedWidget> ChildNode = AsPin->GetArrangedSnappedChildNode(ArrangedPin.Value.Geometry);
+		if(ChildNode.IsSet())
+		{
+			FinalResult.Emplace(MoveTemp(ChildNode.GetValue()));
+		}
+	}
+
+	return FinalResult;
+}
+
+TSharedPtr<SGraphNode> SGraphNode::GetSnappedChildForPosition(const FGeometry& Geometry, const FVector2f& Position, FVector2f& OutPositionInNode )
+{
+	// check for control flow sequence child:
+	if(ControlFlowSequence)
+	{
+		FChildren* ControlFlowChildren = ControlFlowSequence->GetChildren();
+		if(ControlFlowChildren->Num() > 1)
+		{
+			TSet< TSharedRef<SWidget> > SubNodes;
+			for(int32 I = 1; I < ControlFlowChildren->Num(); ++I)
+			{
+				SubNodes.Add( ControlFlowChildren->GetChildAt(I) );
+			}
+			TMap<TSharedRef<SWidget>, FArrangedWidget> Result;
+			FindChildGeometries(Geometry, SubNodes, Result);
+		
+			for( const TPair<TSharedRef<SWidget>, FArrangedWidget>& SubNode : Result)
+			{
+				if(SubNode.Value.Geometry.IsUnderLocation( Position ))
+				{
+					return StaticCastSharedRef<SGraphNode>(SubNode.Key)->GetSnappedChildForPosition(SubNode.Value.Geometry, Position, OutPositionInNode);
+				}
+			}
+		}
+	}
+
+	TSet< TSharedRef<SWidget> > MyPins;
+	{
+		GetPins( MyPins );
+	}
+	
+	TMap<TSharedRef<SWidget>, FArrangedWidget> Result;
+	FindChildGeometries(Geometry, MyPins, Result);
+	
+	if ( Result.Num() > 0 )
+	{
+		FArrangedChildren ArrangedPins(EVisibility::Visible);
+		Result.GenerateValueArray( ArrangedPins.GetInternalArray() );
+		int32 HoveredPinIndex = SWidget::FindChildUnderPosition( ArrangedPins, Position );
+		if ( HoveredPinIndex != INDEX_NONE )
+		{
+			TSharedRef<SGraphPin> HoveredPin = StaticCastSharedRef<SGraphPin>(
+				ArrangedPins[HoveredPinIndex].Widget);
+			// if the hovered pin has a node then return that:
+			TSharedPtr<SGraphNode> NodeSnappedToPin = HoveredPin->GetSnappedChildForPosition(ArrangedPins[HoveredPinIndex].Geometry, Position, OutPositionInNode);
+			if(NodeSnappedToPin)
+			{
+				return NodeSnappedToPin;
+			}
+		}
+	}
+	OutPositionInNode = Geometry.AbsoluteToLocal(Position);
+	return StaticCastSharedRef<SGraphNode>(AsShared());
+}
+
+FGeometry SGraphNode::GetNodeGeometry(const FGeometry& Geometry) const
+{
+	// if we have a snapped execution child we don't want to include it in our displayed geometry (e.g. dropshadow, selection outline):
+	if(ControlFlowSequence )
+	{
+		TMap<TSharedRef<SWidget>, FArrangedWidget> Result;
+		FindChildGeometries(Geometry, {ControlFlowSequence->GetChildren()->GetChildAt(0)}, Result); // this returns the geometry of the node itself, exc. its control flow children
+		
+		if ( Result.Num() > 0 )
+		{
+			TMap<TSharedRef<SWidget>, FArrangedWidget>::TIterator Begin(Result);
+			return Begin->Value.Geometry;
+		}
+	}
+	return Geometry;
+}
+
+void SGraphNode::DrawChildrenNodeShadow( const FGraphPanelSelectionSet& SelectionSet, FSlateWindowElementList& OutDrawElements, int32 NodeShadowsLayerId, const FGeometry& Geometry, FLinearColor Color)
+{
+	const FVector2f NodeShadowSize = UE::Slate::CastToVector2f(GetDefault<UGraphEditorSettings>()->GetShadowDeltaSize());
+	const FSlateBrush* ShadowBrush = GetShadowBrush(SelectionSet.Contains(GraphNode));
+	FGeometry ChildGeometry = GetNodeGeometry(Geometry);
+	FSlateDrawElement::MakeBox(
+		OutDrawElements,
+		NodeShadowsLayerId,
+		ChildGeometry.ToInflatedPaintGeometry(NodeShadowSize),
+		ShadowBrush,
+		ESlateDrawEffect::None,
+		Color
+	);
+
+	TArray<FArrangedWidget> AllDataPinChildren = GetArrangedSnappedChildNodes(Geometry);
+	for(const FArrangedWidget& DataPinNode : AllDataPinChildren)
+	{
+		// bugbugbug: this draws the node shadow behind our parent node, probably need to move it into onpaint:
+		StaticCastSharedRef<SGraphNode>(DataPinNode.Widget)->DrawChildrenNodeShadow(SelectionSet, OutDrawElements, NodeShadowsLayerId, DataPinNode.Geometry, Color);
+	}
+
+	if(!ControlFlowSequence)
+	{
+		return;
+	}
+
+	FChildren* ControlFlowChildren = ControlFlowSequence->GetChildren();
+	if(ControlFlowChildren->Num() <= 1)
+	{
+		return;
+	}
+
+	TSet< TSharedRef<SWidget> > SubNodes;
+	for(int32 I = 1; I < ControlFlowChildren->Num(); ++I)
+	{
+		SubNodes.Add( ControlFlowChildren->GetChildAt(I) );
+	}
+	TMap<TSharedRef<SWidget>, FArrangedWidget> Result;
+	FindChildGeometries(Geometry, SubNodes, Result);
+		
+	if ( Result.Num() > 0 )
+	{
+		for( const TPair<TSharedRef<SWidget>, FArrangedWidget>& SubNode : Result)
+		{
+			TSharedRef<SGraphNode> ChildNode = StaticCastSharedRef<SGraphNode>(SubNode.Key);
+			ChildNode->DrawChildrenNodeShadow(SelectionSet, OutDrawElements, NodeShadowsLayerId, SubNode.Value.Geometry, Color);
+		}
+	}
+}
+
+void SGraphNode::GetPinsIncludingSnappedChildren( TSet< TSharedRef<SWidget> >& AllPins ) const
+{
+	ForEachPinOnSnappedChild([&AllPins](const TSharedRef<SGraphPin>& GraphPin) {
+		AllPins.Add(GraphPin);
+	});
+}
+
+void SGraphNode::GetPinsToDraw( TSet< TSharedRef<SWidget> >& AllPins ) const
+{
+	ForEachPinOnSnappedChild([&AllPins](const TSharedRef<SGraphPin>& GraphPin) {
+		if(!GraphPin->IsSnapped())
+		{
+			AllPins.Add(GraphPin);
+		}
+	});
+}
+
+void SGraphNode::ForEachNodeIncludingChildren(TFunctionRef<void(const TSharedRef<SGraphNode>&)> Ftor)
+{
+	Ftor(StaticCastSharedRef<SGraphNode>(AsShared()));
+	ForEachNodeSnappedToControlFlowPin(	Ftor );
+	ForEachNodeSnappedToDataPin(Ftor);
+}
+
+void SGraphNode::ForEachPinOnSnappedChild(TFunctionRef<void(const TSharedRef<SGraphPin>&)> Ftor) const
+{
+	for( TSharedRef<SGraphPin> InputPin : InputPins )
+	{
+		Ftor(InputPin);
+		TSharedPtr<SGraphNode> SnappedNode = InputPin->GetSnappedNode();
+		if(SnappedNode)
+		{
+			SnappedNode->ForEachPinOnSnappedChild(Ftor);
+		}
+	}
+	
+	for( TSharedRef<SGraphPin> OutputPin : OutputPins )
+	{
+		Ftor(OutputPin);
+		TSharedPtr<SGraphNode> SnappedNode = OutputPin->GetSnappedNode();
+		if(SnappedNode)
+		{
+			SnappedNode->ForEachPinOnSnappedChild(Ftor);
+		}
+	}
+	
+	// recurse into attached control flow pins:
+	ForEachNodeSnappedToControlFlowPin([Ftor](const TSharedRef<SGraphNode>& AttachedNode) {
+		AttachedNode->ForEachPinOnSnappedChild(Ftor);
+	});
+}
+
+void SGraphNode::ForEachNodeSnappedToControlFlowPin(TFunctionRef<void(const TSharedRef<SGraphNode>&)> Ftor) const
+{
+	if(!ControlFlowSequence)
+	{
+		return;
+	}
+
+	FChildren* ControlFlowChildren = ControlFlowSequence->GetChildren();
+	if(ControlFlowChildren->Num() <= 1)
+	{
+		return;
+	}
+
+	for(int32 I = 1; I < ControlFlowChildren->Num(); ++I)
+	{
+		StaticCastSharedRef<SGraphNode>(ControlFlowChildren->GetChildAt(I))->ForEachNodeIncludingChildren(Ftor);
+	}
+}
+
+void SGraphNode::ForEachNodeSnappedToDataPin(TFunctionRef<void(const TSharedRef<SGraphNode>&)> Ftor) const
+{
+	for( TSharedRef<SGraphPin> InputPin : InputPins )
+	{
+		TSharedPtr<SGraphNode> SnappedNode = InputPin->GetSnappedNode();
+		if(SnappedNode)
+		{
+			SnappedNode->ForEachNodeIncludingChildren(Ftor);
+		}
+	}
+	
+	for( TSharedRef<SGraphPin> OutputPin : OutputPins )
+	{
+		TSharedPtr<SGraphNode> SnappedNode = OutputPin->GetSnappedNode();
+		if(SnappedNode)
+		{
+			SnappedNode->ForEachNodeIncludingChildren(Ftor);
+		}
+	}
+}
diff --git a/Engine/Source/Editor/GraphEditor/Private/SGraphPanel.cpp b/Engine/Source/Editor/GraphEditor/Private/SGraphPanel.cpp
--- a/Engine/Source/Editor/GraphEditor/Private/SGraphPanel.cpp
+++ b/Engine/Source/Editor/GraphEditor/Private/SGraphPanel.cpp
@@ -30,6 +30,7 @@
 #include "GraphEditorActions.h"
 #include "GraphEditorDragDropAction.h"
 #include "GraphEditorSettings.h"
+#include "GraphEditorSnapping/NodeSnappingManager.h"
 #include "HAL/PlatformCrt.h"
 #include "Input/DragAndDrop.h"
 #include "InputCoreTypes.h"
@@ -138,6 +139,11 @@ void SGraphPanel::Construct( const SGraphPanel::FArguments& InArgs )
 	
 	bHasCustomPrepass = true;
 	bCheckNodeGraphObjValidity = false;
+
+	if(FNodeSnappingManager::IsNodeSnappingEnabled())
+	{
+		NodeSnappingManager = MakePimpl<FNodeSnappingManager>(StaticCastSharedRef<SGraphPanel>(AsShared()));
+	}
 }
 
 SGraphPanel::~SGraphPanel()
@@ -353,15 +359,14 @@ int32 SGraphPanel::OnPaint( const FPaintArgs& Args, const FGeometry& AllottedGeo
 				// Draw the node's shadow.
 				if (bDrawShadowsThisFrame || bSelected)
 				{
-					const FSlateBrush* ShadowBrush = ChildNode->GetShadowBrush(bSelected);
-					FSlateDrawElement::MakeBox(
-						OutDrawElements,
-						ShadowLayerId,
-						CurWidget.Geometry.ToInflatedPaintGeometry(NodeShadowSize),
-						ShadowBrush,
-						ESlateDrawEffect::None,
+					// draw drop shadow on connected pins:
+					ChildNode->DrawChildrenNodeShadow(
+						*SelectionToVisualize,
+						/*FSlateWindowElementList& */OutDrawElements, 
+						NodeShadowsLayerId,
+						CurWidget.Geometry,
 						FLinearColor(1.0f, 1.0f, 1.0f, Alpha)
-						);
+					);
 				}
 
 				// Draw the comments and information popups for this node, if it has any.
@@ -540,7 +545,9 @@ int32 SGraphPanel::OnPaint( const FPaintArgs& Args, const FGeometry& AllottedGeo
 			}
 			else
 			{
-				ChildNode->GetPins(VisiblePins);
+				// if node snapping is enabled only unspapped pins will have their connections 
+				// drawn, else this will be all visible pins:
+				ChildNode->GetPinsToDraw(VisiblePins);
 			}
 		}
 
@@ -1224,6 +1231,12 @@ FReply SGraphPanel::OnMouseMove(const FGeometry& MyGeometry, const FPointerEvent
 		// Our frame-latent approximation is going to be totally junk if the mouse is moving quickly
 		OnSplineHoverStateChanged(FGraphSplineOverlapResult());
 	}
+	
+	if(NodeSnappingManager && NodeSnappingManager->IsDragging())
+	{
+		NodeSnappingManager->UpdateDrag(MyGeometry, MouseEvent);
+		return FReply::Handled();
+	}
 
 	return SNodePanel::OnMouseMove(MyGeometry, MouseEvent);
 }
@@ -1642,6 +1655,11 @@ FReply SGraphPanel::OnMouseDownDisconnectClosestPinFromWire(TSharedRef<SGraphPin
 	return FReply::Handled().BeginDragDrop(DragEvent.ToSharedRef());
 }
 
+FNodeSnappingManager* SGraphPanel::GetSnappingManager()
+{
+	return NodeSnappingManager.Get();
+}
+
 void SGraphPanel::OnBeginMakingConnection(UEdGraphPin* InOriginatingPin)
 {
 	OnBeginMakingConnection(FGraphPinHandle(InOriginatingPin));
@@ -2343,6 +2361,11 @@ void SGraphPanel::Update()
 		RemoveAllNodes();
 	}
 
+	if(NodeSnappingManager)
+	{
+		NodeSnappingManager->HandleGraphChanged();
+	}
+
 	// Clean out set of added nodes
 	UserAddedNodes.Reset();
 	UserSelectedNodes.Reset();
@@ -2679,3 +2702,8 @@ void SGraphPanel::SetNodeFactory(const TSharedRef<class FGraphNodeFactory>& NewN
 {
 	NodeFactory = NewNodeFactory;
 }
+
+bool SGraphPanel::IsSnappingEnabled()
+{
+	return FNodeSnappingManager::IsNodeSnappingEnabled();
+}
diff --git a/Engine/Source/Editor/GraphEditor/Private/SGraphPin.cpp b/Engine/Source/Editor/GraphEditor/Private/SGraphPin.cpp
--- a/Engine/Source/Editor/GraphEditor/Private/SGraphPin.cpp
+++ b/Engine/Source/Editor/GraphEditor/Private/SGraphPin.cpp
@@ -311,7 +311,11 @@ void SGraphPin::Construct(const FArguments& InArgs, UEdGraphPin* InPin)
 					SAssignNew(ValueBox, SBox)
 					.Padding(0.0f)
 					[
-						ValueWidget.ToSharedRef()
+						SAssignNew(ValueWidgetWrapper, SBox)
+						.Padding(0.0f)
+						[
+							ValueWidget.ToSharedRef()
+						]
 					]
 				];
 
@@ -397,6 +401,8 @@ void SGraphPin::Construct(const FArguments& InArgs, UEdGraphPin* InPin)
 		.Text(this, &SGraphPin::GetTooltipText);
 
 	SetToolTip(TooltipWidget);
+
+	bIsSnappingEnabled = SGraphPanel::IsSnappingEnabled();
 }
 
 TSharedRef<SWidget>	SGraphPin::GetDefaultValueWidget()
@@ -430,6 +436,11 @@ FReply SGraphPin::OnPinMouseDown( const FGeometry& SenderGeometry, const FPointe
 {
 	bIsMovingLinks = false;
 
+	if(IsSnapped())
+	{
+		return FReply::Unhandled();
+	}
+
 	if (MouseEvent.GetEffectingButton() == EKeys::LeftMouseButton && ensure(!bGraphDataInvalid))
 	{
 		if (IsEditingEnabled())
@@ -1536,6 +1547,98 @@ bool SGraphPin::HasInteractiveTooltip() const
 	return ValueInspectorTooltip.IsValid();
 }
 
+bool SGraphPin::SnapSGraphNode(TSharedRef<SGraphNode> Node)
+{
+	if(!ValueWidgetWrapper)
+	{
+		return false;
+	}
+
+	ValueWidgetWrapper->SetContent(Node);
+	SnappedNode = Node;
+	return true;
+}
+
+void SGraphPin::ClearSnappedSGraphNode()
+{
+	if(ValueWidgetWrapper) // control flow snaps won't have a valud widget
+	{
+		ValueWidgetWrapper->SetContent(ValueWidget.ToSharedRef());
+	}
+	SnappedNode.Reset();
+	//ValueWidgetWrapper->SetActiveWidgetIndex(0); // implicit
+}
+
+bool SGraphPin::IsSnapped() const
+{
+	if(!bIsSnappingEnabled)
+	{
+		return false;
+	}
+
+	UEdGraphPin* GraphPin = GetPinObj();
+	if(!GraphPin)
+	{
+		return false;
+	}
+	return GraphPin->IsSnapped();
+}
+
+TSharedPtr<SGraphNode> SGraphPin::GetSnappedChildForPosition(const FGeometry& MyGeometry, const FVector2f& Position, FVector2f& OutPositionInNode) const
+{
+	TSet< TSharedRef<SWidget> > ValueWidgets;
+	TSharedPtr<SGraphNode> SnappedNodePinned = SnappedNode.Pin();
+	if(!SnappedNodePinned)
+	{
+		return nullptr;
+	}
+	ValueWidgets.Add(SnappedNodePinned.ToSharedRef());
+	TMap<TSharedRef<SWidget>, FArrangedWidget> Result;
+	FindChildGeometries(MyGeometry, ValueWidgets, Result);
+	if ( Result.Num() > 0 )
+	{
+		FArrangedChildren ArrangedPins(EVisibility::Visible);
+		Result.GenerateValueArray( ArrangedPins.GetInternalArray() );
+		int32 HoveredPinIndex = SWidget::FindChildUnderPosition( ArrangedPins, Position );
+		if ( HoveredPinIndex != INDEX_NONE )
+		{
+			TSharedPtr<SGraphNode> HoveredNode = StaticCastSharedRef<SGraphNode>(
+				ArrangedPins[HoveredPinIndex].Widget);
+			OutPositionInNode = ArrangedPins[HoveredPinIndex].Geometry.AbsoluteToLocal(Position);
+			return HoveredNode;
+		}
+	}
+	return nullptr;
+}
+
+TOptional<FArrangedWidget> SGraphPin::GetArrangedSnappedChildNode(const FGeometry& MyGeometry) const
+{
+	TSharedPtr<SGraphNode> SnappedNodePinned = SnappedNode.Pin();
+	if(!SnappedNodePinned)
+	{
+		return {};
+	}
+	
+	TSet< TSharedRef<SWidget> > ValueWidgets;
+	ValueWidgets.Add(SnappedNodePinned.ToSharedRef());
+	TMap<TSharedRef<SWidget>, FArrangedWidget> Result;
+	FindChildGeometries(MyGeometry, ValueWidgets, Result);
+	if(Result.Num() == 0)
+	{
+		// culled
+		return {};
+	}
+	check(Result.Num() == 1);
+	FArrangedChildren ArrangedChildNodes(EVisibility::Visible);
+	Result.GenerateValueArray( ArrangedChildNodes.GetInternalArray() );
+	return ArrangedChildNodes[0];
+}
+
+TSharedPtr<SGraphNode> SGraphPin::GetSnappedNode() const
+{
+	return SnappedNode.Pin();
+}
+
 bool SGraphPin::GetIsConnectable() const
 {
 	UEdGraphPin* GraphPin = GetPinObj();
diff --git a/Engine/Source/Editor/GraphEditor/Private/SNodePanel.cpp b/Engine/Source/Editor/GraphEditor/Private/SNodePanel.cpp
--- a/Engine/Source/Editor/GraphEditor/Private/SNodePanel.cpp
+++ b/Engine/Source/Editor/GraphEditor/Private/SNodePanel.cpp
@@ -13,6 +13,7 @@
 #include "GenericPlatform/GenericApplicationMessageHandler.h"
 #include "GenericPlatform/ICursor.h"
 #include "GraphEditorSettings.h"
+#include "GraphEditorSnapping/NodeSnappingManager.h"
 #include "Input/Events.h"
 #include "InputCoreTypes.h"
 #include "Internationalization/Internationalization.h"
@@ -536,6 +537,16 @@ void SNodePanel::OnEndNodeInteraction(const TSharedRef<SNode>& InNodeToDrag)
 	InNodeToDrag->EndUserInteraction();
 }
 
+FNodeSnappingManager* SNodePanel::GetSnappingManager()
+{
+	return nullptr;
+}
+
+const FNodeSnappingManager* SNodePanel::GetSnappingManager() const
+{
+	return const_cast<SNodePanel*>(this)->GetSnappingManager();
+}
+
 EActiveTimerReturnType SNodePanel::HandleZoomToFit(double InCurrentTime, float InDeltaTime)
 {
 	const FVector2f DesiredViewCenter = ( ZoomTargetTopLeft + ZoomTargetBottomRight ) * 0.5f;
@@ -740,26 +751,36 @@ FReply SNodePanel::OnMouseButtonDown( const FGeometry& MyGeometry, const FPointe
 	else if ( bIsLeftMouseButtonEffecting )
 	{
 		// LEFT BUTTON is for selecting nodes and manipulating pins.
-		FArrangedChildren ArrangedChildren(EVisibility::Visible);
-		ArrangeChildNodes(MyGeometry, ArrangedChildren);
-
-		const int32 NodeUnderMouseIndex = SWidget::FindChildUnderMouse( ArrangedChildren, MouseEvent );
-		if ( NodeUnderMouseIndex != INDEX_NONE )
+		if(FNodeSnappingManager* SnapManager = GetSnappingManager())
 		{
-			// PRESSING ON A NODE!
-
-			// This changes selection and starts dragging it.
-			const FArrangedWidget& NodeGeometry = ArrangedChildren[NodeUnderMouseIndex];
-			const FVector2f MousePositionInNode = NodeGeometry.Geometry.AbsoluteToLocal(MouseEvent.GetScreenSpacePosition());
-			TSharedRef<SNode> NodeWidgetUnderMouse = StaticCastSharedRef<SNode>( NodeGeometry.Widget );
-
-			if( NodeWidgetUnderMouse->CanBeSelected(MousePositionInNode) )
+			if(SnapManager->BeginDragWithSnap(MyGeometry, MouseEvent))
 			{
-				// Track the node that we're dragging; we will move it in OnMouseMove.
-				this->OnBeginNodeInteraction(NodeWidgetUnderMouse, MousePositionInNode);
 				return FReply::Handled().CaptureMouse( SharedThis(this) );
 			}
 		}
+		else
+		{
+			FArrangedChildren ArrangedChildren(EVisibility::Visible);
+			ArrangeChildNodes(MyGeometry, ArrangedChildren);
+		
+			const int32 NodeUnderMouseIndex = SWidget::FindChildUnderMouse( ArrangedChildren, MouseEvent );
+			if ( NodeUnderMouseIndex != INDEX_NONE )
+			{
+				// PRESSING ON A NODE!
+
+				// This changes selection and starts dragging it.
+				const FArrangedWidget& NodeGeometry = ArrangedChildren[NodeUnderMouseIndex];
+				const FVector2f MousePositionInNode = NodeGeometry.Geometry.AbsoluteToLocal(MouseEvent.GetScreenSpacePosition());
+				TSharedRef<SNode> NodeWidgetUnderMouse = StaticCastSharedRef<SNode>( NodeGeometry.Widget );
+
+				if( NodeWidgetUnderMouse->CanBeSelected(MousePositionInNode) )
+				{
+					// Track the node that we're dragging; we will move it in OnMouseMove.
+					this->OnBeginNodeInteraction(NodeWidgetUnderMouse, MousePositionInNode);
+					return FReply::Handled().CaptureMouse( SharedThis(this) );
+				}
+			}
+		}
 
 		// START MARQUEE SELECTION.
 		const FVector2f GraphMousePos = PanelCoordToGraphCoord( MyGeometry.AbsoluteToLocal( MouseEvent.GetScreenSpacePosition() ) );
@@ -1207,10 +1228,8 @@ void SNodePanel::FindNodesAffectedByMarquee( FGraphPanelSelectionSet& OutAffecte
 	OutAffectedNodes.Empty();
 
 	FSlateRect MarqueeSlateRect = Marquee.Rect.ToSlateRect();
-
-	for ( int32 NodeIndex=0; NodeIndex < Children.Num(); ++NodeIndex )
+	const auto FindNodeAffectedByMarquee = [&OutAffectedNodes, &MarqueeSlateRect = std::as_const(MarqueeSlateRect)](const TSharedRef<SNode>& SomeNodeWidget)
 	{
-		const TSharedRef<SNode>& SomeNodeWidget = Children[NodeIndex];
 		const FVector2f NodePosition = SomeNodeWidget->GetPosition2f();
 		const FVector2f NodeSize = SomeNodeWidget->GetDesiredSizeForMarquee2f();
 
@@ -1224,6 +1243,20 @@ void SNodePanel::FindNodesAffectedByMarquee( FGraphPanelSelectionSet& OutAffecte
 				OutAffectedNodes.Add(SomeNodeWidget->GetObjectBeingDisplayed());
 			}
 		}
+	};
+
+	if(const FNodeSnappingManager* SnapManager = GetSnappingManager())
+	{
+		SnapManager->ForEachNode([&FindNodeAffectedByMarquee](const TSharedRef<SNode>& Node)
+		{
+			FindNodeAffectedByMarquee(Node);
+		});
+		return;
+	}
+
+	for ( int32 NodeIndex=0; NodeIndex < Children.Num(); ++NodeIndex )
+	{
+		FindNodeAffectedByMarquee(Children[NodeIndex]);
 	}	
 }
 
@@ -1871,6 +1904,14 @@ bool SNodePanel::HasDeferredZoomDestination() const
 
 void SNodePanel::FinalizeNodeMovements()
 {
+	if(FNodeSnappingManager* SnapManager = GetSnappingManager())
+	{
+		if(SnapManager->FinalizeNodeMovements())
+		{
+			return;
+		}
+	}
+
 	// Process moved nodes on focus lost
 	if (OriginalNodePositions.Num() > 0)
 	{
diff --git a/Engine/Source/Editor/GraphEditor/Public/SGraphNode.h b/Engine/Source/Editor/GraphEditor/Public/SGraphNode.h
--- a/Engine/Source/Editor/GraphEditor/Public/SGraphNode.h
+++ b/Engine/Source/Editor/GraphEditor/Public/SGraphNode.h
@@ -144,6 +144,7 @@ public:
 	// SNodePanel::SNode interface
 	UE_API virtual void MoveTo( const FVector2f& NewPosition, FNodeSet& NodeFilter, bool bMarkDirty = true ) override;
 	UE_API virtual FVector2f GetPosition2f() const override;
+	UE_API virtual FVector2f GetDesiredSizeForMarquee2f() const override;
 	UE_API virtual FString GetNodeComment() const override;
 	UE_API virtual UObject* GetObjectBeingDisplayed() const override;
 	// End of SNodePanel::SNode interface
@@ -293,6 +294,19 @@ public:
 	/** Allows the node to override the pin text color */
 	virtual TOptional<FSlateColor> GetPinTextColor(const SGraphPin* InGraphPin) const { return TOptional<FSlateColor>(); }
 
+	// snapping:
+	TArray<FArrangedWidget> GetArrangedSnappedChildNodes(const FGeometry& Geometry);
+	TSharedPtr<SGraphNode> GetSnappedChildForPosition(const FGeometry& Geometry, const FVector2f& Position, FVector2f& OutPositionInNode);
+	FGeometry GetNodeGeometry(const FGeometry& Geometry) const;
+	void DrawChildrenNodeShadow( const FGraphPanelSelectionSet& SelectionSet, FSlateWindowElementList& OutDrawElements, int32 NodeShadowsLayerId, const FGeometry& Geometry, FLinearColor Color);
+	/** Returns the set of pins that should have their connections drawn, including nested pins */
+	void GetPinsToDraw( TSet< TSharedRef<SWidget> >& AllPins ) const;
+	/** Returns all of the pins, including those nested within this node */
+	void GetPinsIncludingSnappedChildren( TSet< TSharedRef<SWidget> >& AllPins ) const;
+	void ForEachNodeIncludingChildren(TFunctionRef<void(const TSharedRef<SGraphNode>&)> Ftor);
+	void ForEachPinOnSnappedChild(TFunctionRef<void(const TSharedRef<SGraphPin>&)> Ftor) const;
+	void ForEachNodeSnappedToControlFlowPin(TFunctionRef<void(const TSharedRef<SGraphNode>&)> Ftor) const;
+	void ForEachNodeSnappedToDataPin(TFunctionRef<void(const TSharedRef<SGraphNode>&)> Ftor) const;
 protected:
 	UE_API SGraphNode();
 
@@ -502,6 +516,9 @@ protected:
 
 private:
 	TSharedPtr<SLevelOfDetailBranchNode> TitleLODBranchNode;
+	friend struct FNodeSnappingManager; // want access to ControlFlowSequence, can reconsider design if snapping is successful
+	TSharedPtr<SVerticalBox> ControlFlowSequence; // ref?
+	TWeakPtr<SGraphNode> SnapParent; // dag so technically we can keep an owning reference.. but do we trust ourselves?
 };
 
 #undef UE_API
diff --git a/Engine/Source/Editor/GraphEditor/Public/SGraphPanel.h b/Engine/Source/Editor/GraphEditor/Public/SGraphPanel.h
--- a/Engine/Source/Editor/GraphEditor/Public/SGraphPanel.h
+++ b/Engine/Source/Editor/GraphEditor/Public/SGraphPanel.h
@@ -54,6 +54,7 @@ struct FAssetData;
 struct FDiffSingleResult;
 struct FEdGraphEditAction;
 struct FGuid;
+struct FNodeSnappingManager;
 
 DECLARE_DELEGATE( FOnUpdateGraphPanel )
 
@@ -281,6 +282,7 @@ public:
 	/** Sets the current widget factory. */
 	UE_API void SetNodeFactory(const TSharedRef<class FGraphNodeFactory>& NewNodeFactory);
 
+	static bool IsSnappingEnabled();
 protected:
 
 	UE_API void NotifyGraphChanged ( const struct FEdGraphEditAction& InAction);
@@ -451,6 +453,11 @@ private:
 	/** Begins a drag operation that disconnects the pin at ClosestPin, but only from the ConnectedPin, rather than all connections */
 	UE_API FReply OnMouseDownDisconnectClosestPinFromWire(TSharedRef<class SGraphPin> ClosestPin, TSharedRef<class SGraphPin> ConnectedPin);
 
+	/** Node snapping extension */
+	virtual struct FNodeSnappingManager* GetSnappingManager() override;
+	friend struct FNodeSnappingManager;
+	TPimplPtr<FNodeSnappingManager> NodeSnappingManager;
+
 	/** Handle to timer callback that allows the UI to refresh it's arrangement each tick, allows animations to occur within the UI */
 	TWeakPtr<FActiveTimerHandle> ActiveTimerHandleInvalidatePerTick;
 
diff --git a/Engine/Source/Editor/GraphEditor/Public/SGraphPin.h b/Engine/Source/Editor/GraphEditor/Public/SGraphPin.h
--- a/Engine/Source/Editor/GraphEditor/Public/SGraphPin.h
+++ b/Engine/Source/Editor/GraphEditor/Public/SGraphPin.h
@@ -26,6 +26,7 @@
 #include "UObject/NameTypes.h"
 #include "Widgets/DeclarativeSyntaxSupport.h"
 #include "Widgets/Layout/SBorder.h"
+#include "Widgets/Layout/SBox.h"
 #include "Widgets/Layout/SWrapBox.h"
 #include "Widgets/SWidget.h"
 
@@ -272,6 +273,12 @@ public:
 	/** Enables or disables drag and drop on the pin */
 	void EnableDragAndDrop(bool bEnable) { bDragAndDropEnabled = bEnable; }
 
+	bool SnapSGraphNode(TSharedRef<SGraphNode> Node);
+	void ClearSnappedSGraphNode();
+	bool IsSnapped() const;
+	TSharedPtr<SGraphNode> GetSnappedChildForPosition(const FGeometry& MyGeometry, const FVector2f& Position, FVector2f& OutPositionInNode) const;
+	TOptional<FArrangedWidget> GetArrangedSnappedChildNode(const FGeometry& MyGeometry) const;
+	TSharedPtr<SGraphNode> GetSnappedNode() const;
 protected:
 
 	/** If true the graph pin subclass is responsible for setting the IsEnabled delegates for the aspects it cares about. If false, the default value widget enabling is done by the base class */
@@ -347,11 +354,14 @@ public:
 	UE_API void SetPinImageWidget(TSharedRef<SWidget> NewWidget);
 
 protected:
+	friend struct FNodeSnappingManager; // want access to owner, can reconsider design if snapping is successful
 	/** The GraphNode that owns this pin */
 	TWeakPtr<SGraphNode> OwnerNodePtr;
 
 	/** Image of pin */
 	TSharedPtr<SWidget> PinImage;
+	TSharedPtr<SBox> ValueWidgetWrapper;
+	TWeakPtr<SGraphNode> SnappedNode;
 
 	/** Label and value */
 	TSharedPtr<SWrapBox> LabelAndValue;
@@ -431,6 +441,9 @@ protected:
 	bool bFadeConnections;
 
 private:
+	/** Cached value of snapping feature flag */
+	bool bIsSnappingEnabled;
+
 	TSharedPtr<SLevelOfDetailBranchNode> PinNameLODBranchNode;
 	friend class SGraphPanel;
 };
diff --git a/Engine/Source/Editor/GraphEditor/Public/SNodePanel.h b/Engine/Source/Editor/GraphEditor/Public/SNodePanel.h
--- a/Engine/Source/Editor/GraphEditor/Public/SNodePanel.h
+++ b/Engine/Source/Editor/GraphEditor/Public/SNodePanel.h
@@ -1188,6 +1188,9 @@ protected:
 	SGraphEditor::FOnNodeSingleClicked OnNodeSingleClicked;
 
 private:
+	UE_API virtual struct FNodeSnappingManager* GetSnappingManager();
+	const struct FNodeSnappingManager* GetSnappingManager() const;
+	
 	/** Active timer that handles deferred zooming until the target zoom is reached */
 	UE_API EActiveTimerReturnType HandleZoomToFit(double InCurrentTime, float InDeltaTime);
 
diff --git a/Engine/Source/Runtime/Engine/Classes/EdGraph/EdGraphNode.h b/Engine/Source/Runtime/Engine/Classes/EdGraph/EdGraphNode.h
--- a/Engine/Source/Runtime/Engine/Classes/EdGraph/EdGraphNode.h
+++ b/Engine/Source/Runtime/Engine/Classes/EdGraph/EdGraphNode.h
@@ -997,6 +997,10 @@ public:
 	 */
 	ENGINE_API void ForEachNodeDirectlyConnectedToOutputs(TFunctionRef<void(UEdGraphNode*)> Func);
 	
+	// snapping
+	ENGINE_API bool IsSnappedToAny(const TSet<TObjectPtr<UObject>>& NodeSet ) const;
+	ENGINE_API UEdGraphNode* GetRootSnappedParent(const TSet<TObjectPtr<UObject>>& NodeSet) const;
+	// snapping
 protected:
 #if WITH_EDITORONLY_DATA
 	/** Internal function used to gather pins from a graph node for localization */
diff --git a/Engine/Source/Runtime/Engine/Classes/EdGraph/EdGraphPin.h b/Engine/Source/Runtime/Engine/Classes/EdGraph/EdGraphPin.h
--- a/Engine/Source/Runtime/Engine/Classes/EdGraph/EdGraphPin.h
+++ b/Engine/Source/Runtime/Engine/Classes/EdGraph/EdGraphPin.h
@@ -278,6 +278,16 @@ private:
 	FGuid PinId;
 };
 
+#if WITH_EDITORONLY_DATA
+struct FSnappedEdGraphPin
+{
+	UEdGraphPin* ParentPin = nullptr;
+	UEdGraphPin* ChildPin = nullptr;
+
+	friend auto operator<=>(const FSnappedEdGraphPin&, const FSnappedEdGraphPin&) = default;
+};
+#endif
+
 enum class EPinResolveType : uint8;
 
 class UEdGraphPin
@@ -347,6 +357,13 @@ public:
 	/** Older content sometimes had an empty autogenerated default value string in cases where that does not mean the property default value (0, none, false, etc.) */
 	uint8 bUseBackwardsCompatForEmptyAutogeneratedValue:1;
 
+private:
+	/** True if we have a snapped child pin in our LinkedTo list, corresponding pin will have bHasSnappedParent set */
+	uint8 bHasSnappedChild:1;
+	
+	/** True if we are snapped within a parent pin in our LinkedTo list, corresponding pin will have bHasSnappedChild set */
+	uint8 bHasSnappedParent:1;
+public:
 #endif // WITH_EDITORONLY_DATA
 
 	/** True when InvalidateAndTrash was called. This pin is intended to be discarded and destroyed. */
@@ -402,12 +419,34 @@ public:
 	/** Create a link. Note, this does not check that schema allows it, and will not break any existing connections */
 	ENGINE_API void MakeLinkTo(UEdGraphPin* ToPin, bool bAlwaysMarkDirty = true);
 
+#if WITH_EDITORONLY_DATA
+	ENGINE_API void SetSnappedChild(UEdGraphPin* SnappedChild);
+	ENGINE_API void ClearSnappedChild();
+	ENGINE_API bool IsSnapped() const;
+	ENGINE_API bool IsSnappedParentOf(const UEdGraphNode* Node) const;
+	ENGINE_API FSnappedEdGraphPin GetSnapDescription() const;
+	ENGINE_API UEdGraphNode* GetSnappedChildNode() const;
+	ENGINE_API UEdGraphPin* GetSnappedChildPin() const;
+	ENGINE_API UEdGraphPin* GetSnappedParentPin() const;
+#endif
+
 	/** Break a link to the specified pin (if present) */
 	ENGINE_API void BreakLinkTo(UEdGraphPin* ToPin, bool bAlwaysMarkDirty = true);
 
 	/** Break all links from this pin */
 	ENGINE_API void BreakAllPinLinks(bool bNotifyNodes = false, bool bAlwaysMarkDirty = true);
 
+private:
+	enum class ETransferPersistentDataMode : uint8
+	{
+		Move,
+		Copy
+	};
+
+	template<class T>
+	static void TransferPersistentDataFromOldPin(UEdGraphPin& DestPin, T& SourcePin, const ETransferPersistentDataMode TransferMode);
+
+public:
 	/**
 	* Moves the persistent data (across a node refresh) from the SourcePin.
 	*
@@ -588,6 +627,11 @@ private:
 	static FString ExportText_PinReference(const UEdGraphPin* Pin);
 	static FString ExportText_PinArray(const TArray<UEdGraphPin*>& PinArray);
 	static bool ImportText_PinArray(const TCHAR*& Buffer, TArray<UEdGraphPin*>& ArrayRef, UEdGraphPin* RequestingPin, EPinResolveType ResolveType);
+	
+#if WITH_EDITORONLY_DATA
+	static FString bHasSnappedChildName;
+	static FString bHasSnappedParentName;
+#endif
 };
 
 UCLASS(MinimalAPI)
diff --git a/Engine/Source/Runtime/Engine/Private/EdGraph/EdGraphNode.cpp b/Engine/Source/Runtime/Engine/Private/EdGraph/EdGraphNode.cpp
--- a/Engine/Source/Runtime/Engine/Private/EdGraph/EdGraphNode.cpp
+++ b/Engine/Source/Runtime/Engine/Private/EdGraph/EdGraphNode.cpp
@@ -1066,6 +1066,52 @@ void UEdGraphNode::ForEachNodeDirectlyConnectedToOutputs(TFunctionRef<void(UEdGr
 	);
 }
 
+bool UEdGraphNode::IsSnappedToAny(const TSet<TObjectPtr<UObject>>& NodeSet ) const
+{
+	for (const UEdGraphPin* Pin : Pins)
+	{
+		for (UEdGraphPin* Connection : Pin->LinkedTo)
+		{
+			if(Connection->IsSnappedParentOf(this))
+			{
+				if( NodeSet.Contains(Connection->GetOwningNode()))
+				{
+					return true;
+				}
+				else if(Connection->GetOwningNode()->IsSnappedToAny(NodeSet))
+				{
+					return true;
+				}
+			}
+		}
+	}
+	return false;
+}
+
+UEdGraphNode* UEdGraphNode::GetRootSnappedParent(const TSet<TObjectPtr<UObject>>& NodeSet) const
+{
+	for (const UEdGraphPin* Pin : Pins)
+	{
+		for (UEdGraphPin* Connection : Pin->LinkedTo)
+		{
+			if(Connection->IsSnappedParentOf(this))
+			{
+				UEdGraphNode* OwningNode = Connection->GetOwningNode();
+				if( NodeSet.Contains(OwningNode))
+				{
+					return OwningNode;
+				}
+				OwningNode = OwningNode->GetRootSnappedParent(NodeSet);
+				if(OwningNode)
+				{
+					return OwningNode;
+				}
+			}
+		}
+	}
+	return nullptr;
+}
+
 #endif	//#if WITH_EDITOR
 
 bool UEdGraphNode::IsInDevelopmentMode() const
diff --git a/Engine/Source/Runtime/Engine/Private/EdGraph/EdGraphPin.cpp b/Engine/Source/Runtime/Engine/Private/EdGraph/EdGraphPin.cpp
--- a/Engine/Source/Runtime/Engine/Private/EdGraph/EdGraphPin.cpp
+++ b/Engine/Source/Runtime/Engine/Private/EdGraph/EdGraphPin.cpp
@@ -540,6 +540,123 @@ void UEdGraphPin::MakeLinkTo(UEdGraphPin* ToPin, bool bAlwaysMarkDirty)
 	}
 }
 
+#if WITH_EDITORONLY_DATA
+void UEdGraphPin::SetSnappedChild(UEdGraphPin* SnappedChild)
+{
+	check(SnappedChild);
+	if( !ensure(SnappedChild->bHasSnappedParent == false && bHasSnappedChild == false) ||
+		!ensure(LinkedTo.Contains(SnappedChild)))
+	{
+		return;
+	}
+
+	SnappedChild->bHasSnappedParent = true;
+	bHasSnappedChild = true;
+}
+
+void UEdGraphPin::ClearSnappedChild()
+{
+	if(UEdGraphPin* OldSnappedChild = GetSnappedChildPin())
+	{
+		OldSnappedChild->bHasSnappedParent = false;
+	}
+	bHasSnappedChild = false;
+}
+
+bool UEdGraphPin::IsSnapped() const
+{
+	return bHasSnappedChild || bHasSnappedParent;
+}
+
+bool UEdGraphPin::IsSnappedParentOf(const UEdGraphNode* Node) const
+{
+	if(!bHasSnappedChild)
+	{
+		return false;
+	}
+
+	if(UEdGraphPin* SnappedChild = GetSnappedChildPin())
+	{
+		return SnappedChild->GetOwningNode() == Node;
+	}
+
+	return false;
+}
+
+FSnappedEdGraphPin UEdGraphPin::GetSnapDescription() const
+{
+	if(!IsSnapped())
+	{
+		return {};
+	}
+
+	// if we have a child attachment return that:
+	FSnappedEdGraphPin Result;
+	if(bHasSnappedChild)
+	{
+		if(UEdGraphPin* ChildPin = GetSnappedChildPin())
+		{
+			Result.ParentPin = const_cast<UEdGraphPin*>(this);
+			Result.ChildPin = ChildPin;
+		}
+	}
+	else
+	{
+		// terminal attachment - fall back to describing our relationship with the parent pin:
+		if(UEdGraphPin* SnappedParentPin = GetSnappedParentPin())
+		{
+			Result.ParentPin = SnappedParentPin;
+			Result.ChildPin = const_cast<UEdGraphPin*>(this);
+		}
+	}
+	return Result;
+}
+
+UEdGraphNode* UEdGraphPin::GetSnappedChildNode() const
+{
+	if(UEdGraphPin* SnappedChild = GetSnappedChildPin())
+	{
+		return SnappedChild->GetOwningNode();
+	}
+	return nullptr;
+}
+
+UEdGraphPin* UEdGraphPin::GetSnappedChildPin() const
+{
+	if(!bHasSnappedChild)
+	{
+		return nullptr;
+	}
+	
+	for(UEdGraphPin* LinkedPin : LinkedTo)
+	{
+		if(LinkedPin->bHasSnappedParent)
+		{
+			return LinkedPin;
+		}
+	}
+	return nullptr;
+}
+
+UEdGraphPin* UEdGraphPin::GetSnappedParentPin() const
+{
+	if(!bHasSnappedParent)
+	{
+		return nullptr;
+	}
+
+	for(UEdGraphPin* LinkedPin : LinkedTo)
+	{
+		if(LinkedPin->bHasSnappedChild)
+		{
+			return LinkedPin;
+		}
+	}
+	return nullptr;
+}
+
+#endif
+
 void UEdGraphPin::BreakLinkTo(UEdGraphPin* ToPin, bool bAlwaysMarkDirty)
 {
 	if (ToPin)
@@ -589,20 +706,18 @@ void UEdGraphPin::BreakAllPinLinks(bool bNotifyNodes, bool bAlwaysMarkDirty)
 	}
 }
 
-enum class ETransferPersistentDataMode : uint8
-{
-	Move,
-	Copy
-};
-
 template<class T>
-void TransferPersistentDataFromOldPin(UEdGraphPin& DestPin, T& SourcePin, const ETransferPersistentDataMode TransferMode)
+void UEdGraphPin::TransferPersistentDataFromOldPin(UEdGraphPin& DestPin, T& SourcePin, const UEdGraphPin::ETransferPersistentDataMode TransferMode)
 {
 	DestPin.Modify(/*bAlwaysMarkDirty =*/ false);
 
 	// The name matches already, doesn't get copied here
 	// The PinType, Direction, and bNotConnectable are properties generated from the schema
 	DestPin.PinId = SourcePin.PinId;
+#if WITH_EDITORONLY_DATA
+	DestPin.bHasSnappedChild = SourcePin.bHasSnappedChild;
+	DestPin.bHasSnappedParent = SourcePin.bHasSnappedParent;
+#endif//WITH_EDITORONLY_DATA
 
 	// Only move the default value if it was modified; inherit the new default value otherwise
 	if (DestPin.Direction == EGPD_Input)
@@ -681,14 +796,14 @@ void UEdGraphPin::MovePersistentDataFromOldPin(UEdGraphPin& SourcePin)
 {
 	SourcePin.Modify(/*bAlwaysMarkDirty =*/ false);
 
-	TransferPersistentDataFromOldPin(*this, SourcePin, ETransferPersistentDataMode::Move);
+	TransferPersistentDataFromOldPin(*this, SourcePin, UEdGraphPin::ETransferPersistentDataMode::Move);
 
 	SourcePin.LinkedTo.Reset();
 }
 
 void UEdGraphPin::CopyPersistentDataFromOldPin(const UEdGraphPin& SourcePin)
 {
-	TransferPersistentDataFromOldPin(*this, SourcePin, ETransferPersistentDataMode::Copy);
+	TransferPersistentDataFromOldPin(*this, SourcePin, UEdGraphPin::ETransferPersistentDataMode::Copy);
 }
 
 void UEdGraphPin::AssignByRefPassThroughConnection(UEdGraphPin* InTargetPin)
@@ -1100,6 +1215,16 @@ bool UEdGraphPin::ExportTextItem(FString& ValueStr, int32 PortFlags) const
 
 		// Intentionally not exporting bSavePinIfOrphaned as it is transient
 		// Intentionally not exporting bWasTrashed as it is transient
+		
+		ValueStr += UEdGraphPin::bHasSnappedChildName + TEXT("=");
+		bool LocalHasSnappedChildName = bHasSnappedChild;
+		BoolPropCDO->ExportTextItem_Direct(ValueStr, &LocalHasSnappedChildName, nullptr, nullptr, PortFlags, nullptr);
+		ValueStr += PinHelpers::ExportTextPropDelimiter;
+		
+		ValueStr += UEdGraphPin::bHasSnappedParentName + TEXT("=");
+		bool LocalHasSnappedParentName = bHasSnappedParent;
+		BoolPropCDO->ExportTextItem_Direct(ValueStr, &LocalHasSnappedParentName, nullptr, nullptr, PortFlags, nullptr);
+		ValueStr += PinHelpers::ExportTextPropDelimiter;
 	}
 #endif //WITH_EDITORONLY_DATA
 	ValueStr += ")";
@@ -1356,6 +1481,20 @@ bool UEdGraphPin::ImportTextItem(const TCHAR*& Buffer, int32 PortFlags, class UO
 		}
 		// Intentionally not including bSavePinIfOrphaned. It is transient.
 		// Intentionally not importing bWasTrashed as it is transient
+		else if (PropertyToken == UEdGraphPin::bHasSnappedChildName)
+		{
+			bool LocalHasSnappedChild = bHasSnappedChild;
+			Buffer = BoolPropCDO->ImportText_Direct(Buffer, &LocalHasSnappedChild, Parent, PortFlags, ErrorText);
+			bParseSuccess = (Buffer != nullptr);
+			bHasSnappedChild = LocalHasSnappedChild;
+		}
+		else if (PropertyToken == UEdGraphPin::bHasSnappedParentName)
+		{
+			bool LocalHasSnappedParent = bHasSnappedParent;
+			Buffer = BoolPropCDO->ImportText_Direct(Buffer, &LocalHasSnappedParent, Parent, PortFlags, ErrorText);
+			bParseSuccess = (Buffer != nullptr);
+			bHasSnappedParent = LocalHasSnappedParent;
+		}
 #endif
 
 		if (!bParseSuccess)
@@ -1434,6 +1573,8 @@ UEdGraphPin::UEdGraphPin(UEdGraphNode* InOwningNode, const FGuid& PinIdGuid)
 	, bOrphanedPin(false)
 	, bSavePinIfOrphaned(true)
 	, bUseBackwardsCompatForEmptyAutogeneratedValue(false)
+	, bHasSnappedChild(false)
+	, bHasSnappedParent(false)
 #endif
 	, bWasTrashed(false)
 #if WITH_EDITORONLY_DATA
@@ -1587,6 +1728,8 @@ void UEdGraphPin::InitFromDeprecatedPin(class UEdGraphPin_Deprecated* Deprecated
 	bAdvancedView = DeprecatedPin->bAdvancedView;
 	bDisplayAsMutableRef = DeprecatedPin->bDisplayAsMutableRef;
 	bUseBackwardsCompatForEmptyAutogeneratedValue = true; // If it is from deprecated pin format then it is definitely old enough to need this
+	bHasSnappedChild = false;
+	bHasSnappedParent = false;
 	PersistentGuid = DeprecatedPin->PersistentGuid;
 #endif // WITH_EDITORONLY_DATA
 
@@ -1741,8 +1884,10 @@ bool UEdGraphPin::Serialize(FArchive& Ar)
 		BitField |= bDefaultValueIsIgnored << 3;
 		BitField |= bAdvancedView << 4;
 		BitField |= bOrphanedPin << 5;
+		BitField |= bHasSnappedChild << 6;
+		BitField |= bHasSnappedParent << 7;
 
-		const int32 PersistentBits = 6;
+		const int32 PersistentBits = 8;
 
 		// While these properties are transient from a serialized point of view, they need to be maintained through undo/redo so add them to
 		// the bitfield after the persistent bits.
@@ -1764,6 +1909,8 @@ bool UEdGraphPin::Serialize(FArchive& Ar)
 		{
 			bOrphanedPin = !!(BitField & (1 << 5));
 		}
+		bHasSnappedChild = !!(BitField & (1 << 6));
+		bHasSnappedParent = !!(BitField & (1 << 7));
 
 		if (Ar.IsTransacting())
 		{
@@ -2244,6 +2391,11 @@ bool UEdGraphPin::ImportText_PinArray(const TCHAR*& Buffer, TArray<UEdGraphPin*>
 	return false;
 }
 
+#if WITH_EDITORONLY_DATA
+FString UEdGraphPin::bHasSnappedChildName = GET_MEMBER_NAME_STRING_CHECKED(UEdGraphPin, bHasSnappedChildName);
+FString UEdGraphPin::bHasSnappedParentName = GET_MEMBER_NAME_STRING_CHECKED(UEdGraphPin, bHasSnappedParentName);
+#endif
+
 int32 GCVarDisableOrphanPins = 0;
 FAutoConsoleVariableRef CVarDisableOrphanPins(TEXT("DisableOrphanPins"), GCVarDisableOrphanPins, TEXT("0=Orphan pins are enabled (default), 1=Orphan pins are disabled (note: this option will go away in the future)"), ECVF_ReadOnly);
 
