If there is a wire under the cursor when the graph editor context is summoned, include that wire's terminating pins in the context menu's source pins
[RN] In the graph editors right clicking to summon the node context menu now includes the hovered wires conneted pins, if any. This allows, for example, insertion of a 'NOT boolean' node on a bool wire in the blueprint editor

diff --git a/Engine/Source/Editor/GraphEditor/Private/SGraphPanel.cpp b/Engine/Source/Editor/GraphEditor/Private/SGraphPanel.cpp
--- a/Engine/Source/Editor/GraphEditor/Private/SGraphPanel.cpp
+++ b/Engine/Source/Editor/GraphEditor/Private/SGraphPanel.cpp
@@ -1154,10 +1154,22 @@ TSharedPtr<SWidget> SGraphPanel::OnSummonContextMenu(const FGeometry& MyGeometry
 	TSharedPtr<SGraphNode> NodeUnderMouse = GetGraphNodeUnderMouse(MyGeometry, MouseEvent);
 	UEdGraphPin* PinUnderCursor = GetPinUnderMouse(MyGeometry, MouseEvent, NodeUnderMouse);
 	UEdGraphNode* EdNodeUnderMouse = NodeUnderMouse.IsValid() ? NodeUnderMouse->GetNodeObj() : nullptr;
-	TArray<UEdGraphPin*> NoSourcePins;
+	TArray<UEdGraphPin*> SourcePins;
+
+	if(!EdNodeUnderMouse && !PinUnderCursor)
+	{
+		// if there's a wire under the mouse, use its begin and end pins as the source pins:
+		UEdGraphPin* Pin1;
+		UEdGraphPin* Pin2;
+		if (PreviousFrameSplineOverlap.GetPins(*this, Pin1, Pin2))
+		{
+			SourcePins.Add(Pin1);
+			SourcePins.Add(Pin2);
+		}
+	}
 
 	const FVector2f NodeAddPosition = PanelCoordToGraphCoord(MyGeometry.AbsoluteToLocal(MouseEvent.GetScreenSpacePosition()));
-	return SummonContextMenu(MouseEvent.GetScreenSpacePosition(), NodeAddPosition, EdNodeUnderMouse, PinUnderCursor, NoSourcePins);
+	return SummonContextMenu(MouseEvent.GetScreenSpacePosition(), NodeAddPosition, EdNodeUnderMouse, PinUnderCursor, SourcePins);
 }
 
 
