A few virtuals added to cosmetic SGraphNode getters.

Useful while adding advanced styling in custom node editors. In some cases it might be possible to avoid overriding engine methods by implementing our own variants, with different names. But not in every case. Plus reusing original methods makes the programmer's intent much clearer - what part of the base class was actually overridden.

PR #9791: A few virtuals added to cosmetic SGraphNode getters

diff --git a/Engine/Source/Editor/GraphEditor/Private/SGraphNode.cpp b/Engine/Source/Editor/GraphEditor/Private/SGraphNode.cpp
--- a/Engine/Source/Editor/GraphEditor/Private/SGraphNode.cpp
+++ b/Engine/Source/Editor/GraphEditor/Private/SGraphNode.cpp
@@ -1077,7 +1077,7 @@ void SGraphNode::UpdateGraphNode()
 END_SLATE_FUNCTION_BUILD_OPTIMIZATION
 
 
-TSharedPtr<SWidget> SGraphNode::GetEnabledStateWidget()
+TSharedPtr<SWidget> SGraphNode::GetEnabledStateWidget() const
 {
 	if ((GraphNode->GetDesiredEnabledState() != ENodeEnabledState::Enabled) && !GraphNode->IsAutomaticallyPlacedGhostNode())
 	{
diff --git a/Engine/Source/Editor/GraphEditor/Public/SGraphNode.h b/Engine/Source/Editor/GraphEditor/Public/SGraphNode.h
--- a/Engine/Source/Editor/GraphEditor/Public/SGraphNode.h
+++ b/Engine/Source/Editor/GraphEditor/Public/SGraphNode.h
@@ -168,22 +168,22 @@ public:
 	UE_API FText GetEditableNodeTitleAsText() const;
 
 	/** @return the tint for the node's title image */
-	UE_API FSlateColor GetNodeTitleColor() const;
+	UE_API virtual FSlateColor GetNodeTitleColor() const;
 
 	/** @return the tint for the node's comment */
 	UE_API FSlateColor GetNodeCommentColor() const;
 
 	/** @return the tint for the node's main body */
-	UE_API FSlateColor GetNodeBodyColor() const;
+	UE_API virtual FSlateColor GetNodeBodyColor() const;
 
 	/** @return the image brush to use for the node */
 	UE_API virtual const FSlateBrush * GetNodeBodyBrush() const;
 
 	/** @return the tint for the node's title icon */
-	UE_API FSlateColor GetNodeTitleIconColor() const;
+	UE_API virtual FSlateColor GetNodeTitleIconColor() const;
 
 	/** @return the tint for the node's title text */
-	UE_API FLinearColor GetNodeTitleTextColor() const;
+	UE_API virtual FLinearColor GetNodeTitleTextColor() const;
 
 	/** @return the tooltip to display when over the node */
 	UE_API FText GetNodeTooltip() const;
@@ -446,7 +446,7 @@ protected:
 	UE_API bool ShouldPinBeHidden(const UEdGraphPin* InPin) const;
 
 	/** Returns the widget to use for the enabled state of the node */
-	UE_API TSharedPtr<SWidget> GetEnabledStateWidget();
+	UE_API virtual TSharedPtr<SWidget> GetEnabledStateWidget() const;
 
 	/** Returns the widget to use for the corner icon of the node */
 	UE_API virtual TSharedRef<SWidget> CreateIconWidget();
