Fix context menu on subnodes' pins

#rnx

diff --git a/Engine/Source/Editor/GraphEditor/Private/GraphEditorSnapping/NodeSnappingManager.cpp b/Engine/Source/Editor/GraphEditor/Private/GraphEditorSnapping/NodeSnappingManager.cpp
--- a/Engine/Source/Editor/GraphEditor/Private/GraphEditorSnapping/NodeSnappingManager.cpp
+++ b/Engine/Source/Editor/GraphEditor/Private/GraphEditorSnapping/NodeSnappingManager.cpp
@@ -222,7 +222,7 @@ bool FNodeSnappingManager::BeginDragWithSnap(const FGeometry& MyGeometry, const
 
 	int32 NodeUnderMouseIndex = INDEX_NONE;
 	FVector2f MousePositionInNode = FVector2f::ZeroVector;
-	TSharedPtr<SGraphNode> NodeStartingDrag = FindNodeForMouseInteraction(MyGeometry, PointerEvent, MousePositionInNode);
+	TSharedPtr<SGraphNode> NodeStartingDrag = FindNodeForMouseInteraction(MyGeometry, PointerEvent, MousePositionInNode).Widget;
 		
 	if( !NodeStartingDrag || !NodeStartingDrag->CanBeSelected(MousePositionInNode) )
 	{
@@ -475,6 +475,12 @@ void FNodeSnappingManager::ForEachNode(TFunctionRef<void(const TSharedRef<SNodeP
 	}
 }
 
+FArrangedGraphNode FNodeSnappingManager::GetGraphNodeUnderMouse(const FGeometry& MyGeometry, const FPointerEvent& MouseEvent)
+{
+	FVector2f Discard;
+	return FindNodeForMouseInteraction(MyGeometry, MouseEvent, Discard);
+}
+
 bool FNodeSnappingManager::IsNodeSnappingEnabled()
 {
 	return UE::Private::CVarBPEnableNodeSnapping->GetBool();
@@ -791,13 +797,12 @@ void FNodeSnappingManager::UpdateSnapsFromModel(TSharedRef<SGraphNode> OnNode)
 	}
 }
 
-
-TSharedPtr<SGraphNode> FNodeSnappingManager::FindNodeForMouseInteraction( const FGeometry& MyGeometry, const FPointerEvent& PointerEvent, FVector2f& OutPositionInNode ) const
+FArrangedGraphNode FNodeSnappingManager::FindNodeForMouseInteraction( const FGeometry& MyGeometry, const FPointerEvent& PointerEvent, FVector2f& OutPositionInNode ) const
 {
 	TSharedPtr<SGraphPanel> OwningPanelPinned = OwningPanel.Pin();
 	if(!OwningPanelPinned)
 	{
-		return nullptr;
+		return FArrangedGraphNode();
 	}
 	FVector2f ArrangedSpacePosition = PointerEvent.GetScreenSpacePosition();
 	FArrangedChildren ArrangedChildren(EVisibility::Visible);
@@ -831,11 +836,11 @@ TSharedPtr<SGraphNode> FNodeSnappingManager::FindNodeForMouseInteraction( const
 				}
 			}
 
-			return NodeWidgetUnderMouse;
+			return {Candidate.Geometry, NodeWidgetUnderMouse};
 		}
 	}
 
-	return nullptr;
+	return FArrangedGraphNode();
 }
 
 TSharedPtr<SGraphNode> FNodeSnappingManager::FindNodeForSnap(TSharedPtr<SGraphNode> NodeBeingMoved)
diff --git a/Engine/Source/Editor/GraphEditor/Private/GraphEditorSnapping/NodeSnappingManager.h b/Engine/Source/Editor/GraphEditor/Private/GraphEditorSnapping/NodeSnappingManager.h
--- a/Engine/Source/Editor/GraphEditor/Private/GraphEditorSnapping/NodeSnappingManager.h
+++ b/Engine/Source/Editor/GraphEditor/Private/GraphEditorSnapping/NodeSnappingManager.h
@@ -5,6 +5,12 @@
 #include "SGraphPin.h"
 #include "SGraphNode.h"
 
+struct FArrangedGraphNode
+{
+	FGeometry Geometry;
+	TSharedPtr<SGraphNode> Widget;
+};
+
 struct FNodeSnappingManager
 {
 	FNodeSnappingManager(TWeakPtr<SGraphPanel> InOwningPanel);
@@ -14,6 +20,7 @@ struct FNodeSnappingManager
 	void UpdateDrag(const FGeometry& MyGeometry, const FPointerEvent& MouseEvent);
 	void HandleGraphChanged();
 	void ForEachNode(TFunctionRef<void(const TSharedRef<SNodePanel::SNode>&)> Ftor) const;
+	FArrangedGraphNode GetGraphNodeUnderMouse(const FGeometry& MyGeometry, const FPointerEvent& MouseEvent);
 
 	static bool IsNodeSnappingEnabled();
 private:
@@ -40,7 +47,7 @@ private:
 	static TArray<TSharedRef<SGraphNode>> GetControlFlowInputNodes(TSharedRef<SGraphNode> ForNode);
 	static TArray<TSharedRef<SGraphNode>> GetControlFlowOutputNodes(TSharedRef<SGraphNode> ForNode);
 	
-	TSharedPtr<SGraphNode> FindNodeForMouseInteraction( const FGeometry& MyGeometry, const FPointerEvent& PointerEvent, FVector2f& OutPositionInNode ) const;
+	FArrangedGraphNode FindNodeForMouseInteraction( const FGeometry& MyGeometry, const FPointerEvent& PointerEvent, FVector2f& OutPositionInNode ) const;
 	TSharedPtr<SGraphNode> FindNodeForSnap(TSharedPtr<SGraphNode> NodeBeingMoved);
 	TSharedPtr<SGraphPin> FindBestControlFlowForSnap(TSharedPtr<SGraphNode> NodeBeingMoved, UEdGraphPin* Pin);
 	TSharedPtr<SGraphPin> GetPreferredControlFlowSnapTarget(TSharedPtr<SGraphPin> A, TSharedPtr<SGraphPin> B) const;
diff --git a/Engine/Source/Editor/GraphEditor/Private/SGraphPanel.cpp b/Engine/Source/Editor/GraphEditor/Private/SGraphPanel.cpp
--- a/Engine/Source/Editor/GraphEditor/Private/SGraphPanel.cpp
+++ b/Engine/Source/Editor/GraphEditor/Private/SGraphPanel.cpp
@@ -1446,27 +1446,35 @@ bool SGraphPanel::PassesAssetReferenceFilter(const TArray<FAssetData>& Reference
 
 TSharedPtr<SGraphNode> SGraphPanel::GetGraphNodeUnderMouse(const FGeometry& MyGeometry, const FPointerEvent& MouseEvent)
 {
-	TSharedPtr<SGraphNode> GraphNode;
 	FArrangedChildren ArrangedNodes(EVisibility::Visible);
 	this->ArrangeChildrenForContextMenuSummon(MyGeometry, ArrangedNodes);
+	FArrangedGraphNode ArrangedWidget = {};
 	const int32 HoveredNodeIndex = SWidget::FindChildUnderMouse(ArrangedNodes, MouseEvent);
 	if (HoveredNodeIndex != INDEX_NONE)
 	{
 		const FArrangedWidget& HoveredNode = ArrangedNodes[HoveredNodeIndex];
-		GraphNode = StaticCastSharedRef<SGraphNode>(HoveredNode.Widget);
-		TSharedPtr<SGraphNode> GraphSubNode = GraphNode->GetNodeUnderMouse(HoveredNode.Geometry, MouseEvent);
-		GraphNode = GraphSubNode.IsValid() ? GraphSubNode.ToSharedRef() : GraphNode;
-
+		ArrangedWidget.Geometry = HoveredNode.Geometry;
+		ArrangedWidget.Widget = StaticCastSharedRef<SGraphNode>(HoveredNode.Widget);
+	}
+	else if (FNodeSnappingManager* SnapManager = GetSnappingManager())
+	{
+		ArrangedWidget = SnapManager->GetGraphNodeUnderMouse(MyGeometry, MouseEvent);
+	}
+	
+	if (ArrangedWidget.Widget)
+	{
 		// Selection should switch to this code if it isn't already selected.
 		// When multiple nodes are selected, we do nothing, provided that the
 		// node for which the context menu is being created is in the selection set.
-		if (!SelectionManager.IsNodeSelected(GraphNode->GetObjectBeingDisplayed()))
+		TSharedPtr<SGraphNode> GraphSubNode = ArrangedWidget.Widget->GetNodeUnderMouse(ArrangedWidget.Geometry, MouseEvent);
+		ArrangedWidget.Widget = GraphSubNode.IsValid() ? GraphSubNode.ToSharedRef() : ArrangedWidget.Widget;
+		if (!SelectionManager.IsNodeSelected(ArrangedWidget.Widget->GetObjectBeingDisplayed()))
 		{
-			SelectionManager.SelectSingleNode(GraphNode->GetObjectBeingDisplayed());
+			SelectionManager.SelectSingleNode(ArrangedWidget.Widget->GetObjectBeingDisplayed());
 		}
 	}
 
-	return GraphNode;
+	return ArrangedWidget.Widget;
 }
 
 UEdGraphPin* SGraphPanel::GetPinUnderMouse(const FGeometry& MyGeometry, const FPointerEvent& MouseEvent, TSharedPtr<SGraphNode> GraphNode) const
