Fix delete of subnodes

#rnx

diff --git a/Engine/Source/Editor/GraphEditor/Private/GraphEditorSnapping/NodeSnappingManager.cpp b/Engine/Source/Editor/GraphEditor/Private/GraphEditorSnapping/NodeSnappingManager.cpp
--- a/Engine/Source/Editor/GraphEditor/Private/GraphEditorSnapping/NodeSnappingManager.cpp
+++ b/Engine/Source/Editor/GraphEditor/Private/GraphEditorSnapping/NodeSnappingManager.cpp
@@ -207,6 +207,18 @@ FNodeSnappingManager::FNodeSnappingManager(TWeakPtr<SGraphPanel> InOwningPanel)
 	check(InOwningPanel.Pin());
 }
 
+void FNodeSnappingManager::RemoveNode(const UEdGraphNode* Node)
+{
+	TSharedRef<SGraphPanel> OwningPanelPinned = OwningPanel.Pin().ToSharedRef();
+	TSharedPtr<SGraphNode> GraphWidget = GetSGraphNode(OwningPanelPinned, Node);
+	if(!GraphWidget)
+	{
+		return;
+	}
+
+	ClearAllSnaps(GraphWidget.ToSharedRef());
+}
+
 bool FNodeSnappingManager::IsDragging() const
 {
 	return bIsActivelyDragging;
diff --git a/Engine/Source/Editor/GraphEditor/Private/GraphEditorSnapping/NodeSnappingManager.h b/Engine/Source/Editor/GraphEditor/Private/GraphEditorSnapping/NodeSnappingManager.h
--- a/Engine/Source/Editor/GraphEditor/Private/GraphEditorSnapping/NodeSnappingManager.h
+++ b/Engine/Source/Editor/GraphEditor/Private/GraphEditorSnapping/NodeSnappingManager.h
@@ -14,6 +14,8 @@ struct FArrangedGraphNode
 struct FNodeSnappingManager
 {
 	FNodeSnappingManager(TWeakPtr<SGraphPanel> InOwningPanel);
+
+	void RemoveNode(const UEdGraphNode* Node);
 	bool IsDragging() const;
 	bool BeginDragWithSnap(const FGeometry& MyGeometry, const FPointerEvent& PointerEvent);
 	bool FinalizeNodeMovements();
diff --git a/Engine/Source/Editor/GraphEditor/Private/SGraphPanel.cpp b/Engine/Source/Editor/GraphEditor/Private/SGraphPanel.cpp
--- a/Engine/Source/Editor/GraphEditor/Private/SGraphPanel.cpp
+++ b/Engine/Source/Editor/GraphEditor/Private/SGraphPanel.cpp
@@ -2259,6 +2259,11 @@ void SGraphPanel::AddNode(UEdGraphNode* Node, AddNodeBehavior Behavior)
 
 void SGraphPanel::RemoveNode(const UEdGraphNode* Node)
 {
+	if(NodeSnappingManager)
+	{
+		NodeSnappingManager->RemoveNode(Node);
+	}
+
 	for (int32 Iter = 0; Iter != Children.Num(); ++Iter)
 	{
 		TSharedRef<SGraphNode> Child = GetChild(Iter);
