Quick fix for node delete path that performs redundant removals (e.g. SGraphPanel::RemoveAllNodesWithInvalidPointers)

#rnx

diff --git a/Engine/Source/Editor/GraphEditor/Private/GraphEditorSnapping/NodeSnappingManager.cpp b/Engine/Source/Editor/GraphEditor/Private/GraphEditorSnapping/NodeSnappingManager.cpp
--- a/Engine/Source/Editor/GraphEditor/Private/GraphEditorSnapping/NodeSnappingManager.cpp
+++ b/Engine/Source/Editor/GraphEditor/Private/GraphEditorSnapping/NodeSnappingManager.cpp
@@ -209,14 +209,11 @@ FNodeSnappingManager::FNodeSnappingManager(TWeakPtr<SGraphPanel> InOwningPanel)
 
 void FNodeSnappingManager::RemoveNode(const UEdGraphNode* Node)
 {
-	TSharedRef<SGraphPanel> OwningPanelPinned = OwningPanel.Pin().ToSharedRef();
-	TSharedPtr<SGraphNode> GraphWidget = GetSGraphNode(OwningPanelPinned, Node);
-	if(!GraphWidget)
+	TSharedPtr<SGraphNode> GraphWidget = TryGetGraphNode(Node);
+	if(GraphWidget)
 	{
-		return;
+		ClearAllSnaps(GraphWidget.ToSharedRef());
 	}
-
-	ClearAllSnaps(GraphWidget.ToSharedRef());
 }
 
 bool FNodeSnappingManager::IsDragging() const
@@ -964,10 +961,21 @@ TSharedPtr<SGraphPin> FNodeSnappingManager::GetPreferredControlFlowSnapTarget(TS
 TSharedRef<SGraphNode> FNodeSnappingManager::GetSGraphNode(TSharedRef<SGraphPanel> FromPanel, const UEdGraphNode* Node)
 {
 	const TSharedRef<SNodePanel::SNode>* Widget = FromPanel->NodeToWidgetLookup.Find(Node);
-	check(Widget != nullptr);
+	check(Widget != nullptr); // this will return null for some delete operations - which run redundant removals - consider TryGetGraphNode
 	return StaticCastSharedRef<SGraphNode>(*Widget);
 }
 
+TSharedPtr<SGraphNode> FNodeSnappingManager::TryGetGraphNode(const UEdGraphNode* Node)
+{
+	TSharedRef<SGraphPanel> OwningPanelPinned = OwningPanel.Pin().ToSharedRef();
+	const TSharedRef<SNodePanel::SNode>* Widget = OwningPanelPinned->NodeToWidgetLookup.Find(Node);
+	if(Widget)
+	{
+		return StaticCastSharedRef<SGraphNode>(*Widget);
+	}
+	return TSharedPtr<SGraphNode>();
+}
+
 void FNodeSnappingManager::ModifySafe(UEdGraphPin* Pin)
 {
 	UEdGraphNode* Node = Pin->GetOwningNode();
diff --git a/Engine/Source/Editor/GraphEditor/Private/GraphEditorSnapping/NodeSnappingManager.h b/Engine/Source/Editor/GraphEditor/Private/GraphEditorSnapping/NodeSnappingManager.h
--- a/Engine/Source/Editor/GraphEditor/Private/GraphEditorSnapping/NodeSnappingManager.h
+++ b/Engine/Source/Editor/GraphEditor/Private/GraphEditorSnapping/NodeSnappingManager.h
@@ -54,6 +54,7 @@ private:
 	TSharedPtr<SGraphPin> FindBestControlFlowForSnap(TSharedPtr<SGraphNode> NodeBeingMoved, UEdGraphPin* Pin);
 	TSharedPtr<SGraphPin> GetPreferredControlFlowSnapTarget(TSharedPtr<SGraphPin> A, TSharedPtr<SGraphPin> B) const;
 	static TSharedRef<SGraphNode> GetSGraphNode(TSharedRef<SGraphPanel> FromPanel, const UEdGraphNode* Node);
+	TSharedPtr<SGraphNode> TryGetGraphNode(const UEdGraphNode* Node);
 	void ModifySafe(UEdGraphPin* Pin);
 
 	TWeakPtr<SGraphPanel> OwningPanel; // must be weak, otherwise we'll have a cycle, never meant to be null
