Include directly snapped nodes in the transaction buffer when dragging

#rnx

diff --git a/Engine/Source/Editor/GraphEditor/Private/GraphEditorSnapping/NodeSnappingManager.cpp b/Engine/Source/Editor/GraphEditor/Private/GraphEditorSnapping/NodeSnappingManager.cpp
--- a/Engine/Source/Editor/GraphEditor/Private/GraphEditorSnapping/NodeSnappingManager.cpp
+++ b/Engine/Source/Editor/GraphEditor/Private/GraphEditorSnapping/NodeSnappingManager.cpp
@@ -197,6 +197,25 @@ bool ConnectAndSnap(UEdGraphPin* PinA, UEdGraphPin* PinB)
 	FirstPin->SetSnappedChild(SecondPin);
 	return true;
 }
+
+void ForEachDirectlySnappedNode(UEdGraphNode* Node, TFunctionRef<void(UEdGraphNode*)> FTor)
+{
+	for (const UEdGraphPin* Pin : Node->Pins)
+	{
+		UEdGraphPin* SnappedChildPin = Pin->GetSnappedChildPin();
+		if(SnappedChildPin)
+		{
+			FTor(SnappedChildPin->GetOwningNode());
+		}
+		
+		UEdGraphPin* SnappedParentPin = Pin->GetSnappedParentPin();
+		if(SnappedParentPin)
+		{
+			FTor(SnappedParentPin->GetOwningNode());
+		}
+	}
+}
+
 }
 
 FNodeSnappingManager::FNodeSnappingManager(TWeakPtr<SGraphPanel> InOwningPanel)
@@ -291,6 +310,7 @@ bool FNodeSnappingManager::BeginDragWithSnap(const FGeometry& MyGeometry, const
 		SelectedNodes.Num() > 1 
 		? NSLOCTEXT("GraphEditor", "MoveNodesAction", "Move Nodes")
 		: NSLOCTEXT("GraphEditor", "MoveNodeAction", "Move Node"));
+	
 	for(TObjectPtr<UObject> Node : SelectedNodes)
 	{
 		UEdGraphNode* AsEdGraphNode = CastChecked<UEdGraphNode>(Node);
@@ -298,6 +318,25 @@ bool FNodeSnappingManager::BeginDragWithSnap(const FGeometry& MyGeometry, const
 		MovingNodes.Add(AsEdGraphNode);
 		Node->Modify();
 	}
+		
+	for(TObjectPtr<UObject> Node : SelectedNodes)
+	{
+		UEdGraphNode* AsEdGraphNode = CastChecked<UEdGraphNode>(Node);
+		// we'll also need to modify any nodes that are snapped to this node
+		// they may be implicitly mutated by mutating these nodes
+		FNodeSnappingManager* Self = this;
+		UE::Private::ForEachDirectlySnappedNode(
+			AsEdGraphNode,
+			[Self](UEdGraphNode* Node)
+			{
+				if(!Self->ModifiedNodes.Contains(Node))
+				{
+					Self->ModifiedNodes.Add(Node);
+					Node->Modify();
+				}
+			}
+		);
+	}
 
 	PendingSnaps.SetNum(SelectedNodes.Num());
 	int32 Idx = 0;
