Cosmetic fixes to allow nodes to filll their snapped container, if any. I believe this has no impact on unsnapped nodes

#rnx

diff --git a/Engine/Source/Editor/GraphEditor/Private/KismetNodes/SGraphNodeK2Base.cpp b/Engine/Source/Editor/GraphEditor/Private/KismetNodes/SGraphNodeK2Base.cpp
--- a/Engine/Source/Editor/GraphEditor/Private/KismetNodes/SGraphNodeK2Base.cpp
+++ b/Engine/Source/Editor/GraphEditor/Private/KismetNodes/SGraphNodeK2Base.cpp
@@ -257,7 +257,7 @@ void SGraphNodeK2Base::UpdateCompactNode()
 		];
 
 	this->GetOrAddSlot( ENodeZone::Center )
-	.HAlign(HAlign_Center)
+	.HAlign(HAlign_Fill)
 	.VAlign(VAlign_Center)
 	[
 		InnerVerticalBox
diff --git a/Engine/Source/Editor/GraphEditor/Private/KismetNodes/SGraphNodeK2Composite.cpp b/Engine/Source/Editor/GraphEditor/Private/KismetNodes/SGraphNodeK2Composite.cpp
--- a/Engine/Source/Editor/GraphEditor/Private/KismetNodes/SGraphNodeK2Composite.cpp
+++ b/Engine/Source/Editor/GraphEditor/Private/KismetNodes/SGraphNodeK2Composite.cpp
@@ -82,7 +82,7 @@ void SGraphNodeK2Composite::UpdateGraphNode()
 	//
 	this->ContentScale.Bind( this, &SGraphNode::GetContentScale );
 	this->GetOrAddSlot( ENodeZone::Center )
-		.HAlign(HAlign_Center)
+		.HAlign(HAlign_Fill)
 		.VAlign(VAlign_Center)
 		[
 			SNew(SBorder)
diff --git a/Engine/Source/Editor/GraphEditor/Private/KismetNodes/SGraphNodeK2Copy.cpp b/Engine/Source/Editor/GraphEditor/Private/KismetNodes/SGraphNodeK2Copy.cpp
--- a/Engine/Source/Editor/GraphEditor/Private/KismetNodes/SGraphNodeK2Copy.cpp
+++ b/Engine/Source/Editor/GraphEditor/Private/KismetNodes/SGraphNodeK2Copy.cpp
@@ -226,7 +226,7 @@ void SGraphNodeK2Copy::UpdateGraphNode()
 	SetupErrorReporting();
 
 	GetOrAddSlot(ENodeZone::Center)
-		.HAlign(HAlign_Center)
+		.HAlign(HAlign_Fill)
 		.VAlign(VAlign_Center)
 		[
 			SNew(SVerticalBox)
diff --git a/Engine/Source/Editor/GraphEditor/Private/KismetNodes/SGraphNodeK2Terminator.cpp b/Engine/Source/Editor/GraphEditor/Private/KismetNodes/SGraphNodeK2Terminator.cpp
--- a/Engine/Source/Editor/GraphEditor/Private/KismetNodes/SGraphNodeK2Terminator.cpp
+++ b/Engine/Source/Editor/GraphEditor/Private/KismetNodes/SGraphNodeK2Terminator.cpp
@@ -79,7 +79,7 @@ void SGraphNodeK2Terminator::UpdateGraphNode()
 	//
 	this->ContentScale.Bind( this, &SGraphNode::GetContentScale );
 	this->GetOrAddSlot( ENodeZone::Center )
-	.HAlign(HAlign_Center)
+	.HAlign(HAlign_Fill)
 	.VAlign(VAlign_Center)
 	[
 		SNew(SVerticalBox)
diff --git a/Engine/Source/Editor/GraphEditor/Private/KismetNodes/SGraphNodeK2Var.cpp b/Engine/Source/Editor/GraphEditor/Private/KismetNodes/SGraphNodeK2Var.cpp
--- a/Engine/Source/Editor/GraphEditor/Private/KismetNodes/SGraphNodeK2Var.cpp
+++ b/Engine/Source/Editor/GraphEditor/Private/KismetNodes/SGraphNodeK2Var.cpp
@@ -244,7 +244,7 @@ void SGraphNodeK2Var::UpdateGraphNode()
 	//
 	this->ContentScale.Bind( this, &SGraphNode::GetContentScale );
 	this->GetOrAddSlot( ENodeZone::Center )
-	.HAlign(HAlign_Center)
+	.HAlign(HAlign_Fill)
 	.VAlign(VAlign_Center)
 	[
 		SNew(SVerticalBox)
