Create ControlFlowSequence vertical box on demand for snapping, this allows control flow snapping to work with nodes with custom layout (e.g. set variable nodes) also reduces risk of behavior change from my previous tweak to SGraphNode.cpp

#rnx

diff --git a/Engine/Source/Editor/GraphEditor/Private/GraphEditorSnapping/NodeSnappingManager.cpp b/Engine/Source/Editor/GraphEditor/Private/GraphEditorSnapping/NodeSnappingManager.cpp
--- a/Engine/Source/Editor/GraphEditor/Private/GraphEditorSnapping/NodeSnappingManager.cpp
+++ b/Engine/Source/Editor/GraphEditor/Private/GraphEditorSnapping/NodeSnappingManager.cpp
@@ -656,11 +656,6 @@ void FNodeSnappingManager::SnapControlFlowPeer(FPendingSnap& SnapState, TSharedR
 
 	TSharedRef<SGraphNode> StartingNode = (PinToSnapTo->GetDirection() ==  EEdGraphPinDirection::EGPD_Input) ? 
 		NodeBeingDragged : NodeDroppedOn;
-		
-	if(!StartingNode->ControlFlowSequence)
-	{
-		return; // @todo: overcome SGraphNode copy paste - we can't snap to this
-	}
 	
 	// update the data model.. this is all we should really do here - but I have made some mistakes
 	ModifySafe(OtherPin);
@@ -821,6 +816,21 @@ void FNodeSnappingManager::UpdateSnapsFromModel(TSharedRef<SGraphNode> OnNode)
 				ChildNode->MoveTo(NodePosition, NodeFilter, false);
 			}
 
+			SNodePanel::SNode::FNodeSlot* CenterSlot = ParentNode->GetSlot(ENodeZone::Center);
+			if(!ParentNode->ControlFlowSequence.IsValid())
+			{
+				ParentNode->ControlFlowSequence = SNew(SVerticalBox)
+				 +SVerticalBox::Slot()
+				.AutoHeight()
+				[
+					CenterSlot->GetWidget()
+				];
+				ParentNode->GetOrAddSlot(ENodeZone::Center)
+					[
+						ParentNode->ControlFlowSequence.ToSharedRef()
+					];
+			}
+
 			ParentNode->ControlFlowSequence->AddSlot()
 				.AutoHeight()
 				[
diff --git a/Engine/Source/Editor/GraphEditor/Private/SGraphNode.cpp b/Engine/Source/Editor/GraphEditor/Private/SGraphNode.cpp
--- a/Engine/Source/Editor/GraphEditor/Private/SGraphNode.cpp
+++ b/Engine/Source/Editor/GraphEditor/Private/SGraphNode.cpp
@@ -859,6 +859,7 @@ void SGraphNode::UpdateGraphNode()
 	//            |       |      | T (>) |
 	//            |_______|______|_______|
 	//
+	TSharedPtr<SVerticalBox> MainVerticalBox;
 	SetupErrorReporting();
 
 	TSharedPtr<SNodeTitle> NodeTitle = SNew(SNodeTitle, GraphNode);
@@ -1013,7 +1014,7 @@ void SGraphNode::UpdateGraphNode()
 		.HAlign(HAlign_Fill)
 		.VAlign(VAlign_Center)
 		[
-			SAssignNew(ControlFlowSequence, SVerticalBox)
+			SAssignNew(MainVerticalBox, SVerticalBox)
 			+SVerticalBox::Slot()
 			.AutoHeight()
 			[
@@ -1067,7 +1068,7 @@ void SGraphNode::UpdateGraphNode()
 			];
 	}
 
-	CreateBelowWidgetControls(ControlFlowSequence);
+	CreateBelowWidgetControls(MainVerticalBox);
 	CreatePinWidgets();
 	CreateInputSideAddButton(LeftNodeBox);
 	CreateOutputSideAddButton(RightNodeBox);
