Flesh out connecting nodes by dragging them together when using node snapping

#rnx

diff --git a/Engine/Source/Editor/GraphEditor/Private/GraphEditorSnapping/NodeSnappingManager.cpp b/Engine/Source/Editor/GraphEditor/Private/GraphEditorSnapping/NodeSnappingManager.cpp
--- a/Engine/Source/Editor/GraphEditor/Private/GraphEditorSnapping/NodeSnappingManager.cpp
+++ b/Engine/Source/Editor/GraphEditor/Private/GraphEditorSnapping/NodeSnappingManager.cpp
@@ -154,18 +154,17 @@ bool ConnectAndSnap(UEdGraphPin* PinA, UEdGraphPin* PinB)
 {
 	check(PinA && PinB);
 
-	// @todo: if the pins are snapped to anything unsnap them - do i need symmetrical 
-	// pin snaps?
-
 	// if the pins are not connected then connect them:
+	bool bMadeNewConnection = false;
 	if(!PinA->LinkedTo.Contains(PinB))
 	{
-		GetDefault<UEdGraphSchema_K2>()->TryCreateConnection(PinA, PinB);
+		bMadeNewConnection = GetDefault<UEdGraphSchema_K2>()->TryCreateConnection(PinA, PinB);
 	}
 	
-	if(!PinA->LinkedTo.Contains(PinB))
+	if(bMadeNewConnection && !PinA->LinkedTo.Contains(PinB))
 	{
-		return false;
+		ensure(false); // TryCreateConnection lied - relatively complicated function so lets be defensive
+		bMadeNewConnection = false;
 	}
 
 	// execution pins should snap in execution order - data pins
@@ -191,11 +190,12 @@ bool ConnectAndSnap(UEdGraphPin* PinA, UEdGraphPin* PinB)
 
 	if(FirstPin->GetSnappedChildPin() == SecondPin)
 	{
-		return true;
+		// already connected, nothing more to do
+		return bMadeNewConnection;
 	}
 
 	FirstPin->SetSnappedChild(SecondPin);
-	return true;
+	return bMadeNewConnection;
 }
 
 void ForEachDirectlySnappedNode(UEdGraphNode* Node, TFunctionRef<void(UEdGraphNode*)> FTor)
@@ -393,6 +393,18 @@ void FNodeSnappingManager::UpdateDrag(const FGeometry& MyGeometry, const FPointe
 	const bool bCursorInDeadZone = TotalMouseDelta <= FSlateApplication::Get().GetDragTriggerDistance();
 	
 	// reset snap data for the nodes we're moving, we want to recalculate it:
+	for(const FPendingConnection& UncommittedConnection : UncommittedConnections)
+	{
+		UEdGraphPin* DraggingPinResolved = UncommittedConnection.DraggingPin.Get();
+		UEdGraphPin* StablePinResolved = UncommittedConnection.StablePin.Get();
+
+		if(DraggingPinResolved && StablePinResolved)
+		{
+			GetSchema()->BreakSinglePinLink(DraggingPinResolved, StablePinResolved);
+		}
+	}
+	UncommittedConnections.Reset();
+
 	for(UEdGraphNode* Node : MovingNodes)
 	{
 		// reset snapped state - a moving node
@@ -549,6 +561,7 @@ void FNodeSnappingManager::ResetDragManager()
 	RootDragNode = nullptr;
 	MovingNodes.Reset();
 	ModifiedNodes.Reset();
+	UncommittedConnections.Reset();
 	TotalMouseDelta = 0.f;
 	bIsActivelyDragging = false;
 }
@@ -595,14 +608,21 @@ bool FNodeSnappingManager::TryToSnapDataflow(FPendingSnap& PendingSnap)
 	{
 		return false;
 	}
+	UEdGraphPin* DestGraphPin = DestPin->GetPinObj();
 	const FPinConnectionResponse Response = GetSchema()->CanCreateConnection(
-		DestPin->GetPinObj(), PinToMatchForSnap);
+		DestGraphPin, PinToMatchForSnap);
 	if( Response.Response == ECanCreateConnectionResponse::CONNECT_RESPONSE_DISALLOW )
 	{
 		return false;
 	}
 
-	return TrySnapTo(DestPin.ToSharedRef(), PendingSnap, PinToMatchForSnap);
+	bool bCreatedConnection = TrySnapTo(DestPin.ToSharedRef(), PendingSnap, PinToMatchForSnap);
+	if(bCreatedConnection)
+	{
+		UncommittedConnections.Add({DestGraphPin, PinToMatchForSnap});
+	}
+
+	return true;
 }
 
 bool FNodeSnappingManager::TryToSnapControlflow(FPendingSnap& PendingSnap)
@@ -642,8 +662,7 @@ bool FNodeSnappingManager::TrySnapTo(TSharedRef<SGraphPin> PinToSnapTo, FPending
 	UEdGraphPin* OtherPin = PinToSnapTo->GetPinObj();
 	check(OtherPin && OtherPin != PinOnMovingNode);
 	ModifySafe(OtherPin);
-	UE::Private::ConnectAndSnap(PinOnMovingNode, OtherPin);
-	return true;
+	return UE::Private::ConnectAndSnap(PinOnMovingNode, OtherPin);
 }
 
 void FNodeSnappingManager::SnapControlFlowPeer(FPendingSnap& SnapState, TSharedRef<SGraphPin> PinToSnapTo, UEdGraphPin* PinOnMovingNode)
diff --git a/Engine/Source/Editor/GraphEditor/Private/GraphEditorSnapping/NodeSnappingManager.h b/Engine/Source/Editor/GraphEditor/Private/GraphEditorSnapping/NodeSnappingManager.h
--- a/Engine/Source/Editor/GraphEditor/Private/GraphEditorSnapping/NodeSnappingManager.h
+++ b/Engine/Source/Editor/GraphEditor/Private/GraphEditorSnapping/NodeSnappingManager.h
@@ -66,6 +66,17 @@ private:
 	TSet<UEdGraphNode*> MovingNodes;
 	TSet<UEdGraphNode*> ModifiedNodes; // track modification to make sure we don't modify after mutating
 
+	// A list of uncommitted connections, which will be created only if the final state of the 
+	// drag operation leaves them in connected positions. This is meant to facilitate connecting
+	// nodes by dragging them on top of eachother/compatible pins. We use FEdGraphPinReference
+	// for persistent references to visible pins across node refreshes:
+	struct FPendingConnection
+	{
+		FEdGraphPinReference DraggingPin;
+		FEdGraphPinReference StablePin;
+	};
+	TArray<FPendingConnection> UncommittedConnections;
+
 	float TotalMouseDelta;
 
 	bool bIsActivelyDragging;
diff --git a/Engine/Source/Editor/GraphEditor/Private/SGraphPanel.cpp b/Engine/Source/Editor/GraphEditor/Private/SGraphPanel.cpp
--- a/Engine/Source/Editor/GraphEditor/Private/SGraphPanel.cpp
+++ b/Engine/Source/Editor/GraphEditor/Private/SGraphPanel.cpp
@@ -1100,7 +1100,7 @@ void SGraphPanel::GetAllPins(TSet< TSharedRef<SWidget> >& AllPins)
 	for (int32 ChildIndex = 0; ChildIndex < Children.Num(); ++ChildIndex)
 	{
 		TSharedRef<SGraphNode> ChildNode = StaticCastSharedRef<SGraphNode>(Children[ChildIndex]);
-		ChildNode->GetPins(AllPins);
+		ChildNode->GetPinsIncludingSnappedChildren(AllPins);
 	}
 }
 
diff --git a/Engine/Source/Runtime/Engine/Private/EdGraph/EdGraphPin.cpp b/Engine/Source/Runtime/Engine/Private/EdGraph/EdGraphPin.cpp
--- a/Engine/Source/Runtime/Engine/Private/EdGraph/EdGraphPin.cpp
+++ b/Engine/Source/Runtime/Engine/Private/EdGraph/EdGraphPin.cpp
@@ -676,6 +676,20 @@ void UEdGraphPin::BreakLinkTo(UEdGraphPin* ToPin, bool bAlwaysMarkDirty)
 				// Ensure that the other pin links to us but ignore it if our parent is invalid or being destroyed
 				ensureAlwaysMsgf(ToPin->LinkedTo.Contains(this), TEXT("%s"), *GetLinkInfoString(LOCTEXT("BreakLinkTo", "BreakLinkTo").ToString(), LOCTEXT("NotLinked", "not reciprocally linked with pin").ToString(), ToPin));
 			}
+			
+#if WITH_EDITORONLY_DATA
+			// update snapping state - if any:
+			if(bHasSnappedParent && ToPin->bHasSnappedChild)
+			{
+				ToPin->bHasSnappedChild = false;
+				bHasSnappedParent = false;
+			}
+			else if(bHasSnappedChild && ToPin->bHasSnappedParent)
+			{
+				bHasSnappedChild = false;
+				ToPin->bHasSnappedParent = false;
+			}
+#endif//WITH_EDITORONLY_DATA
 
 			LinkedTo.Remove(ToPin);
 		}
