Automatically snap control flow pins when we overlap a node with free compatible pins, clean up a helper function reused between control flow and data snapping

#rnx

diff --git a/Engine/Source/Editor/GraphEditor/Private/GraphEditorSnapping/NodeSnappingManager.cpp b/Engine/Source/Editor/GraphEditor/Private/GraphEditorSnapping/NodeSnappingManager.cpp
--- a/Engine/Source/Editor/GraphEditor/Private/GraphEditorSnapping/NodeSnappingManager.cpp
+++ b/Engine/Source/Editor/GraphEditor/Private/GraphEditorSnapping/NodeSnappingManager.cpp
@@ -96,25 +96,30 @@ UEdGraphPin* FindPinConnectedTo(UEdGraphNode* Node, UEdGraphPin* Pin)
 	}
 	return nullptr;
 }
-UEdGraphPin* CanBeConnectedTo(UEdGraphNode* Node, UEdGraphPin* Pin)
+
+UEdGraphPin* FindFreeCompatiblePin(UEdGraphNode* OnNode, const UEdGraphPin* WithPin)
 {
-	check(Node && Pin);
-	if(Pin->LinkedTo.Num() > 0)
+	if(WithPin->LinkedTo.Num() > 0)
 	{
 		return nullptr;
 	}
+	
+	const UEdGraphSchema* Schema = OnNode->GetSchema();
 
-	const UEdGraphSchema* Schema = Node->GetSchema();
-
-	for(UEdGraphPin* PinOnNode : Node->Pins)
+	for(UEdGraphPin* PinOnNode : OnNode->Pins)
 	{		
 		if(PinOnNode->bHidden)
 		{
 			continue;
 		}
 
+		if(PinOnNode->LinkedTo.Num() > 0)
+		{
+			continue;
+		}
+
 		const FPinConnectionResponse Response = Schema->CanCreateConnection(
-			PinOnNode, Pin);
+			PinOnNode, WithPin);
 		if( Response.Response == ECanCreateConnectionResponse::CONNECT_RESPONSE_MAKE )
 		{
 			return PinOnNode;
@@ -616,7 +621,7 @@ bool FNodeSnappingManager::TryToSnapDataflow(FPendingSnap& PendingSnap)
 		return false;
 	}
 
-	bool bCreatedConnection = TrySnapTo(DestPin.ToSharedRef(), PendingSnap, PinToMatchForSnap);
+	bool bCreatedConnection = TrySnapTo(DestPin.ToSharedRef()->GetPinObj(), PinToMatchForSnap);
 	if(bCreatedConnection)
 	{
 		UncommittedConnections.Add({DestGraphPin, PinToMatchForSnap});
@@ -657,9 +662,8 @@ bool FNodeSnappingManager::TryToSnapControlflow(FPendingSnap& PendingSnap)
 	return true;
 }
 
-bool FNodeSnappingManager::TrySnapTo(TSharedRef<SGraphPin> PinToSnapTo, FPendingSnap& SnapState, UEdGraphPin* PinOnMovingNode)
+bool FNodeSnappingManager::TrySnapTo(UEdGraphPin* OtherPin, UEdGraphPin* PinOnMovingNode)
 {
-	UEdGraphPin* OtherPin = PinToSnapTo->GetPinObj();
 	check(OtherPin && OtherPin != PinOnMovingNode);
 	ModifySafe(OtherPin);
 	return UE::Private::ConnectAndSnap(PinOnMovingNode, OtherPin);
@@ -677,8 +681,11 @@ void FNodeSnappingManager::SnapControlFlowPeer(FPendingSnap& SnapState, TSharedR
 		NodeBeingDragged : NodeDroppedOn;
 	
 	// update the data model.. this is all we should really do here - but I have made some mistakes
-	ModifySafe(OtherPin);
-	UE::Private::ConnectAndSnap(PinOnMovingNode, OtherPin);
+	const bool bCreatedConnection = TrySnapTo(PinOnMovingNode, OtherPin);
+	if(bCreatedConnection)
+	{
+		UncommittedConnections.Add({PinOnMovingNode, OtherPin});
+	}
 
 	SnapState.PinSnappedTo = PinToSnapTo;
 }
@@ -1001,8 +1008,16 @@ TSharedPtr<SGraphPin> FNodeSnappingManager::FindBestControlFlowForSnap(TSharedPt
 	//  1. we're already attached
 	// or
 	//  2. we're not attached to anything and the node's execution pin also isn't: (@todo)
-	UEdGraphNode* EdGraphNodeBeingMoved = UE::Private::GetEdGraphNode(Node.ToSharedRef());
-	UEdGraphPin* PinToSnap = UE::Private::FindPinConnectedTo(EdGraphNodeBeingMoved, Pin);
+	UEdGraphNode* EdGraphNodeBeingSnappedTo = UE::Private::GetEdGraphNode(Node.ToSharedRef());
+	UEdGraphPin* PinToSnap = UE::Private::FindPinConnectedTo(EdGraphNodeBeingSnappedTo, Pin);
+
+	if(!PinToSnap)
+	{
+		// we're not connected yet, but lets see if EdGraphNodeBeingSnappedTo has
+		// a compatible pin that is free:
+		PinToSnap = UE::Private::FindFreeCompatiblePin(EdGraphNodeBeingSnappedTo, Pin);
+	}
+
 	if(!PinToSnap)
 	{
 		return nullptr;
diff --git a/Engine/Source/Editor/GraphEditor/Private/GraphEditorSnapping/NodeSnappingManager.h b/Engine/Source/Editor/GraphEditor/Private/GraphEditorSnapping/NodeSnappingManager.h
--- a/Engine/Source/Editor/GraphEditor/Private/GraphEditorSnapping/NodeSnappingManager.h
+++ b/Engine/Source/Editor/GraphEditor/Private/GraphEditorSnapping/NodeSnappingManager.h
@@ -39,7 +39,7 @@ private:
 	const UEdGraphSchema* GetSchema() const;
 	bool TryToSnapDataflow(FPendingSnap& PendingSnap);
 	bool TryToSnapControlflow(FPendingSnap& PendingSnap);
-	bool TrySnapTo(TSharedRef<SGraphPin> PinToSnapTo, FPendingSnap& SnapState, UEdGraphPin* PinOnMovingNode);
+	bool TrySnapTo(UEdGraphPin* OtherPin, UEdGraphPin* PinOnMovingNode);
 	void SnapControlFlowPeer(FPendingSnap& SnapState, TSharedRef<SGraphPin> PinToSnapTo, UEdGraphPin* PinOnMovingNode);
 	
 	void ClearAllSnaps(TSharedRef<SGraphNode> OnNode);
