Fixes for nested data pin snapping, rename the widget used for snapping data nodes and create it on demand

#rnx

diff --git a/Engine/Source/Editor/GraphEditor/Private/GraphEditorSnapping/NodeSnappingManager.cpp b/Engine/Source/Editor/GraphEditor/Private/GraphEditorSnapping/NodeSnappingManager.cpp
--- a/Engine/Source/Editor/GraphEditor/Private/GraphEditorSnapping/NodeSnappingManager.cpp
+++ b/Engine/Source/Editor/GraphEditor/Private/GraphEditorSnapping/NodeSnappingManager.cpp
@@ -48,22 +48,13 @@ UEdGraphPin* GetOnlyVisiblePinWithDirection( UEdGraphNode* Node, EEdGraphPinDire
 	return Result;
 }
 
-UEdGraphPin* GetFirstAndOnlyVisiblePin(UEdGraphNode* ForNode)
+UEdGraphPin* GetFirstAndOnlyVisibleOutputPin(UEdGraphNode* ForNode)
 {
-	const TArray<UEdGraphPin*>& Pins = ForNode->GetAllPins();
-	UEdGraphPin* Result = nullptr;
-	for(UEdGraphPin* Pin : Pins)
-	{
-		if(!Pin->bHidden)
-		{
-			if(Result)
-			{
-				return nullptr;
-			}
-			Result = Pin;
-		}
-	}
-	return Result;
+	return GetOnlyVisiblePinWithDirection(
+		ForNode, 
+		EEdGraphPinDirection::EGPD_Output, 
+		[](UEdGraphPin* Pin) { return true; }
+	);
 }
 
 UEdGraphPin* GetFirstAndOnlyVisibleExecInPin(UEdGraphNode* OnNode)
@@ -583,7 +574,7 @@ bool FNodeSnappingManager::TryToSnapDataflow(FPendingSnap& PendingSnap)
 	TSharedPtr<SGraphNode> NodeBeingSnapped = PendingSnap.NodeBeingDragged;
 	TSharedPtr<SGraphPanel> OwningPanelPinned = OwningPanel.Pin();
 	check(OwningPanelPinned);
-	UEdGraphPin* PinToMatchForSnap = UE::Private::GetFirstAndOnlyVisiblePin(NodeBeingSnapped->GetNodeObj());
+	UEdGraphPin* PinToMatchForSnap = UE::Private::GetFirstAndOnlyVisibleOutputPin(NodeBeingSnapped->GetNodeObj());
 	if(PinToMatchForSnap && PinToMatchForSnap->LinkedTo.Num() > 1)
 	{
 		PinToMatchForSnap = nullptr;
@@ -616,7 +607,8 @@ bool FNodeSnappingManager::TryToSnapDataflow(FPendingSnap& PendingSnap)
 	UEdGraphPin* DestGraphPin = DestPin->GetPinObj();
 	const FPinConnectionResponse Response = GetSchema()->CanCreateConnection(
 		DestGraphPin, PinToMatchForSnap);
-	if( Response.Response == ECanCreateConnectionResponse::CONNECT_RESPONSE_DISALLOW )
+	if( Response.Response != ECanCreateConnectionResponse::CONNECT_RESPONSE_MAKE && 
+		!DestGraphPin->LinkedTo.Contains(PinToMatchForSnap))
 	{
 		return false;
 	}
diff --git a/Engine/Source/Editor/GraphEditor/Private/SGraphPanel.cpp b/Engine/Source/Editor/GraphEditor/Private/SGraphPanel.cpp
--- a/Engine/Source/Editor/GraphEditor/Private/SGraphPanel.cpp
+++ b/Engine/Source/Editor/GraphEditor/Private/SGraphPanel.cpp
@@ -1449,16 +1449,19 @@ TSharedPtr<SGraphNode> SGraphPanel::GetGraphNodeUnderMouse(const FGeometry& MyGe
 	FArrangedChildren ArrangedNodes(EVisibility::Visible);
 	this->ArrangeChildrenForContextMenuSummon(MyGeometry, ArrangedNodes);
 	FArrangedGraphNode ArrangedWidget = {};
-	const int32 HoveredNodeIndex = SWidget::FindChildUnderMouse(ArrangedNodes, MouseEvent);
-	if (HoveredNodeIndex != INDEX_NONE)
+	if (FNodeSnappingManager* SnapManager = GetSnappingManager())
 	{
-		const FArrangedWidget& HoveredNode = ArrangedNodes[HoveredNodeIndex];
-		ArrangedWidget.Geometry = HoveredNode.Geometry;
-		ArrangedWidget.Widget = StaticCastSharedRef<SGraphNode>(HoveredNode.Widget);
+		ArrangedWidget = SnapManager->GetGraphNodeUnderMouse(MyGeometry, MouseEvent);
 	}
-	else if (FNodeSnappingManager* SnapManager = GetSnappingManager())
+	else
 	{
-		ArrangedWidget = SnapManager->GetGraphNodeUnderMouse(MyGeometry, MouseEvent);
+		const int32 HoveredNodeIndex = SWidget::FindChildUnderMouse(ArrangedNodes, MouseEvent);
+		if (HoveredNodeIndex != INDEX_NONE)
+		{
+			const FArrangedWidget& HoveredNode = ArrangedNodes[HoveredNodeIndex];
+			ArrangedWidget.Geometry = HoveredNode.Geometry;
+			ArrangedWidget.Widget = StaticCastSharedRef<SGraphNode>(HoveredNode.Widget);
+		}
 	}
 	
 	if (ArrangedWidget.Widget)
diff --git a/Engine/Source/Editor/GraphEditor/Private/SGraphPin.cpp b/Engine/Source/Editor/GraphEditor/Private/SGraphPin.cpp
--- a/Engine/Source/Editor/GraphEditor/Private/SGraphPin.cpp
+++ b/Engine/Source/Editor/GraphEditor/Private/SGraphPin.cpp
@@ -311,11 +311,7 @@ void SGraphPin::Construct(const FArguments& InArgs, UEdGraphPin* InPin)
 					SAssignNew(ValueBox, SBox)
 					.Padding(0.0f)
 					[
-						SAssignNew(ValueWidgetWrapper, SBox)
-						.Padding(0.0f)
-						[
-							ValueWidget.ToSharedRef()
-						]
+						ValueWidget.ToSharedRef()
 					]
 				];
 
@@ -1549,24 +1545,30 @@ bool SGraphPin::HasInteractiveTooltip() const
 
 bool SGraphPin::SnapSGraphNode(TSharedRef<SGraphNode> Node)
 {
-	if(!ValueWidgetWrapper)
+	if(!SnappedDataNodeContainer)
 	{
-		return false;
+		const bool bIsInput = (GetDirection() == EGPD_Input);
+		LabelAndValue->AddSlot()
+			.Padding(bIsInput ? FMargin(5.f, 0.0f, 0.0f, 0.0f) : FMargin(0.0f, 0.0f, 5.f, 0.0f))
+			.VAlign(VAlign_Center)
+			[
+				SAssignNew(SnappedDataNodeContainer, SBox)
+				.Padding(0.0f)
+			];
 	}
 
-	ValueWidgetWrapper->SetContent(Node);
+	SnappedDataNodeContainer->SetContent(Node);
 	SnappedNode = Node;
 	return true;
 }
 
 void SGraphPin::ClearSnappedSGraphNode()
 {
-	if(ValueWidgetWrapper) // control flow snaps won't have a valud widget
+	if(SnappedDataNodeContainer) // control flow snaps won't have a valud widget
 	{
-		ValueWidgetWrapper->SetContent(ValueWidget.ToSharedRef());
+		SnappedDataNodeContainer->SetContent(SNullWidget::NullWidget);
 	}
 	SnappedNode.Reset();
-	//ValueWidgetWrapper->SetActiveWidgetIndex(0); // implicit
 }
 
 bool SGraphPin::IsSnapped() const
@@ -1605,7 +1607,7 @@ TSharedPtr<SGraphNode> SGraphPin::GetSnappedChildForPosition(const FGeometry& My
 			TSharedPtr<SGraphNode> HoveredNode = StaticCastSharedRef<SGraphNode>(
 				ArrangedPins[HoveredPinIndex].Widget);
 			OutPositionInNode = ArrangedPins[HoveredPinIndex].Geometry.AbsoluteToLocal(Position);
-			return HoveredNode;
+			return HoveredNode->GetSnappedChildForPosition(ArrangedPins[HoveredPinIndex].Geometry, Position, OutPositionInNode);
 		}
 	}
 	return nullptr;
diff --git a/Engine/Source/Editor/GraphEditor/Public/SGraphPin.h b/Engine/Source/Editor/GraphEditor/Public/SGraphPin.h
--- a/Engine/Source/Editor/GraphEditor/Public/SGraphPin.h
+++ b/Engine/Source/Editor/GraphEditor/Public/SGraphPin.h
@@ -360,7 +360,7 @@ protected:
 
 	/** Image of pin */
 	TSharedPtr<SWidget> PinImage;
-	TSharedPtr<SBox> ValueWidgetWrapper;
+	TSharedPtr<SBox> SnappedDataNodeContainer;
 	TWeakPtr<SGraphNode> SnappedNode;
 
 	/** Label and value */
