Commit NodeDraggingManager's Total delta to the owning node panel, so that we can avoid reselecting objects on mouse up after nodes are dragged outside the 'deadzone'

#rnx

diff --git a/Engine/Source/Editor/GraphEditor/Private/GraphEditorSnapping/NodeSnappingManager.cpp b/Engine/Source/Editor/GraphEditor/Private/GraphEditorSnapping/NodeSnappingManager.cpp
--- a/Engine/Source/Editor/GraphEditor/Private/GraphEditorSnapping/NodeSnappingManager.cpp
+++ b/Engine/Source/Editor/GraphEditor/Private/GraphEditorSnapping/NodeSnappingManager.cpp
@@ -216,7 +216,6 @@ void ForEachDirectlySnappedNode(UEdGraphNode* Node, TFunctionRef<void(UEdGraphNo
 
 FNodeSnappingManager::FNodeSnappingManager(TWeakPtr<SGraphPanel> InOwningPanel)
 	: OwningPanel(InOwningPanel)
-	, TotalMouseDelta(0.f)
 	, bIsActivelyDragging(false)
 {
 	check(InOwningPanel.Pin());
@@ -255,7 +254,7 @@ bool FNodeSnappingManager::BeginDragWithSnap(const FGeometry& MyGeometry, const
 
 	OwningPanelPinned->NodeGrabOffset = MousePositionInNode;
 
-	TotalMouseDelta = 0.f;
+	OwningPanelPinned->TotalMouseDelta = 0.f;
 	bIsActivelyDragging = true;
 	
 	UEdGraphNode* ClickedNode = CastChecked<UEdGraphNode>(NodeStartingDrag->GetObjectBeingDisplayed());
@@ -381,12 +380,12 @@ void FNodeSnappingManager::UpdateDrag(const FGeometry& MyGeometry, const FPointe
 	}
 	
 	const FVector2f CursorDelta = MouseEvent.GetCursorDelta();
-	TotalMouseDelta += CursorDelta.Size();
+	OwningPanelPinned->TotalMouseDelta += CursorDelta.Size();
 
 	// Update the amount to pan panel
 	OwningPanelPinned->UpdateViewOffset(MyGeometry, MouseEvent.GetScreenSpacePosition());
 
-	const bool bCursorInDeadZone = TotalMouseDelta <= FSlateApplication::Get().GetDragTriggerDistance();
+	const bool bCursorInDeadZone = OwningPanelPinned->TotalMouseDelta <= FSlateApplication::Get().GetDragTriggerDistance();
 	
 	// reset snap data for the nodes we're moving, we want to recalculate it:
 	for(const FPendingConnection& UncommittedConnection : UncommittedConnections)
@@ -558,7 +557,6 @@ void FNodeSnappingManager::ResetDragManager()
 	MovingNodes.Reset();
 	ModifiedNodes.Reset();
 	UncommittedConnections.Reset();
-	TotalMouseDelta = 0.f;
 	bIsActivelyDragging = false;
 }
 
diff --git a/Engine/Source/Editor/GraphEditor/Private/GraphEditorSnapping/NodeSnappingManager.h b/Engine/Source/Editor/GraphEditor/Private/GraphEditorSnapping/NodeSnappingManager.h
--- a/Engine/Source/Editor/GraphEditor/Private/GraphEditorSnapping/NodeSnappingManager.h
+++ b/Engine/Source/Editor/GraphEditor/Private/GraphEditorSnapping/NodeSnappingManager.h
@@ -77,7 +77,5 @@ private:
 	};
 	TArray<FPendingConnection> UncommittedConnections;
 
-	float TotalMouseDelta;
-
 	bool bIsActivelyDragging;
 };
\ No newline at end of file
