Move node snapping manager out of vtable as that would require exporting the function

#rnx

diff --git a/Engine/Source/Editor/GraphEditor/Private/SGraphPanel.cpp b/Engine/Source/Editor/GraphEditor/Private/SGraphPanel.cpp
--- a/Engine/Source/Editor/GraphEditor/Private/SGraphPanel.cpp
+++ b/Engine/Source/Editor/GraphEditor/Private/SGraphPanel.cpp
@@ -1678,11 +1678,6 @@ FReply SGraphPanel::OnMouseDownDisconnectClosestPinFromWire(TSharedRef<SGraphPin
 	return FReply::Handled().BeginDragDrop(DragEvent.ToSharedRef());
 }
 
-FNodeSnappingManager* SGraphPanel::GetSnappingManager()
-{
-	return NodeSnappingManager.Get();
-}
-
 void SGraphPanel::OnBeginMakingConnection(UEdGraphPin* InOriginatingPin)
 {
 	OnBeginMakingConnection(FGraphPinHandle(InOriginatingPin));
diff --git a/Engine/Source/Editor/GraphEditor/Private/SNodePanel.cpp b/Engine/Source/Editor/GraphEditor/Private/SNodePanel.cpp
--- a/Engine/Source/Editor/GraphEditor/Private/SNodePanel.cpp
+++ b/Engine/Source/Editor/GraphEditor/Private/SNodePanel.cpp
@@ -302,6 +302,7 @@ TArray<SNodePanel::SNode::DiffHighlightInfo> SNodePanel::SNode::GetDiffHighlight
 SNodePanel::SNodePanel()
 	: Children(this)
 	, VisibleChildren(this)
+	, NodeSnappingManager(nullptr)
 {
 	bHasRelativeLayoutScale = true;
 }
@@ -539,7 +540,7 @@ void SNodePanel::OnEndNodeInteraction(const TSharedRef<SNode>& InNodeToDrag)
 
 FNodeSnappingManager* SNodePanel::GetSnappingManager()
 {
-	return nullptr;
+	return NodeSnappingManager.Get();
 }
 
 const FNodeSnappingManager* SNodePanel::GetSnappingManager() const
diff --git a/Engine/Source/Editor/GraphEditor/Public/SGraphPanel.h b/Engine/Source/Editor/GraphEditor/Public/SGraphPanel.h
--- a/Engine/Source/Editor/GraphEditor/Public/SGraphPanel.h
+++ b/Engine/Source/Editor/GraphEditor/Public/SGraphPanel.h
@@ -54,7 +54,6 @@ struct FAssetData;
 struct FDiffSingleResult;
 struct FEdGraphEditAction;
 struct FGuid;
-struct FNodeSnappingManager;
 
 DECLARE_DELEGATE( FOnUpdateGraphPanel )
 
@@ -453,10 +452,8 @@ private:
 	/** Begins a drag operation that disconnects the pin at ClosestPin, but only from the ConnectedPin, rather than all connections */
 	UE_API FReply OnMouseDownDisconnectClosestPinFromWire(TSharedRef<class SGraphPin> ClosestPin, TSharedRef<class SGraphPin> ConnectedPin);
 
-	/** Node snapping extension */
-	virtual struct FNodeSnappingManager* GetSnappingManager() override;
+	/** Node snapping is a private level extension, give it access to private graph editor state: */
 	friend struct FNodeSnappingManager;
-	TPimplPtr<FNodeSnappingManager> NodeSnappingManager;
 
 	/** Handle to timer callback that allows the UI to refresh it's arrangement each tick, allows animations to occur within the UI */
 	TWeakPtr<FActiveTimerHandle> ActiveTimerHandleInvalidatePerTick;
diff --git a/Engine/Source/Editor/GraphEditor/Public/SNodePanel.h b/Engine/Source/Editor/GraphEditor/Public/SNodePanel.h
--- a/Engine/Source/Editor/GraphEditor/Public/SNodePanel.h
+++ b/Engine/Source/Editor/GraphEditor/Public/SNodePanel.h
@@ -1186,11 +1186,13 @@ protected:
 
 	/** Called when the user left clicks on a node without dragging */
 	SGraphEditor::FOnNodeSingleClicked OnNodeSingleClicked;
+	
+	/** Implementation detail for node snapping, experimental and not user extensible */
+	TPimplPtr<struct FNodeSnappingManager> NodeSnappingManager;
+	struct FNodeSnappingManager* GetSnappingManager();
+	const struct FNodeSnappingManager* GetSnappingManager() const;
 
 private:
-	UE_API virtual struct FNodeSnappingManager* GetSnappingManager();
-	const struct FNodeSnappingManager* GetSnappingManager() const;
-	
 	/** Active timer that handles deferred zooming until the target zoom is reached */
 	UE_API EActiveTimerReturnType HandleZoomToFit(double InCurrentTime, float InDeltaTime);
 
