SA fix

diff --git a/Engine/Source/Editor/GraphEditor/Private/GraphDiffControl.cpp b/Engine/Source/Editor/GraphEditor/Private/GraphDiffControl.cpp
--- a/Engine/Source/Editor/GraphEditor/Private/GraphDiffControl.cpp
+++ b/Engine/Source/Editor/GraphEditor/Private/GraphDiffControl.cpp
@@ -497,7 +497,11 @@ static bool CheckGraphNodePin(UEdGraphPin* OldLinkedPin, UEdGraphPin* NewLinkedP
 	for (int32 i = 0; i < OldLinkedPin->LinkedTo.Num(); ++i)
 	{
 		UEdGraphPin* OldLinkedToPin = OldLinkedPin->LinkedTo[i];
-		UEdGraphPin* NewLinkedToPin = NewLinkedPin->LinkedTo.IsValidIndex(i) ? FindOtherLink(NewLinkedPin->LinkedTo, i, OldLinkedToPin) : nullptr;
+		UEdGraphPin* NewLinkedToPin = nullptr;
+		if(NewLinkedPin && OldLinkedToPin)
+		{
+			NewLinkedToPin = NewLinkedPin->LinkedTo.IsValidIndex(i) ? FindOtherLink(NewLinkedPin->LinkedTo, i, OldLinkedToPin) : nullptr;
+		}
 
 		if (!OldLinkedToPin || !NewLinkedToPin)
 		{
