PR #11872: Feature: Slice graph connections by alt+middle-mouse dragging in graph editors

diff --git a/Engine/Source/Editor/GraphEditor/Private/ConnectionDrawingPolicy.cpp b/Engine/Source/Editor/GraphEditor/Private/ConnectionDrawingPolicy.cpp
--- a/Engine/Source/Editor/GraphEditor/Private/ConnectionDrawingPolicy.cpp
+++ b/Engine/Source/Editor/GraphEditor/Private/ConnectionDrawingPolicy.cpp
@@ -90,6 +90,24 @@ UE::Slate::FDeprecateVector2DResult FGeometryHelper::FindClosestPointOnGeom(cons
 	return BestPoint;
 }
 
+namespace IntersectionHelpers
+{
+	static bool SegmentIntersection2D(const FVector2f& SegmentAStart, const FVector2f& SegmentAEnd, const FVector2f& SegmentBStart, const FVector2f& SegmentBEnd, FVector2f& OutIntersectionPoint)
+	{
+		FVector Intersection;
+		if ( FMath::SegmentIntersection2D(FVector(SegmentAStart.X, SegmentAStart.Y, 0),
+			FVector(SegmentAEnd.X, SegmentAEnd.Y, 0),
+			FVector(SegmentBStart.X, SegmentBStart.Y, 0),
+			FVector(SegmentBEnd.X, SegmentBEnd.Y, 0),
+			Intersection))
+		{
+			OutIntersectionPoint = FVector2f(static_cast<float>(Intersection.X), static_cast<float>(Intersection.Y));
+			return true;
+		}
+		return false;
+	}
+}
+
 /////////////////////////////////////////////////////
 // FConnectionDrawingPolicy
 
@@ -107,6 +125,7 @@ FConnectionDrawingPolicy::FConnectionDrawingPolicy(int32 InBackLayerID, int32 In
 	MidpointImage = nullptr;
 	MidpointRadius = FVector2f::ZeroVector;
 	HoverDeemphasisDarkFraction = 0.8f;
+	SliceDeemphasisAlphaMultiplier = 0.5f;
 
 	BubbleImage = FAppStyle::GetBrush( TEXT("Graph.ExecutionBubble") );
 }
@@ -202,6 +221,11 @@ void FConnectionDrawingPolicy::SetMousePosition(const UE::Slate::FDeprecateVecto
 	LocalMousePosition = InMousePos;
 }
 
+void FConnectionDrawingPolicy::SetSliceLine(const FMarqueeRect& InLine)
+{
+	SliceLine = InLine;
+}
+
 void FConnectionDrawingPolicy::SetMarkedPin(TWeakPtr<SGraphPin> InMarkedPin)
 {
 	if (InMarkedPin.IsValid())
@@ -290,6 +314,8 @@ void FConnectionDrawingPolicy::DrawConnection_DeprecationHelper(int32 LayerId, c
 	const FVector2f P0Tangent = (Params.StartTangent.IsNearlyZero()) ? ((Params.StartDirection == EGPD_Output) ? SplineTangent : -SplineTangent) : Params.StartTangent;
 	const FVector2f P1Tangent = (Params.EndTangent.IsNearlyZero()) ? ((Params.EndDirection == EGPD_Input) ? SplineTangent : -SplineTangent) : Params.EndTangent;
 
+	bool bSliceLineIntersectsSpline = false;
+
 	if (Settings->bTreatSplinesLikePins)
 	{
 		// Distance to consider as an overlap
@@ -299,11 +325,12 @@ void FConnectionDrawingPolicy::DrawConnection_DeprecationHelper(int32 LayerId, c
 		// dead zone to avoid mistakes caused by missing a double-click on a connection.
 		const float QueryDistanceForCloseSquared = FMath::Square(FMath::Sqrt(QueryDistanceTriggerThresholdSquared) + Settings->SplineCloseTolerance);
 
+		bool bSliceLineIntersectsSplineBounds = false;
 		bool bCloseToSpline = false;
 		{
 			// The curve will include the endpoints but can extend out of a tight bounds because of the tangents
 			// P0Tangent coefficient maximizes to 4/27 at a=1/3, and P1Tangent minimizes to -4/27 at a=2/3.
-			const float MaximumTangentContribution = 4.0f / 27.0f;
+			constexpr float MaximumTangentContribution = 4.0f / 27.0f;
 			FBox2f Bounds(ForceInit);
 
 			Bounds += FVector2f(P0);
@@ -313,6 +340,12 @@ void FConnectionDrawingPolicy::DrawConnection_DeprecationHelper(int32 LayerId, c
 
 			bCloseToSpline = Bounds.ComputeSquaredDistanceToPoint(LocalMousePosition) < QueryDistanceForCloseSquared;
 
+			// Re-use the bounds we just constructed to check if the slice line intersects the spline
+			if (SliceLine.IsValid())
+			{
+				bSliceLineIntersectsSplineBounds = FMath::LineBoxIntersection2D(Bounds, SliceLine.StartPoint, SliceLine.EndPoint);
+			}
+
 			// Draw the bounding box for debugging
 #if 0
 #define DrawSpaceLine(Point1, Point2, DebugWireColor) {const FVector2f FakeTangent = (Point2 - Point1).GetSafeNormal(); FSlateDrawElement::MakeDrawSpaceSpline(DrawElementsList, LayerId, Point1, FakeTangent, Point2, FakeTangent, ClippingRect, 1.0f, ESlateDrawEffect::None, DebugWireColor); }
@@ -333,19 +366,40 @@ void FConnectionDrawingPolicy::DrawConnection_DeprecationHelper(int32 LayerId, c
 			}
 #endif
 		}
+		
+		constexpr int32 NumStepsToTest = 16;
+		constexpr float TestStepInterval = 1.0f / NumStepsToTest;
+		
+		// Only bother checking for precise spline intersection if we intersected the bounding box
+		if (bSliceLineIntersectsSplineBounds)
+		{
+			// Find the intersecting point on the spline
+			FVector2f OutIntersectionPoint;
+			FVector2f Point1 = FMath::CubicInterp(P0, P0Tangent, P1, P1Tangent, 0.0f);
+			for (float TestAlpha = 0.0f; TestAlpha < 1.0f; TestAlpha += TestStepInterval)
+			{
+				const FVector2f Point2 = FMath::CubicInterp(P0, P0Tangent, P1, P1Tangent, TestAlpha + TestStepInterval);
+				
+				if (IntersectionHelpers::SegmentIntersection2D(SliceLine.StartPoint, SliceLine.EndPoint, Point1, Point2, OutIntersectionPoint))
+				{
+					bSliceLineIntersectsSpline = true;
+					ConnectionsIntersectingSliceLine.Emplace(Params.AssociatedPin1, Params.AssociatedPin2);
+					break;
+				}
+
+				Point1 = Point2;
+			}
+		}
 
 		if (bCloseToSpline)
 		{
 			// Find the closest approach to the spline
 			FVector2f ClosestPoint(ForceInit);
 			float ClosestDistanceSquared = FLT_MAX;
-
-			const int32 NumStepsToTest = 16;
-			const float StepInterval = 1.0f / (float)NumStepsToTest;
 			FVector2f Point1 = FMath::CubicInterp(P0, P0Tangent, P1, P1Tangent, 0.0f);
-			for (float TestAlpha = 0.0f; TestAlpha < 1.0f; TestAlpha += StepInterval)
+			for (float TestAlpha = 0.0f; TestAlpha < 1.0f; TestAlpha += TestStepInterval)
 			{
-				const FVector2f Point2 = FMath::CubicInterp(P0, P0Tangent, P1, P1Tangent, TestAlpha + StepInterval);
+				const FVector2f Point2 = FMath::CubicInterp(P0, P0Tangent, P1, P1Tangent, TestAlpha + TestStepInterval);
 
 				const FVector2f ClosestPointToSegment = FMath::ClosestPointOnSegment2D(LocalMousePosition, Point1, Point2);
 				const float DistanceSquared = (LocalMousePosition - ClosestPointToSegment).SizeSquared();
@@ -377,6 +431,13 @@ void FConnectionDrawingPolicy::DrawConnection_DeprecationHelper(int32 LayerId, c
 		}
 	}
 
+	// Apply an opacity fade to wires intersected by the slice line
+	FLinearColor WireColor = Params.WireColor;
+	if (bSliceLineIntersectsSpline)
+	{
+		WireColor.A *= SliceDeemphasisAlphaMultiplier;
+	}
+
 	// Draw the spline itself
 	FSlateDrawElement::MakeDrawSpaceSpline(
 		DrawElementsList,
@@ -385,7 +446,7 @@ void FConnectionDrawingPolicy::DrawConnection_DeprecationHelper(int32 LayerId, c
 		P1, P1Tangent,
 		Params.WireThickness,
 		ESlateDrawEffect::None,
-		Params.WireColor
+		WireColor
 	);
 
 	if (Params.bDrawBubbles || (MidpointImage != nullptr))
@@ -474,6 +535,46 @@ void FConnectionDrawingPolicy::DrawPreviewConnector_DeprecationHelper(const FGeo
 	DrawSplineWithArrow(StartPoint, EndPoint, Params);
 }
 
+bool FConnectionDrawingPolicy::DrawSliceLine()
+{
+	if (SliceLine.IsValid())
+	{
+		TArray<FVector2f> LinePoints;
+		LinePoints.Add(SliceLine.StartPoint);
+		LinePoints.Add(SliceLine.EndPoint);
+		
+		// Draw a solid line instead of dashes if the dash length is zero (or close to, to avoid the line disappearing due to dashes being too small)
+		if (Settings->SliceLineDashLength <= 1.f)
+		{
+			FSlateDrawElement::MakeLines(
+				DrawElementsList,
+				WireLayerID,
+				FPaintGeometry(),
+				LinePoints,
+				ESlateDrawEffect::None,
+				Settings->SliceLineColor,
+				true,
+				Settings->SliceLineThickness
+			);
+		}
+		else
+		{
+			FSlateDrawElement::MakeDashedLines(
+				DrawElementsList,
+				WireLayerID,
+				FPaintGeometry(),
+				CopyTemp(LinePoints),
+				ESlateDrawEffect::None,
+				Settings->SliceLineColor,
+				Settings->SliceLineThickness,
+				Settings->SliceLineDashLength
+			);
+		}
+		return true;
+	}
+	return false;
+}
+
 void FConnectionDrawingPolicy::DetermineWiringStyle(UEdGraphPin* OutputPin, UEdGraphPin* InputPin, /*inout*/ FConnectionParams& Params)
 {
 	Params.AssociatedPin1 = OutputPin;
diff --git a/Engine/Source/Editor/GraphEditor/Private/GraphEditorModule.cpp b/Engine/Source/Editor/GraphEditor/Private/GraphEditorModule.cpp
--- a/Engine/Source/Editor/GraphEditor/Private/GraphEditorModule.cpp
+++ b/Engine/Source/Editor/GraphEditor/Private/GraphEditorModule.cpp
@@ -73,6 +73,7 @@ TSharedRef<SGraphEditor> FGraphEditorModule::PRIVATE_MakeGraphEditor(
 	FSimpleDelegate InOnNavigateHistoryBack,
 	FSimpleDelegate InOnNavigateHistoryForward,
 	TAttribute<bool> ShowGraphStateOverlay,
+	TAttribute<bool> AllowConnectionSlicing,
 	FOnGraphEditorTickDelegate InOnTickDelegate)
 {
 	return
@@ -90,7 +91,8 @@ TSharedRef<SGraphEditor> FGraphEditorModule::PRIVATE_MakeGraphEditor(
 		.OnNavigateHistoryBack(InOnNavigateHistoryBack)
 		.OnNavigateHistoryForward(InOnNavigateHistoryForward)
 		.OnTick(InOnTickDelegate)
-		.ShowGraphStateOverlay(ShowGraphStateOverlay);
+		.ShowGraphStateOverlay(ShowGraphStateOverlay)
+		.AllowConnectionSlicing(AllowConnectionSlicing);
 }
 
 /////////////////////////////////////////////////////
diff --git a/Engine/Source/Editor/GraphEditor/Private/GraphEditorSettings.cpp b/Engine/Source/Editor/GraphEditor/Private/GraphEditorSettings.cpp
--- a/Engine/Source/Editor/GraphEditor/Private/GraphEditorSettings.cpp
+++ b/Engine/Source/Editor/GraphEditor/Private/GraphEditorSettings.cpp
@@ -98,6 +98,11 @@ UGraphEditorSettings::UGraphEditorSettings( const FObjectInitializer& ObjectInit
 	ResultNodeTitleColor = FLinearColor(1.0f, 0.65f, 0.4f, 1.0f);
 	DefaultCommentNodeTitleColor = FLinearColor(0.15f, 0.15f, 0.15f, 0.5f);
 	PreviewNodeTitleColor = FLinearColor(0.0f, 0.0f, 1.0f, 1.0f);
+
+	// Graph connection slicing visuals
+	SliceLineColor = FLinearColor(1.0f, 0.9f, 0.1f, 1.f);
+	SliceLineDashLength = 8.f;
+	SliceLineThickness = 2.f;
 }
 
 #if WITH_EDITOR
diff --git a/Engine/Source/Editor/GraphEditor/Private/SGraphEditorImpl.cpp b/Engine/Source/Editor/GraphEditor/Private/SGraphEditorImpl.cpp
--- a/Engine/Source/Editor/GraphEditor/Private/SGraphEditorImpl.cpp
+++ b/Engine/Source/Editor/GraphEditor/Private/SGraphEditorImpl.cpp
@@ -755,6 +755,7 @@ void SGraphEditorImpl::Construct( const FArguments& InArgs )
 			//.OnUpdateGraphPanel( this, &SGraphEditorImpl::GraphEd_OnPanelUpdated )
 			.OnDisallowedPinConnection( InArgs._GraphEvents.OnDisallowedPinConnection )
 			.ShowGraphStateOverlay(InArgs._ShowGraphStateOverlay)
+			.AllowConnectionSlicing(InArgs._AllowConnectionSlicing)
 			.OnDoubleClicked(InArgs._GraphEvents.OnDoubleClicked)
 			.OnMouseButtonDown(InArgs._GraphEvents.OnMouseButtonDown)
 		    .OnNodeSingleClicked(InArgs._GraphEvents.OnNodeSingleClicked)
diff --git a/Engine/Source/Editor/GraphEditor/Private/SGraphEditorImpl.h b/Engine/Source/Editor/GraphEditor/Private/SGraphEditorImpl.h
--- a/Engine/Source/Editor/GraphEditor/Private/SGraphEditorImpl.h
+++ b/Engine/Source/Editor/GraphEditor/Private/SGraphEditorImpl.h
@@ -208,6 +208,7 @@ public:
 		: _AdditionalCommands( TSharedPtr<FUICommandList>() )
 		, _IsEditable(true)
 		, _DisplayAsReadOnly(false)
+		, _AllowConnectionSlicing(true)
 		{}
 
 
@@ -227,6 +228,8 @@ public:
 
 		/** Show overlay elements for the graph state such as the PIE and read-only borders and text */
 		SLATE_ATTRIBUTE(bool, ShowGraphStateOverlay)
+		/** Allow graph connections to be sliced and removed */
+		SLATE_ATTRIBUTE(bool, AllowConnectionSlicing)
 	SLATE_END_ARGS()
 
 	void Construct( const FArguments& InArgs );
diff --git a/Engine/Source/Editor/GraphEditor/Private/SGraphPanel.cpp b/Engine/Source/Editor/GraphEditor/Private/SGraphPanel.cpp
--- a/Engine/Source/Editor/GraphEditor/Private/SGraphPanel.cpp
+++ b/Engine/Source/Editor/GraphEditor/Private/SGraphPanel.cpp
@@ -131,6 +131,7 @@ void SGraphPanel::Construct( const SGraphPanel::FArguments& InArgs )
 	MyRegisteredGraphChangedDelegateHandle = this->GraphObj->AddOnGraphChangedHandler(MyRegisteredGraphChangedDelegate);
 	
 	ShowGraphStateOverlay = InArgs._ShowGraphStateOverlay;
+	AllowConnectionSlicing = InArgs._AllowConnectionSlicing;
 
 	SavedMousePosForOnPaintEventLocalSpace = FVector2f::ZeroVector;
 	PreviousFrameSavedMousePosForSplineOverlap = FVector2f::ZeroVector;
@@ -492,24 +493,40 @@ int32 SGraphPanel::OnPaint( const FPaintArgs& Args, const FGeometry& AllottedGeo
 		}
 
 		const bool bUseDrawStateCaching = ConnectionDrawingPolicy->UseDrawStateCaching();
-		
-		TArray<TSharedPtr<SGraphPin>> OverridePins;
-		for (const FGraphPinHandle& Handle : PreviewConnectorFromPins)
+
+		// Even if we're now in 'slice mode', you need to keep holding the modifiers to actually commit the slice,
+		// to basically allow letting go of the modifiers to 'abort' your slice early instead of needing to undo it
+		const bool bShouldPreviewSliceLine = CurrentSlice.IsSet() && FSlateApplication::Get().GetModifierKeys().IsAltDown();
+
+		// Tell the connection policy about the slice line before we draw any connections so it will do intersection tests
+		if (bShouldPreviewSliceLine)
 		{
-			TSharedPtr<SGraphPin> Pin = Handle.FindInGraphPanel(*this);
-			if (Pin.IsValid() && Pin->GetPinObj())
+			FMarqueeRect AbsoluteLine = FMarqueeRect(AllottedGeometry.LocalToAbsolute(GraphCoordToPanelCoord(CurrentSlice->Line.StartPoint)));
+            AbsoluteLine.UpdateEndPoint(AllottedGeometry.LocalToAbsolute(GraphCoordToPanelCoord(CurrentSlice->Line.EndPoint)));
+            ConnectionDrawingPolicy->SetSliceLine(AbsoluteLine);
+		}
+		else
+		{
+			TArray<TSharedPtr<SGraphPin>> OverridePins;
+			for (const FGraphPinHandle& Handle : PreviewConnectorFromPins)
 			{
-				OverridePins.Add(Pin);
+				TSharedPtr<SGraphPin> Pin = Handle.FindInGraphPanel(*this);
+				if (Pin.IsValid() && Pin->GetPinObj())
+				{
+					OverridePins.Add(Pin);
+				}
 			}
-		}
-		ConnectionDrawingPolicy->SetHoveredPins(CurrentHoveredPins, OverridePins, TimeWhenMouseEnteredPin);
-		ConnectionDrawingPolicy->SetMarkedPin(MarkedPin);
-		ConnectionDrawingPolicy->SetMousePosition(AllottedGeometry.LocalToAbsolute(SavedMousePosForOnPaintEventLocalSpace));
 
-		if (IsRelinkingConnection())
-		{
-			ConnectionDrawingPolicy->SetRelinkConnections(RelinkConnections);
-			ConnectionDrawingPolicy->SetSelectedNodes(GetSelectedGraphNodes());
+			ConnectionDrawingPolicy->SetHoveredPins(CurrentHoveredPins, OverridePins, TimeWhenMouseEnteredPin);
+			ConnectionDrawingPolicy->SetMarkedPin(MarkedPin);
+			
+			ConnectionDrawingPolicy->SetMousePosition(AllottedGeometry.LocalToAbsolute(SavedMousePosForOnPaintEventLocalSpace));
+
+			if (IsRelinkingConnection())
+			{
+				ConnectionDrawingPolicy->SetRelinkConnections(RelinkConnections);
+				ConnectionDrawingPolicy->SetSelectedNodes(GetSelectedGraphNodes());
+			}
 		}
 
 		// Get the set of pins for all children and synthesize geometry for culled out pins so lines can be drawn to them.
@@ -632,10 +649,16 @@ int32 SGraphPanel::OnPaint( const FPaintArgs& Args, const FGeometry& AllottedGeo
 			// Update the spline hover state
 			if (const_cast<SGraphPanel*>(this)->OnSplineHoverStateChanged(OverlapData))
 			{
-				
 				// if hover state changed, we update the tooltip text based on the connection drawing policy
 				const_cast<SGraphPanel*>(this)->SetToolTip(ConnectionDrawingPolicy->GetConnectionToolTip(*this, OverlapData));
 			}
+
+			// After drawing all connections, we can draw the slice line on top of them
+			if (bShouldPreviewSliceLine && ConnectionDrawingPolicy->DrawSliceLine())
+			{
+				// Take ownership of the intersected connections to save a copy, since we're about to delete the drawing policy
+				const_cast<SGraphPanel*>(this)->CurrentSlice->ConnectionsIntersectingSliceLine = MoveTemp(ConnectionDrawingPolicy->ConnectionsIntersectingSliceLine);
+			}
 		}
 
 		delete ConnectionDrawingPolicy;
@@ -970,16 +993,35 @@ FReply SGraphPanel::OnMouseButtonDown(const FGeometry& MyGeometry, const FPointe
 {
 	const bool bIsControlDown = MouseEvent.IsControlDown();
 	const bool bIsAltDown = MouseEvent.IsAltDown();
-	if ((MouseEvent.GetEffectingButton() == EKeys::LeftMouseButton) && (bIsAltDown || bIsControlDown))
+	
+	if (MouseEvent.GetEffectingButton() == EKeys::LeftMouseButton && (bIsAltDown || bIsControlDown))
 	{
-		// Intercept alt-left clicking on the hovered spline for targeted break link
-		UEdGraphPin* Pin1;
-		UEdGraphPin* Pin2;
-		if (bIsAltDown && PreviousFrameSplineOverlap.GetPins(*this, Pin1, Pin2))
+		if (bIsAltDown)
 		{
-			const UEdGraphSchema* Schema = GraphObj->GetSchema();
-			Schema->BreakSinglePinLink(Pin1, Pin2);
-			return FReply::Handled();
+			// If we allow slicing, we want to defer break links to OnMouseButtonUp so we can tell if the mouse was dragged to slice or if we want a normal
+			if (AllowConnectionSlicing.Get())
+			{
+				LastClickSplineOverlap = PreviousFrameSplineOverlap;
+				// Don't allow starting a slice while hovering a node, or for read-only graphs
+				if (!NodeUnderMousePtr.IsValid() && IsGraphEditable() && !DisplayAsReadOnly.Get())
+				{
+					const FVector2f GraphMousePos = PanelCoordToGraphCoord( MyGeometry.AbsoluteToLocal( MouseEvent.GetScreenSpacePosition() ) );
+					CurrentSlice = FConnectionSlice(GraphMousePos);
+					return FReply::Handled().CaptureMouse(SharedThis(this));
+				}
+			}
+			// If we don't allow slicing, run default pin break behavior (executes on mouse down instead of mouse up)
+			else
+			{
+				UEdGraphPin* Pin1;
+				UEdGraphPin* Pin2;
+				if (PreviousFrameSplineOverlap.GetPins(*this, Pin1, Pin2))
+				{
+					const UEdGraphSchema* Schema = GraphObj->GetSchema();
+					Schema->BreakSinglePinLink(Pin1, Pin2);
+					return FReply::Handled();
+				}
+			}
 		}
 		else if (SGraphPin* BestPinFromHoveredSpline = GetBestPinFromHoveredSpline())
 		{
@@ -1057,6 +1099,44 @@ TArray<UEdGraphNode*> SGraphPanel::GetSelectedGraphNodes() const
 
 FReply SGraphPanel::OnMouseButtonUp(const FGeometry& MyGeometry, const FPointerEvent& MouseEvent)
 {
+	const bool bIsAltDown = MouseEvent.IsAltDown();
+	if (CurrentSlice && MouseEvent.GetEffectingButton() == EKeys::LeftMouseButton && bIsAltDown)
+	{
+		if (bInSlicingMode)
+		{
+			ensure(HasMouseCapture());
+			const FScopedTransaction SliceTransaction(NSLOCTEXT("GraphEditor", "SliceConnectionsAction", "Slice Connections"));
+		
+			const UEdGraphSchema* Schema = GraphObj->GetSchema();
+		
+			for (const TPair<FEdGraphPinReference, FEdGraphPinReference>& ConnectionPair : CurrentSlice->ConnectionsIntersectingSliceLine)
+			{
+				UEdGraphPin* Pin1 = ConnectionPair.Key.Get();
+				UEdGraphPin* Pin2 = ConnectionPair.Value.Get();
+			
+				if (Pin1 && Pin2)
+				{
+					Schema->BreakSinglePinLink(Pin1, Pin2);
+				}
+			}
+
+			CurrentSlice.Reset();
+            bInSlicingMode = false;
+			return FReply::Handled().ReleaseMouseCapture();
+		}
+
+		// Intercept alt-left clicking on the hovered spline for targeted break link
+		UEdGraphPin* Pin1;
+		UEdGraphPin* Pin2;
+		if (LastClickSplineOverlap.GetPins(*this, Pin1, Pin2))
+		{
+			const UEdGraphSchema* Schema = GraphObj->GetSchema();
+			Schema->BreakSinglePinLink(Pin1, Pin2);
+			CurrentSlice.Reset();
+			return FReply::Handled().ReleaseMouseCapture();
+		}
+	}
+	
 	if (!NodeUnderMousePtr.IsValid() && !Marquee.IsValid() && (MouseEvent.GetEffectingButton() == EKeys::LeftMouseButton) && (MouseEvent.IsShiftDown()))
 	{
 		if (SGraphPin* BestPinFromHoveredSpline = GetBestPinFromHoveredSpline())
@@ -1250,6 +1330,25 @@ FReply SGraphPanel::OnMouseMove(const FGeometry& MyGeometry, const FPointerEvent
 		return FReply::Handled();
 	}
 
+	// Keep updating the slice line as long as the mouse button is held, regardless of whether still holding the modifier keys
+	// Also consume the mouse move since we don't really want other graph behaviour happening while slicing
+	if (CurrentSlice)
+	{
+		if (this->HasMouseCapture())
+		{
+			const FVector2f GraphMousePos = PanelCoordToGraphCoord(MyGeometry.AbsoluteToLocal(MouseEvent.GetScreenSpacePosition()));
+			CurrentSlice->Line.UpdateEndPoint(GraphMousePos);
+			if (!CurrentSlice->Line.GetSize().IsNearlyZero())
+			{
+				bInSlicingMode = true;
+			}
+			return FReply::Handled();
+		}
+
+		// If we somehow lost mouse capture (eg. window was deactivated) then we should cancel the slice
+		CurrentSlice.Reset();
+	}
+
 	return SNodePanel::OnMouseMove(MyGeometry, MouseEvent);
 }
 
diff --git a/Engine/Source/Editor/GraphEditor/Public/ConnectionDrawingPolicy.h b/Engine/Source/Editor/GraphEditor/Public/ConnectionDrawingPolicy.h
--- a/Engine/Source/Editor/GraphEditor/Public/ConnectionDrawingPolicy.h
+++ b/Engine/Source/Editor/GraphEditor/Public/ConnectionDrawingPolicy.h
@@ -8,6 +8,7 @@
 #include "CoreMinimal.h"
 #include "EdGraph/EdGraphNode.h"
 #include "EdGraph/EdGraphPin.h"
+#include "Framework/MarqueeRect.h"
 #include "GraphEditorSettings.h"
 #include "GraphSplineOverlapResult.h"
 #include "HAL/Platform.h"
@@ -115,6 +116,7 @@ public:
 	FDeprecateSlateVector2D MidpointRadius;
 
 	FGraphSplineOverlapResult SplineOverlapResult;
+	TArray<TPair<FEdGraphPinReference, FEdGraphPinReference>> ConnectionsIntersectingSliceLine;
 
 	/** Handle for a currently relinked connection. */
 	struct FRelinkConnection
@@ -126,6 +128,7 @@ public:
 protected:
 	float ZoomFactor; 
 	float HoverDeemphasisDarkFraction;
+	float SliceDeemphasisAlphaMultiplier;
 	const FSlateRect& ClippingRect;
 	FSlateWindowElementList& DrawElementsList;
 	TMap< UEdGraphPin*, TSharedPtr<SGraphPin> > PinToPinWidgetMap;
@@ -139,6 +142,8 @@ protected:
 
 	/** Selected nodes in the graph panel. */
 	TArray<UEdGraphNode*> SelectedGraphNodes;
+	
+	FMarqueeRect SliceLine;
 public:
 	virtual ~FConnectionDrawingPolicy() {}
 	
@@ -149,6 +154,10 @@ public:
 
 	UE_API void SetMousePosition(const UE::Slate::FDeprecateVector2DParameter& InMousePos);
 
+	// If a (valid) slice line has been set, we'll intersect any drawn splines via DrawConnection() with that line,
+	// and add any intersecting pin pairs to the public ConnectionsIntersectingSliceLine array. 
+	UE_API void SetSliceLine(const FMarqueeRect& InLine);
+
 	// Update the drawing policy with the marked pin (which may not be valid)
 	UE_API void SetMarkedPin(TWeakPtr<SGraphPin> InMarkedPin);
 
@@ -179,6 +188,10 @@ public:
 	UE_API virtual void DrawPreviewConnector(const FGeometry& PinGeometry, const FVector2D& StartPoint, const FVector2D& EndPoint, UEdGraphPin* Pin) UE_SLATE_DEPRECATED_VECTOR_VIRTUAL_FUNCTION;
 	UE_API virtual void DrawPreviewConnector(const FGeometry& PinGeometry, const FVector2f& StartPoint, const FVector2f& EndPoint, UEdGraphPin* Pin);
 
+	// If a SliceLine has already been set, draw it. Subclasses may return false to indicate they don't support drawing a
+	// slice line or draw it in a different style.
+	UE_API virtual bool DrawSliceLine();
+
 	// Give specific editor modes a chance to highlight this connection or darken non-interesting connections
 	UE_API virtual void DetermineWiringStyle(UEdGraphPin* OutputPin, UEdGraphPin* InputPin, /*inout*/ FConnectionParams& Params);
 
diff --git a/Engine/Source/Editor/GraphEditor/Public/GraphEditorModule.h b/Engine/Source/Editor/GraphEditor/Public/GraphEditorModule.h
--- a/Engine/Source/Editor/GraphEditor/Public/GraphEditorModule.h
+++ b/Engine/Source/Editor/GraphEditor/Public/GraphEditorModule.h
@@ -58,6 +58,7 @@ private:
 		FSimpleDelegate InOnNavigateHistoryBack,
 		FSimpleDelegate InOnNavigateHistoryForward,
 		TAttribute<bool> ShowGraphStateOverlay,
+		TAttribute<bool> AllowConnectionSlicing,
 		FOnGraphEditorTickDelegate InOnTickDelegate);
 
 private:
diff --git a/Engine/Source/Editor/GraphEditor/Public/GraphEditorSettings.h b/Engine/Source/Editor/GraphEditor/Public/GraphEditorSettings.h
--- a/Engine/Source/Editor/GraphEditor/Public/GraphEditorSettings.h
+++ b/Engine/Source/Editor/GraphEditor/Public/GraphEditorSettings.h
@@ -283,6 +283,18 @@ public:
 	UPROPERTY(EditAnywhere, config, Category = NodeTitleColors)
 	FLinearColor PreviewNodeTitleColor;
 
+	/** The color of the line shown while slicing connections */
+	UPROPERTY(EditAnywhere, config, Category=Slicing)
+	FLinearColor SliceLineColor;
+
+	/** The length of each dash in the slice line */
+	UPROPERTY(EditAnywhere, config, Category=Slicing)
+	float SliceLineDashLength;
+
+	/** The visual thickness of the slice line, does not affect collision */
+	UPROPERTY(EditAnywhere, config, Category=Slicing)
+	float SliceLineThickness;
+
 public:
 	/** The thickness of a data wire */
 	UPROPERTY(EditAnywhere, config, Category=Tracing)
@@ -351,7 +363,7 @@ public:
 	UPROPERTY(config, EditAnywhere, Category=ContextMenu)
 	bool bOpenCreateMenuOnBlankGraphAreas;
 
-/** Whether a comment node should move any fully enclosed nodes around when it is moved */
+	/** Whether a comment node should move any fully enclosed nodes around when it is moved */
 	UPROPERTY(EditAnywhere, config, Category = CommentNodes)
 	TEnumAsByte<ECommentBoxMode::Type> DefaultCommentNodeMoveMode;
 
diff --git a/Engine/Source/Editor/GraphEditor/Public/SGraphPanel.h b/Engine/Source/Editor/GraphEditor/Public/SGraphPanel.h
--- a/Engine/Source/Editor/GraphEditor/Public/SGraphPanel.h
+++ b/Engine/Source/Editor/GraphEditor/Public/SGraphPanel.h
@@ -12,6 +12,7 @@
 #include "CoreMinimal.h"
 #include "Delegates/Delegate.h"
 #include "EdGraph/EdGraphPin.h"
+#include "Framework/MarqueeRect.h"
 #include "GraphEditAction.h"
 #include "GraphEditor.h"
 #include "GraphSplineOverlapResult.h"
@@ -73,6 +74,19 @@ struct FGraphContextMenuArguments
 	TArray<UEdGraphPin*> DragFromPins;
 };
 
+// Data we need to store while the user is slicing through connections in the graph
+struct FConnectionSlice
+{
+	// The start/end points of the slice line, in graph coords
+	FMarqueeRect Line;
+
+	// The pin connections which were intersecting the line last frame when the connection policy was drawn,
+	// and which should be broken if this slice is executed
+	TArray<TPair<FEdGraphPinReference, FEdGraphPinReference>> ConnectionsIntersectingSliceLine;
+
+	explicit FConnectionSlice(const FVector2f StartPosition)
+		: Line(StartPosition) { }
+};
 
 class SGraphPanel : public SNodePanel, public FGCObject
 {
@@ -88,6 +102,7 @@ public:
 		, _IsEditable( true )
 		, _DisplayAsReadOnly( false )
 		, _ShowGraphStateOverlay(true)
+		, _AllowConnectionSlicing(true)
 		, _OnUpdateGraphPanel()
 		{
 			_Clipping = EWidgetClipping::ClipToBounds;
@@ -140,6 +155,7 @@ PRAGMA_ENABLE_DEPRECATION_WARNINGS
 		SLATE_ATTRIBUTE( bool, DisplayAsReadOnly )
 		/** Show overlay elements for the graph state such as the PIE and read-only borders and text */
 		SLATE_ATTRIBUTE(bool, ShowGraphStateOverlay)
+		SLATE_ATTRIBUTE(bool, AllowConnectionSlicing)
 		SLATE_EVENT( FOnNodeVerifyTextCommit, OnVerifyTextCommit )
 		SLATE_EVENT( FOnNodeTextCommitted, OnTextCommitted )
 		
@@ -344,6 +360,8 @@ protected:
 	/** Time since the last mouse enter/exit on a pin */
 	double TimeWhenMouseEnteredPin;
 	double TimeWhenMouseLeftPin;
+	
+	bool bInSlicingMode = false;
 
 	/** Sometimes the panel draws a preview connector; e.g. when the user is connecting pins */
 	TArray< FGraphPinHandle > PreviewConnectorFromPins;
@@ -356,6 +374,7 @@ protected:
 	/** The overlap results from the previous OnPaint call */
 	FDeprecateSlateVector2D PreviousFrameSavedMousePosForSplineOverlap;
 	FGraphSplineOverlapResult PreviousFrameSplineOverlap;
+	FGraphSplineOverlapResult LastClickSplineOverlap;
 
 	/** The mouse state from the last mouse move event, used to synthesize pin actions when hovering over a spline on the panel */
 	FGeometry LastPointerGeometry;
@@ -403,6 +422,9 @@ protected:
 	/** Whether to draw decorations for graph state (PIE / ReadOnly etc.) */
 	TAttribute<bool> ShowGraphStateOverlay;
 
+	/** Whether to allow slicing of connections in this graph panel */
+	TAttribute<bool> AllowConnectionSlicing;
+
 private:
 	/** Set of nodes selected by the user, tracked while a visual update is pending */
 	TSet<TWeakObjectPtr<class UEdGraphNode>> UserSelectedNodes;
@@ -469,6 +491,9 @@ private:
 
 	/** A flag indicating that we need to check if we have any nodes which use invalid pointers */
 	bool bCheckNodeGraphObjValidity;
+
+	/** Exists while the user is making a slice, and tracks the line points and any connections intersecting the line */
+	TOptional<FConnectionSlice> CurrentSlice;
 };
 
 #undef UE_API
diff --git a/Engine/Source/Editor/PhysicsAssetEditor/Private/PhysicsAssetGraph/SPhysicsAssetGraph.cpp b/Engine/Source/Editor/PhysicsAssetEditor/Private/PhysicsAssetGraph/SPhysicsAssetGraph.cpp
--- a/Engine/Source/Editor/PhysicsAssetEditor/Private/PhysicsAssetGraph/SPhysicsAssetGraph.cpp
+++ b/Engine/Source/Editor/PhysicsAssetEditor/Private/PhysicsAssetGraph/SPhysicsAssetGraph.cpp
@@ -68,7 +68,8 @@ void SPhysicsAssetGraph::Construct(const FArguments& InArgs, const TSharedRef<FP
 		.GraphToEdit(GraphObj)
 		.GraphEvents(GraphEvents)
 		.Appearance(AppearanceInfo)
-		.ShowGraphStateOverlay(false);
+		.ShowGraphStateOverlay(false)
+		.AllowConnectionSlicing(false);
 
 	static const FName DefaultForegroundName("DefaultForeground");
 
diff --git a/Engine/Source/Editor/UnrealEd/Private/GraphEditor.cpp b/Engine/Source/Editor/UnrealEd/Private/GraphEditor.cpp
--- a/Engine/Source/Editor/UnrealEd/Private/GraphEditor.cpp
+++ b/Engine/Source/Editor/UnrealEd/Private/GraphEditor.cpp
@@ -45,6 +45,7 @@ void SGraphEditor::ConstructImplementation( const FArguments& InArgs )
 		InArgs._OnNavigateHistoryBack,
 		InArgs._OnNavigateHistoryForward,
 		InArgs._ShowGraphStateOverlay,
+		InArgs._AllowConnectionSlicing,
 		InArgs._OnTick
 		);
 
diff --git a/Engine/Source/Editor/UnrealEd/Public/GraphEditor.h b/Engine/Source/Editor/UnrealEd/Public/GraphEditor.h
--- a/Engine/Source/Editor/UnrealEd/Public/GraphEditor.h
+++ b/Engine/Source/Editor/UnrealEd/Public/GraphEditor.h
@@ -192,6 +192,7 @@ PRAGMA_ENABLE_DEPRECATION_WARNINGS
 		, _GraphToEdit(NULL)
 		, _AutoExpandActionMenu(false)
 		, _ShowGraphStateOverlay(true)
+		, _AllowConnectionSlicing(true)
 		{}
 
 		SLATE_ARGUMENT( TSharedPtr<FUICommandList>, AdditionalCommands )
@@ -221,7 +222,8 @@ PRAGMA_ENABLE_DEPRECATION_WARNINGS
 		SLATE_EVENT(FOnGraphEditorTickDelegate, OnTick);
 
 		/** Show overlay elements for the graph state such as the PIE and read-only borders and text */
-		SLATE_ATTRIBUTE(bool, ShowGraphStateOverlay)				
+		SLATE_ATTRIBUTE(bool, ShowGraphStateOverlay)
+		SLATE_ATTRIBUTE(bool, AllowConnectionSlicing)
 	SLATE_END_ARGS()
 
 	/**
