Disable this feature for 5.7

#rnx

diff --git a/Engine/Source/Editor/GraphEditor/Private/GraphEditorSnapping/NodeSnappingManager.cpp b/Engine/Source/Editor/GraphEditor/Private/GraphEditorSnapping/NodeSnappingManager.cpp
--- a/Engine/Source/Editor/GraphEditor/Private/GraphEditorSnapping/NodeSnappingManager.cpp
+++ b/Engine/Source/Editor/GraphEditor/Private/GraphEditorSnapping/NodeSnappingManager.cpp
@@ -11,13 +11,6 @@
 namespace UE::Private
 {
 
-TAutoConsoleVariable<bool> CVarBPEnableNodeSnapping(
-	TEXT("bp.EnableNodeSnapping"),
-	0,
-	TEXT("if not 0, node snapping is enabled"),
-	ECVF_Default
-);
-
 template<typename Predicate>
 UEdGraphPin* GetOnlyVisiblePinWithDirection( UEdGraphNode* Node, EEdGraphPinDirection PinDirection, Predicate Pred )
 {
@@ -538,7 +531,7 @@ FArrangedGraphNode FNodeSnappingManager::GetGraphNodeUnderMouse(const FGeometry&
 
 bool FNodeSnappingManager::IsNodeSnappingEnabled()
 {
-	return UE::Private::CVarBPEnableNodeSnapping->GetBool();
+	return false;
 }
 
 void FNodeSnappingManager::CommitDrag()
