Scribble: NodePanel / GraphPanel requirements

GraphPanel options to disable background / shadow
NodePanel API to interact with viewoffset / zoomamount
NodePanel API to disable navigation for locked panels
NodePanel ZoomLevels container to be shared ptr so it can be shared between synced panels

#rnx

diff --git a/Engine/Source/Editor/GraphEditor/Private/SGraphNode.cpp b/Engine/Source/Editor/GraphEditor/Private/SGraphNode.cpp
--- a/Engine/Source/Editor/GraphEditor/Private/SGraphNode.cpp
+++ b/Engine/Source/Editor/GraphEditor/Private/SGraphNode.cpp
@@ -1862,6 +1862,10 @@ void SGraphNode::DrawChildrenNodeShadow( const FGraphPanelSelectionSet& Selectio
 {
 	const FVector2f NodeShadowSize = UE::Slate::CastToVector2f(GetDefault<UGraphEditorSettings>()->GetShadowDeltaSize());
 	const FSlateBrush* ShadowBrush = GetShadowBrush(SelectionSet.Contains(GraphNode));
+	if (ShadowBrush == nullptr)
+	{
+		return;
+	}
 	FGeometry ChildGeometry = GetNodeGeometry(Geometry);
 	FSlateDrawElement::MakeBox(
 		OutDrawElements,
diff --git a/Engine/Source/Editor/GraphEditor/Private/SGraphPanel.cpp b/Engine/Source/Editor/GraphEditor/Private/SGraphPanel.cpp
--- a/Engine/Source/Editor/GraphEditor/Private/SGraphPanel.cpp
+++ b/Engine/Source/Editor/GraphEditor/Private/SGraphPanel.cpp
@@ -110,6 +110,8 @@ void SGraphPanel::Construct( const SGraphPanel::FArguments& InArgs )
 	this->OnDoubleClicked = InArgs._OnDoubleClicked;
 	this->OnClicked = InArgs._OnMouseButtonDown;
 	this->OnNodeSingleClicked = InArgs._OnNodeSingleClicked;
+	this->bAllowZoom = InArgs._AllowZoom;
+	this->bAllowPanning = InArgs._AllowPanning;
 
 	this->bPreservePinPreviewConnection = false;
 	this->PinVisibility = SGraphEditor::Pin_Show;
@@ -132,6 +134,8 @@ void SGraphPanel::Construct( const SGraphPanel::FArguments& InArgs )
 	
 	ShowGraphStateOverlay = InArgs._ShowGraphStateOverlay;
 	AllowConnectionSlicing = InArgs._AllowConnectionSlicing;
+	ShouldDrawBackground = InArgs._ShouldDrawBackground;
+	ShouldDrawSurroundingShadow = InArgs._ShouldDrawSurroundingShadow;
 
 	SavedMousePosForOnPaintEventLocalSpace = FVector2f::ZeroVector;
 	PreviousFrameSavedMousePosForSplineOverlap = FVector2f::ZeroVector;
@@ -168,10 +172,13 @@ int32 SGraphPanel::OnPaint( const FPaintArgs& Args, const FGeometry& AllottedGeo
 	// First paint the background
 	const UEditorExperimentalSettings& Options = *GetDefault<UEditorExperimentalSettings>();
 
-	const FSlateBrush* DefaultBackground = FAppStyle::GetBrush(TEXT("Graph.Panel.SolidBackground"));
-	const FSlateBrush* CustomBackground = &GetDefault<UEditorStyleSettings>()->GraphBackgroundBrush;
-	const FSlateBrush* BackgroundImage = CustomBackground->HasUObject() ? CustomBackground : DefaultBackground;
-	PaintBackgroundAsLines(BackgroundImage, AllottedGeometry, MyCullingRect, OutDrawElements, LayerId);
+	if (ShouldDrawBackground.Get(true))
+	{
+		const FSlateBrush* DefaultBackground = FAppStyle::GetBrush(TEXT("Graph.Panel.SolidBackground"));
+		const FSlateBrush* CustomBackground = &GetDefault<UEditorStyleSettings>()->GraphBackgroundBrush;
+		const FSlateBrush* BackgroundImage = CustomBackground->HasUObject() ? CustomBackground : DefaultBackground;
+		PaintBackgroundAsLines(BackgroundImage, AllottedGeometry, MyCullingRect, OutDrawElements, LayerId);
+	}
 
 	const float ZoomFactor = AllottedGeometry.Scale * GetZoomAmount();
 
@@ -664,9 +671,12 @@ int32 SGraphPanel::OnPaint( const FPaintArgs& Args, const FGeometry& AllottedGeo
 		delete ConnectionDrawingPolicy;
 	}
 
-	// Draw a shadow overlay around the edges of the graph
-	++MaxLayerId;
-	PaintSurroundSunkenShadow(FAppStyle::GetBrush(TEXT("Graph.Shadow")), AllottedGeometry, MyCullingRect, OutDrawElements, MaxLayerId);
+	if (ShouldDrawSurroundingShadow.Get(true))
+	{
+		// Draw a shadow overlay around the edges of the graph
+		++MaxLayerId;
+		PaintSurroundSunkenShadow(FAppStyle::GetBrush(TEXT("Graph.Shadow")), AllottedGeometry, MyCullingRect, OutDrawElements, MaxLayerId);
+	}
 
 	if (ShowGraphStateOverlay.Get())
 	{
diff --git a/Engine/Source/Editor/GraphEditor/Private/SNodePanel.cpp b/Engine/Source/Editor/GraphEditor/Private/SNodePanel.cpp
--- a/Engine/Source/Editor/GraphEditor/Private/SNodePanel.cpp
+++ b/Engine/Source/Editor/GraphEditor/Private/SNodePanel.cpp
@@ -419,6 +419,8 @@ void SNodePanel::Construct()
 	ZoomTargetBottomRight = FVector2f::ZeroVector;
 	ZoomPadding = NodePanelDefs::DefaultZoomPadding;
 
+	bAllowZoom = true;
+	bAllowPanning = true;
 	bAllowContinousZoomInterpolation = false;
 	bTeleportInsteadOfScrollingWhenZoomingToFit = false;
 
@@ -490,6 +492,10 @@ FVector2f SNodePanel::ComputeEdgePanAmount(const FGeometry& MyGeometry, const FV
 
 void SNodePanel::UpdateViewOffset(const FGeometry& MyGeometry, const FVector2f& TargetPosition)
 {
+	if (!bAllowPanning)
+	{
+		return;
+	}
 	const FVector2f PanAmount = ComputeEdgePanAmount( MyGeometry, TargetPosition ) / GetZoomAmount();
 	ViewOffset += PanAmount;
 }
@@ -550,6 +556,12 @@ const FNodeSnappingManager* SNodePanel::GetSnappingManager() const
 
 EActiveTimerReturnType SNodePanel::HandleZoomToFit(double InCurrentTime, float InDeltaTime)
 {
+	// don't run this timer if we don't allow pan or zoom
+	if (!bAllowPanning || !bAllowZoom)
+	{
+		return EActiveTimerReturnType::Stop;
+	}
+	
 	const FVector2f DesiredViewCenter = ( ZoomTargetTopLeft + ZoomTargetBottomRight ) * 0.5f;
 	const bool bDoneScrolling = ScrollToLocation(CachedGeometry, DesiredViewCenter, bTeleportInsteadOfScrollingWhenZoomingToFit ? 1000.0f : InDeltaTime);
 	bool bDoneZooming = ZoomToLocation(CachedGeometry.GetLocalSize(), ZoomTargetBottomRight - ZoomTargetTopLeft, bDoneScrolling);
@@ -665,6 +677,11 @@ FReply SNodePanel::OnMouseButtonDown( const FGeometry& MyGeometry, const FPointe
 
 	auto PerformMousePan = [this, &MyGeometry, &MouseEvent]()
 	{
+		if (!bAllowPanning)
+		{
+			return FReply::Unhandled();
+		}
+			
 		// Cache current cursor position as zoom origin and software cursor position
 		ZoomStartOffset = MyGeometry.AbsoluteToLocal( MouseEvent.GetLastScreenSpacePosition() );
 		SoftwareCursorPosition = PanelCoordToGraphCoord( ZoomStartOffset );
@@ -1758,6 +1775,11 @@ bool SNodePanel::ScrollToLocation(const FGeometry& MyGeometry, FVector2f Desired
 
 bool SNodePanel::ZoomToLocation(const FVector2f& CurrentSizeWithoutZoom, const FVector2f& InDesiredSize, bool bDoneScrolling)
 {
+	if (!bAllowZoom)
+	{
+		return false;
+	}
+
 	if (bAllowContinousZoomInterpolation && ZoomLevelGraphFade.IsPlaying())
 	{
 		return false;
@@ -1819,6 +1841,11 @@ bool SNodePanel::ZoomToLocation(const FVector2f& CurrentSizeWithoutZoom, const F
 
 void SNodePanel::ZoomToFit(bool bOnlySelection)
 {
+	if (!bAllowZoom)
+	{
+		return;
+	}
+
 	bDeferredZoomToNodeExtents = true;
 	bDeferredZoomToSelection = bOnlySelection;
 	ZoomPadding = NodePanelDefs::DefaultZoomPadding;
@@ -1826,6 +1853,11 @@ void SNodePanel::ZoomToFit(bool bOnlySelection)
 
 void SNodePanel::ZoomToTarget(const FVector2f& TopLeft, const FVector2f& BottomRight)
 {
+	if (!bAllowZoom)
+	{
+		return;
+	}
+	
 	bDeferredZoomToNodeExtents = false;
 
 	ZoomTargetTopLeft = TopLeft;
@@ -1836,6 +1868,11 @@ void SNodePanel::ZoomToTarget(const FVector2f& TopLeft, const FVector2f& BottomR
 
 void SNodePanel::ChangeZoomLevel(int32 ZoomLevelDelta, const FVector2f& WidgetSpaceZoomOrigin, bool bOverrideZoomLimiting)
 {
+	if (!bAllowZoom)
+	{
+		return;
+	}
+	
 	// We want to zoom into this point; i.e. keep it the same fraction offset into the panel
 	const FVector2f PointToMaintainGraphSpace = PanelCoordToGraphCoord( WidgetSpaceZoomOrigin );
 
@@ -1875,6 +1912,19 @@ void SNodePanel::ChangeZoomLevel(int32 ZoomLevelDelta, const FVector2f& WidgetSp
 	this->ViewOffset = PointToMaintainGraphSpace - WidgetSpaceZoomOrigin / GetZoomAmount();
 }
 
+void SNodePanel::SetZoomLevel(int32 InZoomLevel)
+{
+	if (!ZoomLevels)
+	{
+		return;
+	}
+	if (InZoomLevel >= 0 && InZoomLevel < ZoomLevels->GetNumZoomLevels())
+	{
+		ZoomLevel = InZoomLevel;
+		PostChangedZoom();
+	}
+}
+
 bool SNodePanel::GetBoundsForSelectedNodes( class FSlateRect& Rect, float Padding )
 {
 	bool Result = false;
diff --git a/Engine/Source/Editor/GraphEditor/Public/SGraphPanel.h b/Engine/Source/Editor/GraphEditor/Public/SGraphPanel.h
--- a/Engine/Source/Editor/GraphEditor/Public/SGraphPanel.h
+++ b/Engine/Source/Editor/GraphEditor/Public/SGraphPanel.h
@@ -102,7 +102,11 @@ public:
 		, _DisplayAsReadOnly( false )
 		, _ShowGraphStateOverlay(true)
 		, _AllowConnectionSlicing(true)
+		, _ShouldDrawBackground(true)
+		, _ShouldDrawSurroundingShadow(true)
 		, _OnUpdateGraphPanel()
+		, _AllowZoom(true)
+		, _AllowPanning(true)
 		{
 			_Clipping = EWidgetClipping::ClipToBounds;
 		}
@@ -155,6 +159,8 @@ PRAGMA_ENABLE_DEPRECATION_WARNINGS
 		/** Show overlay elements for the graph state such as the PIE and read-only borders and text */
 		SLATE_ATTRIBUTE(bool, ShowGraphStateOverlay)
 		SLATE_ATTRIBUTE(bool, AllowConnectionSlicing)
+		SLATE_ATTRIBUTE(bool, ShouldDrawBackground)
+		SLATE_ATTRIBUTE(bool, ShouldDrawSurroundingShadow)
 		SLATE_EVENT( FOnNodeVerifyTextCommit, OnVerifyTextCommit )
 		SLATE_EVENT( FOnNodeTextCommitted, OnTextCommitted )
 		
@@ -182,6 +188,8 @@ PRAGMA_ENABLE_DEPRECATION_WARNINGS
 		SLATE_EVENT( SGraphEditor::FOnMouseButtonDown, OnMouseButtonDown )
 		SLATE_EVENT( SGraphEditor::FOnNodeSingleClicked, OnNodeSingleClicked )
 		//SLATE_ATTRIBUTE( FGraphAppearanceInfo, Appearance )
+		SLATE_ARGUMENT(bool, AllowZoom)
+		SLATE_ARGUMENT(bool, AllowPanning)
 	SLATE_END_ARGS()
 
 	/**
@@ -424,6 +432,12 @@ protected:
 	/** Whether to allow slicing of connections in this graph panel */
 	TAttribute<bool> AllowConnectionSlicing;
 
+	/** Whether this panel should be drawing a background */
+	TAttribute<bool> ShouldDrawBackground;
+
+	/** Whether this panel should be drawing a sunken shadow on its borders */
+	TAttribute<bool> ShouldDrawSurroundingShadow;
+
 private:
 	/** Set of nodes selected by the user, tracked while a visual update is pending */
 	TSet<TWeakObjectPtr<class UEdGraphNode>> UserSelectedNodes;
diff --git a/Engine/Source/Editor/GraphEditor/Public/SNodePanel.h b/Engine/Source/Editor/GraphEditor/Public/SNodePanel.h
--- a/Engine/Source/Editor/GraphEditor/Public/SNodePanel.h
+++ b/Engine/Source/Editor/GraphEditor/Public/SNodePanel.h
@@ -839,6 +839,9 @@ public:
 	/** @return the view offset in graph space */
 	UE_API UE::Slate::FDeprecateVector2DResult GetViewOffset() const;
 
+	/** Sets the current view offset */
+	void SetViewOffset(const UE::Slate::FDeprecateVector2DResult& InViewOffset) { ViewOffset = InViewOffset; }
+
 	/** 
 	 * when a panel is scrolling/zooming to a target, this can be called to get it's destination
 	 * @param TopLeft top left corner of the destination
@@ -849,6 +852,12 @@ public:
 	UE_API bool GetZoomTargetRect(FVector2D& TopLeft, FVector2D& BottomRight) const;
 	UE_API bool GetZoomTargetRect(FVector2f& TopLeft, FVector2f& BottomRight) const;
 
+	/** @return the current zoom level */
+	int32 GetZoomLevel() const { return ZoomLevel; }
+
+	/** Sets the current zoom level */
+	UE_API void SetZoomLevel(int32 InZoomLevel);
+
 	/** @return the current view bookmark ID */
 	const FGuid& GetViewBookmarkId() const { return CurrentBookmarkGuid; }
 
@@ -1048,19 +1057,26 @@ public:
 	template<typename T>
 	void SetZoomLevelsContainer()
 	{
-		ZoomLevels = MakeUnique<T>();
-		OldZoomAmount = ZoomLevels->GetZoomAmount(ZoomLevel);
-		ZoomLevel = PreviousZoomLevel = ZoomLevels->GetNearestZoomLevel(OldZoomAmount);
+		SetZoomLevelsContainer(MakeShared<T>());
+	}
+
+	void SetZoomLevelsContainer(const TSharedPtr<FZoomLevelsContainer>& InZoomLevels)
+	{
+		ZoomLevels = InZoomLevels;
+		OldZoomAmount = InZoomLevels->GetZoomAmount(ZoomLevel);
+		ZoomLevel = PreviousZoomLevel = InZoomLevels->GetNearestZoomLevel(OldZoomAmount);
 		PostChangedZoom();
 	}
 
+	const TSharedPtr<FZoomLevelsContainer>& GetZoomLevels() const { return ZoomLevels; }
+
 	void SetCurrentHoveredNode(UObject* InNode) { CurrentHoveredNode = InNode; }
 	UObject* GetCurrentHoveredNode() { return CurrentHoveredNode; }
 	
 protected:
 	
 	// The interface for mapping ZoomLevel values to actual node scaling values
-	TUniquePtr<FZoomLevelsContainer> ZoomLevels;
+	TSharedPtr<FZoomLevelsContainer> ZoomLevels;
 
 	/** The position within the graph at which the user is looking */
 	FDeprecateSlateVector2D ViewOffset;
@@ -1135,6 +1151,12 @@ protected:
 	/** Zoom selection padding */
 	float ZoomPadding;
 
+	/** Allow Zoom overall */
+	bool bAllowZoom;
+
+	/** Allow Panning overall */
+	bool bAllowPanning;
+
 	/** Allow continous zoom interpolation? */
 	bool bAllowContinousZoomInterpolation;
 
diff --git a/Engine/Source/Editor/WorldBrowser/Private/Tiles/SWorldComposition.cpp b/Engine/Source/Editor/WorldBrowser/Private/Tiles/SWorldComposition.cpp
--- a/Engine/Source/Editor/WorldBrowser/Private/Tiles/SWorldComposition.cpp
+++ b/Engine/Source/Editor/WorldBrowser/Private/Tiles/SWorldComposition.cpp
@@ -109,7 +109,7 @@ public:
 
 	void Construct(const FArguments& InArgs)
 	{
-		ZoomLevels = MakeUnique<FWorldZoomLevelsContainer>();
+		SetZoomLevelsContainer<FWorldZoomLevelsContainer>();
 
 		SNodePanel::Construct();
 
