Slate Dynamic Invalidation - Material Editor Graph Pin

#Synced-CL 48711348

diff --git a/Engine/Source/Editor/GraphEditor/Private/SGraphPin.cpp b/Engine/Source/Editor/GraphEditor/Private/SGraphPin.cpp
--- a/Engine/Source/Editor/GraphEditor/Private/SGraphPin.cpp
+++ b/Engine/Source/Editor/GraphEditor/Private/SGraphPin.cpp
@@ -102,7 +102,8 @@ TSharedPtr<SGraphPin> FGraphPinHandle::FindInGraphPanel(const SGraphPanel& InPan
 
 
 SGraphPin::SGraphPin()
-	: GraphPinObj(nullptr)
+	: PinImageHighDetailSlot(nullptr)
+	, GraphPinObj(nullptr)
 	, PinColorModifier(FLinearColor::White)
 	, CachedNodeOffset(FVector2f::ZeroVector)
 	, Custom_Brush_Connected(nullptr)
@@ -115,6 +116,7 @@ SGraphPin::SGraphPin()
 	, bDragAndDropEnabled(true)
 	, bIsDiffHighlighted(false)
 	, bFadeConnections(false)
+	, bIsSnappingEnabled(false)
 {
 	IsEditable = true;
 
@@ -338,9 +340,7 @@ void SGraphPin::Construct(const FArguments& InArgs, UEdGraphPin* InPin)
 			.AutoWidth()
 			.VAlign(VAlign_Center)
 			.Padding(0.0f, 0.0f, InArgs._SideToSideMargin, 0.0f)
-			[
-				PinWidgetRef
-			]
+			.Expose(PinImageHighDetailSlot)
 			+SHorizontalBox::Slot()
 			.AutoWidth()
 			.VAlign(VAlign_Center)
@@ -362,9 +362,7 @@ void SGraphPin::Construct(const FArguments& InArgs, UEdGraphPin* InPin)
 			.AutoWidth()
 			.VAlign(VAlign_Center)
 			.Padding(InArgs._SideToSideMargin, 0.0f, 0.0f, 0.0f)
-			[
-				PinWidgetRef
-			];
+			.Expose(PinImageHighDetailSlot);
 	}
 
 	// Set up a hover for pins that is tinted the color of the pin.
@@ -380,11 +378,7 @@ void SGraphPin::Construct(const FArguments& InArgs, UEdGraphPin* InPin)
 			[
 				SAssignNew(PinNameLODBranchNode, SLevelOfDetailBranchNode)
 				.UseLowDetailSlot(this, &SGraphPin::UseLowDetailPinNames)
-				.LowDetail()
-				[
-					//@TODO: Try creating a pin-colored line replacement that doesn't measure text / call delegates but still renders
-					PinWidgetRef
-				]
+				.OnGetActiveDetailSlotContent(this, &SGraphPin::OnGetActiveDetailSlotContent)
 				.HighDetail()
 				[
 					PinContent.ToSharedRef()
@@ -425,7 +419,7 @@ void SGraphPin::RefreshLOD()
 
 void SGraphPin::SetIsEditable(TAttribute<bool> InIsEditable)
 {
-	IsEditable = InIsEditable;
+	IsEditable = MoveTemp(InIsEditable);
 }
 
 FReply SGraphPin::OnPinMouseDown( const FGeometry& SenderGeometry, const FPointerEvent& MouseEvent )
@@ -1512,6 +1506,37 @@ bool SGraphPin::UseLowDetailPinNames() const
 	return false;
 }
 
+TSharedRef<SWidget> SGraphPin::OnGetActiveDetailSlotContent(bool bChangedToHighDetail) const
+{
+	if (PinImage)
+	{
+		// Here, we need to make sure we detach PinImage before setting it into the other slot
+		PinImage->ConditionallyDetatchParentWidget(PinImage->GetParentWidget().Get());
+
+		if (bChangedToHighDetail)
+		{
+			if (const TSharedPtr<SHorizontalBox> RowWidget = FullPinHorizontalRowWidget.Pin())
+			{
+				for (int32 i = 0; i < RowWidget->GetChildren()->NumSlot(); ++i)
+				{
+					if (&RowWidget->GetSlot(i) == PinImageHighDetailSlot)
+					{
+						RowWidget->GetSlot(i).AttachWidget(PinImage.ToSharedRef());
+						break;
+					}
+				}
+				return RowWidget.ToSharedRef();
+			}
+		}
+		else
+		{
+			return PinImage.ToSharedRef();
+		}
+	}
+
+	return SNullWidget::NullWidget;
+}
+
 EVisibility SGraphPin::GetPinVisiblity() const
 {
 	// The pin becomes too small to use at low LOD, so disable the hit test.
diff --git a/Engine/Source/Editor/GraphEditor/Public/SGraphPin.h b/Engine/Source/Editor/GraphEditor/Public/SGraphPin.h
--- a/Engine/Source/Editor/GraphEditor/Public/SGraphPin.h
+++ b/Engine/Source/Editor/GraphEditor/Public/SGraphPin.h
@@ -99,6 +99,13 @@ struct FGraphPinHandle
 /////////////////////////////////////////////////////
 // SGraphPin
 
+
+template <>
+struct TWidgetTypeTraits<class SGraphPin>
+{
+	static constexpr bool SupportsInvalidation() { return true; }
+};
+
 /**
  * Represents a pin on a node in the GraphEditor
  */
@@ -342,6 +349,9 @@ protected:
 	// Should we use low-detail pin names?
 	UE_API virtual bool UseLowDetailPinNames() const;
 
+	// What widget should we show when we change detail slot?
+	UE_API virtual TSharedRef<SWidget> OnGetActiveDetailSlotContent(bool bChangedToHighDetail) const;
+
 	/** Determines the pin's visibility based on the LOD factor, when it is low LOD, no hit test will occur */
 	UE_API EVisibility GetPinVisiblity() const;
 
@@ -369,6 +379,9 @@ protected:
 	/** Horizontal box that holds the full detail pin widget, useful for outsiders to inject widgets into the pin */
 	TWeakPtr<SHorizontalBox> FullPinHorizontalRowWidget;
 
+	/** The slot in which PinImage will be attached when PinNameLODBranchNode switches to HighDetail. Might be stale or null. */
+	SHorizontalBox::FSlot* PinImageHighDetailSlot;
+
 	/** Value widget for the pin, created with GetDefaultValueWidget() */
 	TSharedPtr<SWidget> ValueWidget;
 
diff --git a/Engine/Source/Editor/KismetWidgets/Private/SLevelOfDetailBranchNode.cpp b/Engine/Source/Editor/KismetWidgets/Private/SLevelOfDetailBranchNode.cpp
--- a/Engine/Source/Editor/KismetWidgets/Private/SLevelOfDetailBranchNode.cpp
+++ b/Engine/Source/Editor/KismetWidgets/Private/SLevelOfDetailBranchNode.cpp
@@ -25,16 +25,7 @@ void SLevelOfDetailBranchNode::Construct(const FArguments& InArgs)
 	ChildSlotLowDetail = InArgs._LowDetail.Widget;
 	ChildSlotHighDetail = InArgs._HighDetail.Widget;
 
-	const int32 CurrentValue = ShowLowDetailAttr.Get() ? 1 : 0;
-	if (OnGetActiveDetailSlotContent.IsBound())
-	{
-		ChildSlot[OnGetActiveDetailSlotContent.Execute((CurrentValue != 0) ? false : true)];
-	}
-	else
-	{
-		ChildSlot[(CurrentValue != 0) ? ChildSlotLowDetail : ChildSlotHighDetail];
-	}
-
+	RefreshLODSlotContent();
 }
 
 void SLevelOfDetailBranchNode::Tick( const FGeometry& AllottedGeometry, const double InCurrentTime, const float InDeltaTime )
