Removing `SNodePanel`'s unused attributes

#Synced-CL 49309628

diff --git a/Engine/Plugins/AI/EnvironmentQueryEditor/Source/EnvironmentQueryEditor/Private/SGraphNode_EnvironmentQuery.cpp b/Engine/Plugins/AI/EnvironmentQueryEditor/Source/EnvironmentQueryEditor/Private/SGraphNode_EnvironmentQuery.cpp
--- a/Engine/Plugins/AI/EnvironmentQueryEditor/Source/EnvironmentQueryEditor/Private/SGraphNode_EnvironmentQuery.cpp
+++ b/Engine/Plugins/AI/EnvironmentQueryEditor/Source/EnvironmentQueryEditor/Private/SGraphNode_EnvironmentQuery.cpp
@@ -174,7 +174,6 @@ void SGraphNode_EnvironmentQuery::UpdateGraphNode()
 		return FMath::Max(22.0f, DesiredHeight);
 	};
 
-	this->ContentScale.Bind( this, &SGraphNode::GetContentScale );
 	this->GetOrAddSlot( ENodeZone::Center )
 		.HAlign(HAlign_Fill)
 		.VAlign(VAlign_Center)
diff --git a/Engine/Plugins/Animation/ControlRig/Source/ControlRigEditor/Private/RigDependencyGraph/SRigDependencyGraphNode.cpp b/Engine/Plugins/Animation/ControlRig/Source/ControlRigEditor/Private/RigDependencyGraph/SRigDependencyGraphNode.cpp
--- a/Engine/Plugins/Animation/ControlRig/Source/ControlRigEditor/Private/RigDependencyGraph/SRigDependencyGraphNode.cpp
+++ b/Engine/Plugins/Animation/ControlRig/Source/ControlRigEditor/Private/RigDependencyGraph/SRigDependencyGraphNode.cpp
@@ -165,8 +165,6 @@ void SRigDependencyGraphNode::UpdateGraphNode()
 	RightNodeBox.Reset();
 	LeftNodeBox.Reset();
 
-	this->ContentScale.Bind( this, &SGraphNode::GetContentScale );
-
 	this->GetOrAddSlot( ENodeZone::Center )
 		.HAlign(HAlign_Center)
 		.VAlign(VAlign_Center)
diff --git a/Engine/Plugins/Cameras/GameplayCameras/Source/GameplayCamerasEditor/Private/Editors/SCameraObjectInterfaceParameterGraphNode.cpp b/Engine/Plugins/Cameras/GameplayCameras/Source/GameplayCamerasEditor/Private/Editors/SCameraObjectInterfaceParameterGraphNode.cpp
--- a/Engine/Plugins/Cameras/GameplayCameras/Source/GameplayCamerasEditor/Private/Editors/SCameraObjectInterfaceParameterGraphNode.cpp
+++ b/Engine/Plugins/Cameras/GameplayCameras/Source/GameplayCamerasEditor/Private/Editors/SCameraObjectInterfaceParameterGraphNode.cpp
@@ -37,8 +37,6 @@ void SCameraObjectInterfaceParameterGraphNode::UpdateGraphNode()
 
 	SetupErrorReporting();
 
-	ContentScale.Bind(this, &SGraphNode::GetContentScale);
-
 	GetOrAddSlot(ENodeZone::Center)
 	.HAlign(HAlign_Center)
 	.VAlign(VAlign_Center)
diff --git a/Engine/Plugins/Editor/AssetManagerEditor/Source/AssetManagerEditor/Private/ReferenceViewer/SReferenceNode.cpp b/Engine/Plugins/Editor/AssetManagerEditor/Source/AssetManagerEditor/Private/ReferenceViewer/SReferenceNode.cpp
--- a/Engine/Plugins/Editor/AssetManagerEditor/Source/AssetManagerEditor/Private/ReferenceViewer/SReferenceNode.cpp
+++ b/Engine/Plugins/Editor/AssetManagerEditor/Source/AssetManagerEditor/Private/ReferenceViewer/SReferenceNode.cpp
@@ -130,7 +130,6 @@ void SReferenceNode::UpdateGraphNode()
 			];
 	}
 
-	ContentScale.Bind( this, &SReferenceNode::GetContentScale );
 	GetOrAddSlot( ENodeZone::Center )
 	.HAlign(HAlign_Center)
 	.VAlign(VAlign_Center)
diff --git a/Engine/Plugins/Editor/AssetManagerEditor/Source/AssetManagerEditor/Private/ReferenceViewer/SReferencedPropertiesNode.cpp b/Engine/Plugins/Editor/AssetManagerEditor/Source/AssetManagerEditor/Private/ReferenceViewer/SReferencedPropertiesNode.cpp
--- a/Engine/Plugins/Editor/AssetManagerEditor/Source/AssetManagerEditor/Private/ReferenceViewer/SReferencedPropertiesNode.cpp
+++ b/Engine/Plugins/Editor/AssetManagerEditor/Source/AssetManagerEditor/Private/ReferenceViewer/SReferencedPropertiesNode.cpp
@@ -76,8 +76,6 @@ void SReferencedPropertiesNode::UpdateGraphNode()
 	RightNodeBox.Reset();
 	LeftNodeBox.Reset();
 
-	ContentScale.Bind(this, &SReferencedPropertiesNode::GetContentScale);
-
 	TSharedPtr<SWidget> MainWidget;
 
 	if (const UEdGraphNode_ReferencedProperties* ReferencedProperties = Cast<UEdGraphNode_ReferencedProperties>(GraphNode))
diff --git a/Engine/Plugins/Enterprise/DataprepEditor/Source/DataprepEditor/Private/Widgets/DataprepGraph/SDataprepGraphActionNode.cpp b/Engine/Plugins/Enterprise/DataprepEditor/Source/DataprepEditor/Private/Widgets/DataprepGraph/SDataprepGraphActionNode.cpp
--- a/Engine/Plugins/Enterprise/DataprepEditor/Source/DataprepEditor/Private/Widgets/DataprepGraph/SDataprepGraphActionNode.cpp
+++ b/Engine/Plugins/Enterprise/DataprepEditor/Source/DataprepEditor/Private/Widgets/DataprepGraph/SDataprepGraphActionNode.cpp
@@ -51,7 +51,6 @@ public:
 	// SGraphNode interface
 	virtual void UpdateGraphNode() override
 	{
-		this->ContentScale.Bind( this, &SGraphNode::GetContentScale );
 		this->GetOrAddSlot( ENodeZone::Center )
 		.HAlign(HAlign_Fill)
 		.VAlign(VAlign_Center)
@@ -481,8 +480,6 @@ void SDataprepGraphActionNode::UpdateGraphNode()
 	RightNodeBox.Reset();
 	LeftNodeBox.Reset();
 
-	this->ContentScale.Bind( this, &SGraphNode::GetContentScale );
-
 	if(!DataprepActionPtr.IsValid())
 	{
 		this->GetOrAddSlot( ENodeZone::Center )
@@ -689,8 +686,6 @@ FReply SDataprepGraphActionNode::OnMouseButtonDown(const FGeometry& MyGeometry,
 
 	if ( !Reply.IsEventHandled() )
 	{
-		BorderBackgroundColor.Set(FDataprepEditorStyle::GetColor("DataprepActionStep.DragAndDrop"));
-
 		if( MouseEvent.GetEffectingButton() ==  EKeys::LeftMouseButton )
 		{
 			GetOwnerPanel()->SelectionManager.ClickedOnNode( GraphNode, MouseEvent );
@@ -1243,8 +1238,6 @@ FReply SDataprepGraphActionGroupNode::OnMouseButtonDown(const FGeometry& MyGeome
 
 	if ( !Reply.IsEventHandled() )
 	{
-		BorderBackgroundColor.Set(FDataprepEditorStyle::GetColor("DataprepActionStep.DragAndDrop"));
-
 		if( MouseEvent.GetEffectingButton() ==  EKeys::LeftMouseButton )
 		{
 			GetOwnerPanel()->SelectionManager.ClickedOnNode( GraphNode, MouseEvent );
diff --git a/Engine/Plugins/Enterprise/DataprepEditor/Source/DataprepEditor/Private/Widgets/DataprepGraph/SDataprepGraphActionStepNode.cpp b/Engine/Plugins/Enterprise/DataprepEditor/Source/DataprepEditor/Private/Widgets/DataprepGraph/SDataprepGraphActionStepNode.cpp
--- a/Engine/Plugins/Enterprise/DataprepEditor/Source/DataprepEditor/Private/Widgets/DataprepGraph/SDataprepGraphActionStepNode.cpp
+++ b/Engine/Plugins/Enterprise/DataprepEditor/Source/DataprepEditor/Private/Widgets/DataprepGraph/SDataprepGraphActionStepNode.cpp
@@ -152,7 +152,6 @@ void SDataprepGraphActionStepNode::UpdateGraphNode()
 	TAttribute<FMargin> ArrowOverlayPadding = TAttribute<FMargin>::Create(TAttribute<FMargin>::FGetter::CreateSP(this, &SDataprepGraphActionStepNode::GetArrowPadding));
 	TAttribute<FSlateColor> BlockColorAndOpacity = TAttribute<FSlateColor>::Create(TAttribute<FSlateColor>::FGetter::CreateSP(this, &SDataprepGraphActionStepNode::GetBlockOverlayColor));
 
-	this->ContentScale.Bind( this, &SGraphNode::GetContentScale );
 	this->GetOrAddSlot( ENodeZone::Center )
 	.HAlign(HAlign_Fill)
 	.VAlign(VAlign_Center)
diff --git a/Engine/Plugins/Enterprise/DataprepEditor/Source/DataprepEditor/Private/Widgets/DataprepGraph/SDataprepGraphTrackNode.cpp b/Engine/Plugins/Enterprise/DataprepEditor/Source/DataprepEditor/Private/Widgets/DataprepGraph/SDataprepGraphTrackNode.cpp
--- a/Engine/Plugins/Enterprise/DataprepEditor/Source/DataprepEditor/Private/Widgets/DataprepGraph/SDataprepGraphTrackNode.cpp
+++ b/Engine/Plugins/Enterprise/DataprepEditor/Source/DataprepEditor/Private/Widgets/DataprepGraph/SDataprepGraphTrackNode.cpp
@@ -424,8 +424,6 @@ void SDataprepGraphTrackNode::UpdateGraphNode()
 		}
 	}
 
-	ContentScale.Bind( this, &SGraphNode::GetContentScale );
-
 	GetOrAddSlot( ENodeZone::Center )
 	.HAlign(HAlign_Fill)
 	.VAlign(VAlign_Center)
diff --git a/Engine/Plugins/Experimental/Animation/RigMapper/Source/RigMapperEditor/Private/RigMapperGraph/SRigMapperDefinitionGraphEditorNode.cpp b/Engine/Plugins/Experimental/Animation/RigMapper/Source/RigMapperEditor/Private/RigMapperGraph/SRigMapperDefinitionGraphEditorNode.cpp
--- a/Engine/Plugins/Experimental/Animation/RigMapper/Source/RigMapperEditor/Private/RigMapperGraph/SRigMapperDefinitionGraphEditorNode.cpp
+++ b/Engine/Plugins/Experimental/Animation/RigMapper/Source/RigMapperEditor/Private/RigMapperGraph/SRigMapperDefinitionGraphEditorNode.cpp
@@ -61,8 +61,6 @@ void SRigMapperDefinitionGraphEditorNode::UpdateGraphNode()
 	RightNodeBox.Reset();
 	LeftNodeBox.Reset();
 
-	this->ContentScale.Bind( this, &SGraphNode::GetContentScale );
-
 	this->GetOrAddSlot( ENodeZone::Center )
 		.HAlign(HAlign_Center)
 		.VAlign(VAlign_Center)
diff --git a/Engine/Plugins/Experimental/Animation/Scribble/Source/ScribbleEditor/Private/SScribbleGraphNode.cpp b/Engine/Plugins/Experimental/Animation/Scribble/Source/ScribbleEditor/Private/SScribbleGraphNode.cpp
--- a/Engine/Plugins/Experimental/Animation/Scribble/Source/ScribbleEditor/Private/SScribbleGraphNode.cpp
+++ b/Engine/Plugins/Experimental/Animation/Scribble/Source/ScribbleEditor/Private/SScribbleGraphNode.cpp
@@ -41,7 +41,6 @@ const FSlateBrush* SScribbleGraphNode::GetShadowBrush(bool bSelected) const
 
 void SScribbleGraphNode::UpdateGraphNode()
 {
-	this->ContentScale.Bind( this, &SGraphNode::GetContentScale );
 }
 
 FReply SScribbleGraphNode::OnMouseButtonDown(const FGeometry& MyGeometry, const FPointerEvent& MouseEvent)
diff --git a/Engine/Plugins/Experimental/BlueprintSnapNodes/Source/BlueprintSnapNodes/Private/SGraphNodeSnapContainer.cpp b/Engine/Plugins/Experimental/BlueprintSnapNodes/Source/BlueprintSnapNodes/Private/SGraphNodeSnapContainer.cpp
--- a/Engine/Plugins/Experimental/BlueprintSnapNodes/Source/BlueprintSnapNodes/Private/SGraphNodeSnapContainer.cpp
+++ b/Engine/Plugins/Experimental/BlueprintSnapNodes/Source/BlueprintSnapNodes/Private/SGraphNodeSnapContainer.cpp
@@ -61,7 +61,6 @@ void SGraphNodeSnapContainer::UpdateGraphNode()
 	//            |       |      | T (>) |
 	//            |_______|______|_______|
 	//
-	this->ContentScale.Bind( this, &SGraphNode::GetContentScale );
 	this->GetOrAddSlot( ENodeZone::Center )
 		.HAlign(HAlign_Center)
 		.VAlign(VAlign_Center)
diff --git a/Engine/Plugins/Experimental/CommonConversation/Source/CommonConversationGraph/Private/SConversationGraphNode.cpp b/Engine/Plugins/Experimental/CommonConversation/Source/CommonConversationGraph/Private/SConversationGraphNode.cpp
--- a/Engine/Plugins/Experimental/CommonConversation/Source/CommonConversationGraph/Private/SConversationGraphNode.cpp
+++ b/Engine/Plugins/Experimental/CommonConversation/Source/CommonConversationGraph/Private/SConversationGraphNode.cpp
@@ -340,7 +340,6 @@ void SConversationGraphNode::UpdateGraphNode()
 
 	PropertyRowGenerator->OnRowsRefreshed().AddSP(this, &SConversationGraphNode::PropertyRowsRefreshed);
 
-	this->ContentScale.Bind( this, &SGraphNode::GetContentScale );
 	this->GetOrAddSlot( ENodeZone::Center )
 		.HAlign(HAlign_Fill)
 		.VAlign(VAlign_Center)
diff --git a/Engine/Plugins/Experimental/PluginReferenceViewer/Source/PluginReferenceViewer/Private/SPluginReferenceNode.cpp b/Engine/Plugins/Experimental/PluginReferenceViewer/Source/PluginReferenceViewer/Private/SPluginReferenceNode.cpp
--- a/Engine/Plugins/Experimental/PluginReferenceViewer/Source/PluginReferenceViewer/Private/SPluginReferenceNode.cpp
+++ b/Engine/Plugins/Experimental/PluginReferenceViewer/Source/PluginReferenceViewer/Private/SPluginReferenceNode.cpp
@@ -98,7 +98,6 @@ void SPluginReferenceNode::UpdateGraphNode()
 			];
 	}
 
-	ContentScale.Bind(this, &SPluginReferenceNode::GetContentScale);
 	GetOrAddSlot(ENodeZone::Center)
 	.HAlign(HAlign_Center)
 	.VAlign(VAlign_Center)
diff --git a/Engine/Plugins/FX/Niagara/Source/NiagaraEditor/Private/Widgets/SNiagaraGraphNode.cpp b/Engine/Plugins/FX/Niagara/Source/NiagaraEditor/Private/Widgets/SNiagaraGraphNode.cpp
--- a/Engine/Plugins/FX/Niagara/Source/NiagaraEditor/Private/Widgets/SNiagaraGraphNode.cpp
+++ b/Engine/Plugins/FX/Niagara/Source/NiagaraEditor/Private/Widgets/SNiagaraGraphNode.cpp
@@ -243,8 +243,6 @@ void SNiagaraGraphNode::UpdateGraphNodeCompact()
 	//            |       |      | T (>) |
 	//            |_______|______|_______|
 	//
-	this->ContentScale.Bind( this, &SGraphNode::GetContentScale );
-	
 	TSharedRef<SVerticalBox> InnerVerticalBox =
 		SNew(SVerticalBox)
 		+SVerticalBox::Slot()
diff --git a/Engine/Plugins/PCG/Source/PCGEditor/Private/Nodes/SPCGEditorGraphNodeCompact.cpp b/Engine/Plugins/PCG/Source/PCGEditor/Private/Nodes/SPCGEditorGraphNodeCompact.cpp
--- a/Engine/Plugins/PCG/Source/PCGEditor/Private/Nodes/SPCGEditorGraphNodeCompact.cpp
+++ b/Engine/Plugins/PCG/Source/PCGEditor/Private/Nodes/SPCGEditorGraphNodeCompact.cpp
@@ -165,8 +165,6 @@ void SPCGEditorGraphNodeCompact::UpdateGraphNode()
 
 	check(ContentOverlay);
 
-	this->ContentScale.Bind(this, &SGraphNode::GetContentScale);
-
 	// First add the body overlay
 	const TSharedRef<SOverlay> NodeContentOverlay =
 		SNew(SOverlay)
diff --git a/Engine/Plugins/Runtime/RigVM/Source/RigVMEditor/Private/Widgets/SRigVMGraphNode.cpp b/Engine/Plugins/Runtime/RigVM/Source/RigVMEditor/Private/Widgets/SRigVMGraphNode.cpp
--- a/Engine/Plugins/Runtime/RigVM/Source/RigVMEditor/Private/Widgets/SRigVMGraphNode.cpp
+++ b/Engine/Plugins/Runtime/RigVM/Source/RigVMEditor/Private/Widgets/SRigVMGraphNode.cpp
@@ -527,8 +527,6 @@ void SRigVMGraphNode::UpdateCompactNode()
 		CreateNodeContentArea()
 	];
 
-	this->ContentScale.Bind( this, &SGraphNode::GetContentScale );
-
 	InnerVerticalBox->AddSlot()
 		.AutoHeight()
 		.Padding( FMargin(5.0f, 1.0f) )
diff --git a/Engine/Plugins/TextureGraph/Source/TextureGraphEditor/Private/EdGraph/STG_EditorGraphNode.cpp b/Engine/Plugins/TextureGraph/Source/TextureGraphEditor/Private/EdGraph/STG_EditorGraphNode.cpp
--- a/Engine/Plugins/TextureGraph/Source/TextureGraphEditor/Private/EdGraph/STG_EditorGraphNode.cpp
+++ b/Engine/Plugins/TextureGraph/Source/TextureGraphEditor/Private/EdGraph/STG_EditorGraphNode.cpp
@@ -212,7 +212,6 @@ void STG_EditorGraphNode::UpdateGraphNode()
 	PopulateMetaTag(&TagMeta);
 
 	TSharedPtr<SVerticalBox> InnerVerticalBox;
-	this->ContentScale.Bind(this, &SGraphNode::GetContentScale);
 
 	InnerVerticalBox = SNew(SVerticalBox)
 		+ SVerticalBox::Slot()
diff --git a/Engine/Plugins/VirtualProduction/SceneState/Source/SceneStateMachineEditor/Private/Widgets/SSceneStateMachineNode.cpp b/Engine/Plugins/VirtualProduction/SceneState/Source/SceneStateMachineEditor/Private/Widgets/SSceneStateMachineNode.cpp
--- a/Engine/Plugins/VirtualProduction/SceneState/Source/SceneStateMachineEditor/Private/Widgets/SSceneStateMachineNode.cpp
+++ b/Engine/Plugins/VirtualProduction/SceneState/Source/SceneStateMachineEditor/Private/Widgets/SSceneStateMachineNode.cpp
@@ -41,8 +41,6 @@ void SStateMachineNode::UpdateGraphNode()
 
 	RightNodeBox.Reset();
 	LeftNodeBox.Reset();
-
-	ContentScale.Bind(this, &SGraphNode::GetContentScale);
 }
 
 void SStateMachineNode::CreateStandardPinWidget(UEdGraphPin* InPin)
diff --git a/Engine/Source/Editor/AnimationBlueprintEditor/Private/AnimationStateNodes/SGraphNodeAnimState.cpp b/Engine/Source/Editor/AnimationBlueprintEditor/Private/AnimationStateNodes/SGraphNodeAnimState.cpp
--- a/Engine/Source/Editor/AnimationBlueprintEditor/Private/AnimationStateNodes/SGraphNodeAnimState.cpp
+++ b/Engine/Source/Editor/AnimationBlueprintEditor/Private/AnimationStateNodes/SGraphNodeAnimState.cpp
@@ -233,7 +233,6 @@ void SGraphNodeAnimState::UpdateGraphNode()
 	// Set ourselves as SelfHitTestInvisible so we dont end up with a 1-unit border where the node is selected rather than the pin
 	SetVisibility(EVisibility::SelfHitTestInvisible);
 	
-	this->ContentScale.Bind( this, &SGraphNode::GetContentScale );
 	this->GetOrAddSlot( ENodeZone::Center )
 		.HAlign(HAlign_Center)
 		.VAlign(VAlign_Center)
diff --git a/Engine/Source/Editor/AnimationBlueprintEditor/Private/AnimationStateNodes/SGraphNodeAnimStateEntry.cpp b/Engine/Source/Editor/AnimationBlueprintEditor/Private/AnimationStateNodes/SGraphNodeAnimStateEntry.cpp
--- a/Engine/Source/Editor/AnimationBlueprintEditor/Private/AnimationStateNodes/SGraphNodeAnimStateEntry.cpp
+++ b/Engine/Source/Editor/AnimationBlueprintEditor/Private/AnimationStateNodes/SGraphNodeAnimStateEntry.cpp
@@ -68,7 +68,6 @@ void SGraphNodeAnimStateEntry::UpdateGraphNode()
 
 	FLinearColor TitleShadowColor(0.6f, 0.6f, 0.6f);
 
-	this->ContentScale.Bind( this, &SGraphNode::GetContentScale );
 	this->GetOrAddSlot( ENodeZone::Center )
 		.HAlign(HAlign_Center)
 		.VAlign(VAlign_Center)
diff --git a/Engine/Source/Editor/AnimationBlueprintEditor/Private/AnimationStateNodes/SGraphNodeAnimTransition.cpp b/Engine/Source/Editor/AnimationBlueprintEditor/Private/AnimationStateNodes/SGraphNodeAnimTransition.cpp
--- a/Engine/Source/Editor/AnimationBlueprintEditor/Private/AnimationStateNodes/SGraphNodeAnimTransition.cpp
+++ b/Engine/Source/Editor/AnimationBlueprintEditor/Private/AnimationStateNodes/SGraphNodeAnimTransition.cpp
@@ -251,7 +251,6 @@ void SGraphNodeAnimTransition::UpdateGraphNode()
 	}));
 	DirectionImage->SetRenderTransformPivot(FVector2D(0.5f, 0.5f));
 
-	this->ContentScale.Bind( this, &SGraphNode::GetContentScale );
 	this->GetOrAddSlot( ENodeZone::Center )
 		.HAlign(HAlign_Center)
 		.VAlign(VAlign_Center)
diff --git a/Engine/Source/Editor/BehaviorTreeEditor/Private/SGraphNode_BehaviorTree.cpp b/Engine/Source/Editor/BehaviorTreeEditor/Private/SGraphNode_BehaviorTree.cpp
--- a/Engine/Source/Editor/BehaviorTreeEditor/Private/SGraphNode_BehaviorTree.cpp
+++ b/Engine/Source/Editor/BehaviorTreeEditor/Private/SGraphNode_BehaviorTree.cpp
@@ -354,7 +354,6 @@ void SGraphNode_BehaviorTree::UpdateGraphNode()
 		.OnHoverStateChanged(this, &SGraphNode_BehaviorTree::OnIndexHoverStateChanged)
 		.OnGetIndexColor(this, &SGraphNode_BehaviorTree::GetIndexColor);
 
-	this->ContentScale.Bind( this, &SGraphNode::GetContentScale );
 	this->GetOrAddSlot( ENodeZone::Center )
 		.HAlign(HAlign_Fill)
 		.VAlign(VAlign_Center)
diff --git a/Engine/Source/Editor/GraphEditor/Private/KismetNodes/SGraphNodeK2Base.cpp b/Engine/Source/Editor/GraphEditor/Private/KismetNodes/SGraphNodeK2Base.cpp
--- a/Engine/Source/Editor/GraphEditor/Private/KismetNodes/SGraphNodeK2Base.cpp
+++ b/Engine/Source/Editor/GraphEditor/Private/KismetNodes/SGraphNodeK2Base.cpp
@@ -211,8 +211,6 @@ void SGraphNodeK2Base::UpdateCompactNode()
 	//            |       |      | T (>) |
 	//            |_______|______|_______|
 	//
-	this->ContentScale.Bind( this, &SGraphNode::GetContentScale );
-	
 	TSharedRef<SVerticalBox> InnerVerticalBox =
 		SNew(SVerticalBox)
 		+SVerticalBox::Slot()
diff --git a/Engine/Source/Editor/GraphEditor/Private/KismetNodes/SGraphNodeK2Composite.cpp b/Engine/Source/Editor/GraphEditor/Private/KismetNodes/SGraphNodeK2Composite.cpp
--- a/Engine/Source/Editor/GraphEditor/Private/KismetNodes/SGraphNodeK2Composite.cpp
+++ b/Engine/Source/Editor/GraphEditor/Private/KismetNodes/SGraphNodeK2Composite.cpp
@@ -80,7 +80,6 @@ void SGraphNodeK2Composite::UpdateGraphNode()
 	//            |       |      | T (>) |
 	//            |_______|______|_______|
 	//
-	this->ContentScale.Bind( this, &SGraphNode::GetContentScale );
 	this->GetOrAddSlot( ENodeZone::Center )
 		.HAlign(HAlign_Fill)
 		.VAlign(VAlign_Center)
diff --git a/Engine/Source/Editor/GraphEditor/Private/KismetNodes/SGraphNodeK2Terminator.cpp b/Engine/Source/Editor/GraphEditor/Private/KismetNodes/SGraphNodeK2Terminator.cpp
--- a/Engine/Source/Editor/GraphEditor/Private/KismetNodes/SGraphNodeK2Terminator.cpp
+++ b/Engine/Source/Editor/GraphEditor/Private/KismetNodes/SGraphNodeK2Terminator.cpp
@@ -77,7 +77,6 @@ void SGraphNodeK2Terminator::UpdateGraphNode()
 	//            |       |      | T (>) |
 	//            |_______|______|_______|
 	//
-	this->ContentScale.Bind( this, &SGraphNode::GetContentScale );
 	this->GetOrAddSlot( ENodeZone::Center )
 	.HAlign(HAlign_Fill)
 	.VAlign(VAlign_Center)
diff --git a/Engine/Source/Editor/GraphEditor/Private/KismetNodes/SGraphNodeK2Var.cpp b/Engine/Source/Editor/GraphEditor/Private/KismetNodes/SGraphNodeK2Var.cpp
--- a/Engine/Source/Editor/GraphEditor/Private/KismetNodes/SGraphNodeK2Var.cpp
+++ b/Engine/Source/Editor/GraphEditor/Private/KismetNodes/SGraphNodeK2Var.cpp
@@ -242,7 +242,6 @@ void SGraphNodeK2Var::UpdateGraphNode()
 	//            |       |  T (>) |
 	//            |_______|________|
 	//
-	this->ContentScale.Bind( this, &SGraphNode::GetContentScale );
 	this->GetOrAddSlot( ENodeZone::Center )
 	.HAlign(HAlign_Fill)
 	.VAlign(VAlign_Center)
diff --git a/Engine/Source/Editor/GraphEditor/Private/SGraphNode.cpp b/Engine/Source/Editor/GraphEditor/Private/SGraphNode.cpp
--- a/Engine/Source/Editor/GraphEditor/Private/SGraphNode.cpp
+++ b/Engine/Source/Editor/GraphEditor/Private/SGraphNode.cpp
@@ -961,7 +961,6 @@ void SGraphNode::UpdateGraphNode()
 	PopulateMetaTag(&TagMeta);
 	
 	TSharedPtr<SVerticalBox> InnerVerticalBox;
-	this->ContentScale.Bind( this, &SGraphNode::GetContentScale );
 
 
 	InnerVerticalBox = SNew(SVerticalBox)
diff --git a/Engine/Source/Editor/GraphEditor/Private/SGraphNodeComment.cpp b/Engine/Source/Editor/GraphEditor/Private/SGraphNodeComment.cpp
--- a/Engine/Source/Editor/GraphEditor/Private/SGraphNodeComment.cpp
+++ b/Engine/Source/Editor/GraphEditor/Private/SGraphNodeComment.cpp
@@ -169,7 +169,6 @@ void SGraphNodeComment::UpdateGraphNode()
 	CommentStyle.EditableTextBoxStyle.TextStyle.Font.Size = static_cast<float>(CachedFontSize);
 	CommentStyle.TextStyle.Font.Size = static_cast<float>(CachedFontSize);
 
-	this->ContentScale.Bind( this, &SGraphNode::GetContentScale );
 	this->GetOrAddSlot( ENodeZone::Center )
 		.HAlign(HAlign_Fill)
 		.VAlign(VAlign_Fill)
diff --git a/Engine/Source/Editor/GraphEditor/Private/SGraphNodeKnot.cpp b/Engine/Source/Editor/GraphEditor/Private/SGraphNodeKnot.cpp
--- a/Engine/Source/Editor/GraphEditor/Private/SGraphNodeKnot.cpp
+++ b/Engine/Source/Editor/GraphEditor/Private/SGraphNodeKnot.cpp
@@ -348,8 +348,6 @@ void SGraphNodeKnot::UpdateGraphNode()
 		.IsReadOnly(this, &SGraphNodeKnot::IsNameReadOnly)
 		.IsSelected(this, &SGraphNodeKnot::IsSelectedExclusively);
 
-	this->ContentScale.Bind( this, &SGraphNode::GetContentScale );
-
 	SetupErrorReporting();
 
 	this->GetOrAddSlot( ENodeZone::Center )
diff --git a/Engine/Source/Editor/GraphEditor/Public/SNodePanel.h b/Engine/Source/Editor/GraphEditor/Public/SNodePanel.h
--- a/Engine/Source/Editor/GraphEditor/Public/SNodePanel.h
+++ b/Engine/Source/Editor/GraphEditor/Public/SNodePanel.h
@@ -753,9 +753,7 @@ public:
 
 	protected:
 		SNode()
-		: BorderImage( FCoreStyle::Get().GetBrush( "NoBorder" ) )
-		, BorderBackgroundColor( FAppStyle::GetColor("Graph.ForegroundColor"))
-		, DesiredSizeScale(FVector2f::UnitVector)
+		: DesiredSizeScale(FVector2f::UnitVector)
 		, Children(this)
 		{
 			bHasRelativeLayoutScale = true;
@@ -763,26 +761,38 @@ public:
 
 	protected:
 
-		// SBorder Begin
-		TAttribute<const FSlateBrush*> BorderImage;
-		TAttribute<FSlateColor> BorderBackgroundColor;
 		TAttribute<FDeprecateSlateVector2D> DesiredSizeScale;
+
+		// SBorder Begin
+		UE_DEPRECATED(5.8, "The attribute was unused and can be removed.")
+		TSlateDeprecatedTAttribute<const FSlateBrush*> BorderImage;
+		UE_DEPRECATED(5.8, "The attribute was unused and can be removed.")
+		TSlateDeprecatedTAttribute<FSlateColor> BorderBackgroundColor;
+
 		/** Whether or not to show the disabled effect when this border is disabled */
-		TAttribute<bool> ShowDisabledEffect;
+		UE_DEPRECATED(5.8, "The attribute was unused and can be removed.")
+		TSlateDeprecatedTAttribute<bool> ShowDisabledEffect;
 		/** Mouse event handlers */
+		UE_DEPRECATED(5.8, "The EventHandler was unused and can be removed.")
 		FPointerEventHandler MouseButtonDownHandler;
+		UE_DEPRECATED(5.8, "The EventHandler was unused and can be removed.")
 		FPointerEventHandler MouseButtonUpHandler;
+		UE_DEPRECATED(5.8, "The EventHandler was unused and can be removed.")
 		FPointerEventHandler MouseMoveHandler;
+		UE_DEPRECATED(5.8, "The EventHandler was unused and can be removed.")
 		FPointerEventHandler MouseDoubleClickHandler;
 		// SBorder End
 
 		// SPanel Begin
 		/** The layout scale to apply to this widget's contents; useful for animation. */
-		TAttribute<FDeprecateSlateVector2D> ContentScale;
+		UE_DEPRECATED(5.8, "The attribute was unused and can be removed.")
+		TSlateDeprecatedTAttribute<FDeprecateSlateVector2D> ContentScale;
 		/** The color and opacity to apply to this widget and all its descendants. */
-		TAttribute<FLinearColor> ColorAndOpacity;
+		UE_DEPRECATED(5.8, "The attribute was unused and can be removed.")
+		TSlateDeprecatedTAttribute<FLinearColor> ColorAndOpacity;
 		/** Optional foreground color that will be inherited by all of this widget's contents */
-		TAttribute<FSlateColor> ForegroundColor;
+		UE_DEPRECATED(5.8, "The attribute was unused and can be removed.")
+		TSlateDeprecatedTAttribute<FSlateColor> ForegroundColor;
 		// SPanel End
 
 	private:
diff --git a/Engine/Source/Editor/MaterialEditor/Private/MaterialNodes/SGraphNodeMaterialComposite.cpp b/Engine/Source/Editor/MaterialEditor/Private/MaterialNodes/SGraphNodeMaterialComposite.cpp
--- a/Engine/Source/Editor/MaterialEditor/Private/MaterialNodes/SGraphNodeMaterialComposite.cpp
+++ b/Engine/Source/Editor/MaterialEditor/Private/MaterialNodes/SGraphNodeMaterialComposite.cpp
@@ -88,7 +88,6 @@ void SGraphNodeMaterialComposite::UpdateGraphNode()
 	//            |       |      | T (>) |
 	//            |_______|______|_______|
 	//
-	this->ContentScale.Bind( this, &SGraphNode::GetContentScale );
 	this->GetOrAddSlot( ENodeZone::Center )
 		.HAlign(HAlign_Center)
 		.VAlign(VAlign_Center)
diff --git a/Engine/Source/Editor/PhysicsAssetEditor/Private/PhysicsAssetGraph/SPhysicsAssetGraphNode.cpp b/Engine/Source/Editor/PhysicsAssetEditor/Private/PhysicsAssetGraph/SPhysicsAssetGraphNode.cpp
--- a/Engine/Source/Editor/PhysicsAssetEditor/Private/PhysicsAssetGraph/SPhysicsAssetGraphNode.cpp
+++ b/Engine/Source/Editor/PhysicsAssetEditor/Private/PhysicsAssetGraph/SPhysicsAssetGraphNode.cpp
@@ -102,8 +102,6 @@ void SPhysicsAssetGraphNode::UpdateGraphNode()
 	RightNodeBox.Reset();
 	LeftNodeBox.Reset();
 
-	this->ContentScale.Bind( this, &SGraphNode::GetContentScale );
-
 	this->GetOrAddSlot( ENodeZone::Center )
 		.HAlign(HAlign_Center)
 		.VAlign(VAlign_Center)
