[Core UI] Fixed a bug preventing the Y value box in SVector2DSlider from working properly

#Synced-CL 49728766

diff --git a/Engine/Source/Editor/GraphEditor/Public/KismetPins/SVector2DSlider.h b/Engine/Source/Editor/GraphEditor/Public/KismetPins/SVector2DSlider.h
--- a/Engine/Source/Editor/GraphEditor/Public/KismetPins/SVector2DSlider.h
+++ b/Engine/Source/Editor/GraphEditor/Public/KismetPins/SVector2DSlider.h
@@ -107,6 +107,12 @@ public:
 						.Text( LOCTEXT("VectorNodeYAxisValueLabel", "Y") )
 						.ColorAndOpacity( LabelClr )
 					]
+					.MinValue(NumericPropertyParams.MinValue)
+					.MaxValue(NumericPropertyParams.MaxValue)
+					.MinSliderValue(NumericPropertyParams.MinSliderValue)
+					.MaxSliderValue(NumericPropertyParams.MaxSliderValue)
+					.SliderExponent(NumericPropertyParams.SliderExponent)
+					.Delta(NumericPropertyParams.Delta)
 					.Value( this, &SVector2DSlider::GetTypeInValue_Y )
 					.LinearDeltaSensitivity(NumericPropertyParams.GetLinearDeltaSensitivityAttribute())
 					.AllowWheel(bAllowSpin)
@@ -216,13 +222,13 @@ private:
 	// Get value for X text box
 	TOptional<NumericType> GetTypeInValue_X() const
 	{
-		return GetValueType(VisibleText_X.Get());
+		return bIsUsingSlider ? SliderValue_X : GetValueType(VisibleText_X.Get());
 	}
 
 	// Get value for Y text box
 	TOptional<NumericType> GetTypeInValue_Y() const
 	{
-		return GetValueType(VisibleText_Y.Get());
+		return bIsUsingSlider ? SliderValue_Y : GetValueType(VisibleText_Y.Get());
 	}
 
 	TAttribute<FString> VisibleText_X;
