Slate Dynamic Invalidation - Graph Editor Wrapper

Adding an argument to `SGraphEditor` to allow graph editors to wrap the `SGraphPanel` inside a dynamic invalidation panel.

#Synced-CL 50145007

diff --git a/Engine/Source/Editor/GraphEditor/Private/GraphEditorModule.cpp b/Engine/Source/Editor/GraphEditor/Private/GraphEditorModule.cpp
--- a/Engine/Source/Editor/GraphEditor/Private/GraphEditorModule.cpp
+++ b/Engine/Source/Editor/GraphEditor/Private/GraphEditorModule.cpp
@@ -74,7 +74,8 @@ TSharedRef<SGraphEditor> FGraphEditorModule::PRIVATE_MakeGraphEditor(
 	FSimpleDelegate InOnNavigateHistoryForward,
 	TAttribute<bool> ShowGraphStateOverlay,
 	TAttribute<bool> AllowConnectionSlicing,
-	FOnGraphEditorTickDelegate InOnTickDelegate)
+	FOnGraphEditorTickDelegate InOnTickDelegate,
+	bool bGraphSupportsInvalidation)
 {
 	return
 		SNew(SGraphEditorImpl)
@@ -92,7 +93,8 @@ TSharedRef<SGraphEditor> FGraphEditorModule::PRIVATE_MakeGraphEditor(
 		.OnNavigateHistoryForward(InOnNavigateHistoryForward)
 		.OnTick(InOnTickDelegate)
 		.ShowGraphStateOverlay(ShowGraphStateOverlay)
-		.AllowConnectionSlicing(AllowConnectionSlicing);
+		.AllowConnectionSlicing(AllowConnectionSlicing)
+		.GraphSupportsInvalidation(bGraphSupportsInvalidation);
 }
 
 /////////////////////////////////////////////////////
diff --git a/Engine/Source/Editor/GraphEditor/Private/SGraphEditorImpl.cpp b/Engine/Source/Editor/GraphEditor/Private/SGraphEditorImpl.cpp
--- a/Engine/Source/Editor/GraphEditor/Private/SGraphEditorImpl.cpp
+++ b/Engine/Source/Editor/GraphEditor/Private/SGraphEditorImpl.cpp
@@ -66,6 +66,7 @@
 #include "HAL/IConsoleManager.h"
 #include "RenderUtils.h"
 #include "Materials/Material.h"
+#include "Widgets/SInvalidationPanel.h"
 
 class FActiveTimerHandle;
 struct FGeometry;
@@ -730,6 +731,42 @@ void SGraphEditorImpl::Construct( const FArguments& InArgs )
 
 	TSharedPtr<SOverlay> OverlayWidget;
 
+	GraphPanel = SNew(SGraphPanel)
+		.GraphObj( EdGraphObj )
+		.DiffResults( InArgs._DiffResults)
+		.FocusedDiffResult( InArgs._FocusedDiffResult)
+		.OnGetContextMenuFor( this, &SGraphEditorImpl::GraphEd_OnGetContextMenuFor )
+		.OnSelectionChanged( InArgs._GraphEvents.OnSelectionChanged )
+		.OnNodeDoubleClicked( InArgs._GraphEvents.OnNodeDoubleClicked )
+		.IsEditable( this, &SGraphEditorImpl::IsGraphEditable )
+		.DisplayAsReadOnly( this, &SGraphEditorImpl::DisplayGraphAsReadOnly )
+		.OnDropActors(InArgs._GraphEvents.OnDropActors)
+		.OnDropStreamingLevels( InArgs._GraphEvents.OnDropStreamingLevels )
+		.OnVerifyTextCommit( InArgs._GraphEvents.OnVerifyTextCommit )
+		.OnTextCommitted( InArgs._GraphEvents.OnTextCommitted )
+		.OnSpawnNodeByShortcutAtLocation( InArgs._GraphEvents.OnSpawnNodeByShortcutAtLocation)
+		//.OnUpdateGraphPanel( this, &SGraphEditorImpl::GraphEd_OnPanelUpdated )
+		.OnDisallowedPinConnection( InArgs._GraphEvents.OnDisallowedPinConnection )
+		.ShowGraphStateOverlay(InArgs._ShowGraphStateOverlay)
+		.AllowConnectionSlicing(InArgs._AllowConnectionSlicing)
+		.OnDoubleClicked(InArgs._GraphEvents.OnDoubleClicked)
+		.OnMouseButtonDown(InArgs._GraphEvents.OnMouseButtonDown)
+		.OnNodeSingleClicked(InArgs._GraphEvents.OnNodeSingleClicked);
+
+	TSharedPtr<SWidget> GraphPanelWidget;
+	if (InArgs._GraphSupportsInvalidation)
+	{
+		GraphPanelWidget = SNew(SInvalidationPanel)
+		.UseDynamicInvalidation(true)
+		[
+			GraphPanel.ToSharedRef()
+		];
+	}
+	else
+	{
+		GraphPanelWidget = GraphPanel;
+	}
+
 	this->ChildSlot
 	[
 		SAssignNew(OverlayWidget, SOverlay)
@@ -738,27 +775,7 @@ void SGraphEditorImpl::Construct( const FArguments& InArgs )
 		+SOverlay::Slot()
 		.Expose(GraphPanelSlot)
 		[
-			SAssignNew(GraphPanel, SGraphPanel)
-			.GraphObj( EdGraphObj )
-			.DiffResults( InArgs._DiffResults)
-			.FocusedDiffResult( InArgs._FocusedDiffResult)
-			.OnGetContextMenuFor( this, &SGraphEditorImpl::GraphEd_OnGetContextMenuFor )
-			.OnSelectionChanged( InArgs._GraphEvents.OnSelectionChanged )
-			.OnNodeDoubleClicked( InArgs._GraphEvents.OnNodeDoubleClicked )
-			.IsEditable( this, &SGraphEditorImpl::IsGraphEditable )
-			.DisplayAsReadOnly( this, &SGraphEditorImpl::DisplayGraphAsReadOnly )
-			.OnDropActors(InArgs._GraphEvents.OnDropActors)
-			.OnDropStreamingLevels( InArgs._GraphEvents.OnDropStreamingLevels )
-			.OnVerifyTextCommit( InArgs._GraphEvents.OnVerifyTextCommit )
-			.OnTextCommitted( InArgs._GraphEvents.OnTextCommitted )
-			.OnSpawnNodeByShortcutAtLocation( InArgs._GraphEvents.OnSpawnNodeByShortcutAtLocation)
-			//.OnUpdateGraphPanel( this, &SGraphEditorImpl::GraphEd_OnPanelUpdated )
-			.OnDisallowedPinConnection( InArgs._GraphEvents.OnDisallowedPinConnection )
-			.ShowGraphStateOverlay(InArgs._ShowGraphStateOverlay)
-			.AllowConnectionSlicing(InArgs._AllowConnectionSlicing)
-			.OnDoubleClicked(InArgs._GraphEvents.OnDoubleClicked)
-			.OnMouseButtonDown(InArgs._GraphEvents.OnMouseButtonDown)
-		    .OnNodeSingleClicked(InArgs._GraphEvents.OnNodeSingleClicked)
+			GraphPanelWidget.ToSharedRef()
 		]
 
 		// Indicator of current zoom level
@@ -800,7 +817,7 @@ void SGraphEditorImpl::Construct( const FArguments& InArgs )
 					.TextStyle(FAppStyle::Get(), "Graph.InstructionText")
 					.Text(this, &SGraphEditorImpl::GetInstructionText)
 				]
-			]			
+			]
 		]
 
 		// Bottom-left corner text for Substrate
diff --git a/Engine/Source/Editor/GraphEditor/Private/SGraphEditorImpl.h b/Engine/Source/Editor/GraphEditor/Private/SGraphEditorImpl.h
--- a/Engine/Source/Editor/GraphEditor/Private/SGraphEditorImpl.h
+++ b/Engine/Source/Editor/GraphEditor/Private/SGraphEditorImpl.h
@@ -209,6 +209,7 @@ public:
 		, _IsEditable(true)
 		, _DisplayAsReadOnly(false)
 		, _AllowConnectionSlicing(true)
+		, _GraphSupportsInvalidation(false)
 		{}
 
 
@@ -230,6 +231,9 @@ public:
 		SLATE_ATTRIBUTE(bool, ShowGraphStateOverlay)
 		/** Allow graph connections to be sliced and removed */
 		SLATE_ATTRIBUTE(bool, AllowConnectionSlicing)
+
+		/** If set to true, the GraphPanel will be surrounded with a Dynamic Invalidation Panel */
+		SLATE_ARGUMENT(bool, GraphSupportsInvalidation)
 	SLATE_END_ARGS()
 
 	void Construct( const FArguments& InArgs );
diff --git a/Engine/Source/Editor/GraphEditor/Public/GraphEditorModule.h b/Engine/Source/Editor/GraphEditor/Public/GraphEditorModule.h
--- a/Engine/Source/Editor/GraphEditor/Public/GraphEditorModule.h
+++ b/Engine/Source/Editor/GraphEditor/Public/GraphEditorModule.h
@@ -59,7 +59,8 @@ private:
 		FSimpleDelegate InOnNavigateHistoryForward,
 		TAttribute<bool> ShowGraphStateOverlay,
 		TAttribute<bool> AllowConnectionSlicing,
-		FOnGraphEditorTickDelegate InOnTickDelegate);
+		FOnGraphEditorTickDelegate InOnTickDelegate,
+		bool bGraphSupportsInvalidation);
 
 private:
 	/** All extender delegates for the graph menus */
diff --git a/Engine/Source/Editor/MaterialEditor/Private/MaterialEditor.cpp b/Engine/Source/Editor/MaterialEditor/Private/MaterialEditor.cpp
--- a/Engine/Source/Editor/MaterialEditor/Private/MaterialEditor.cpp
+++ b/Engine/Source/Editor/MaterialEditor/Private/MaterialEditor.cpp
@@ -7771,6 +7771,7 @@ TSharedRef<SGraphEditor> FMaterialEditor::CreateGraphEditorWidget(TSharedRef<cla
 		.MaterialInfoList(&MaterialInfoList);
 
 	return SNew(SGraphEditor)
+		.GraphSupportsInvalidation(false)
 		.AdditionalCommands(GraphEditorCommands)
 		.IsEditable(true)
 		.TitleBar(TitleBarWidget)
diff --git a/Engine/Source/Editor/UnrealEd/Private/GraphEditor.cpp b/Engine/Source/Editor/UnrealEd/Private/GraphEditor.cpp
--- a/Engine/Source/Editor/UnrealEd/Private/GraphEditor.cpp
+++ b/Engine/Source/Editor/UnrealEd/Private/GraphEditor.cpp
@@ -46,7 +46,8 @@ void SGraphEditor::ConstructImplementation( const FArguments& InArgs )
 		InArgs._OnNavigateHistoryForward,
 		InArgs._ShowGraphStateOverlay,
 		InArgs._AllowConnectionSlicing,
-		InArgs._OnTick
+		InArgs._OnTick,
+		InArgs._GraphSupportsInvalidation
 		);
 
 	Implementation->AssetEditorToolkit = InArgs._AssetEditorToolkit;
diff --git a/Engine/Source/Editor/UnrealEd/Public/GraphEditor.h b/Engine/Source/Editor/UnrealEd/Public/GraphEditor.h
--- a/Engine/Source/Editor/UnrealEd/Public/GraphEditor.h
+++ b/Engine/Source/Editor/UnrealEd/Public/GraphEditor.h
@@ -193,6 +193,7 @@ PRAGMA_ENABLE_DEPRECATION_WARNINGS
 		, _AutoExpandActionMenu(false)
 		, _ShowGraphStateOverlay(true)
 		, _AllowConnectionSlicing(true)
+		, _GraphSupportsInvalidation(false)
 		{}
 
 		SLATE_ARGUMENT( TSharedPtr<FUICommandList>, AdditionalCommands )
@@ -224,6 +225,9 @@ PRAGMA_ENABLE_DEPRECATION_WARNINGS
 		/** Show overlay elements for the graph state such as the PIE and read-only borders and text */
 		SLATE_ATTRIBUTE(bool, ShowGraphStateOverlay)
 		SLATE_ATTRIBUTE(bool, AllowConnectionSlicing)
+
+		/** If set to true, the GraphPanel will be surrounded with a Dynamic Invalidation Panel */
+		SLATE_ARGUMENT(bool, GraphSupportsInvalidation)
 	SLATE_END_ARGS()
 
 	/**
