Slate Dynamic Invalidation - Graph Editor - Adding TypeTraits to trivial GraphPins

#Synced-CL 50145007

diff --git a/Engine/Source/Editor/GraphEditor/Public/KismetPins/SGraphPinBool.h b/Engine/Source/Editor/GraphEditor/Public/KismetPins/SGraphPinBool.h
--- a/Engine/Source/Editor/GraphEditor/Public/KismetPins/SGraphPinBool.h
+++ b/Engine/Source/Editor/GraphEditor/Public/KismetPins/SGraphPinBool.h
@@ -14,6 +14,12 @@
 class SWidget;
 class UEdGraphPin;
 
+template <>
+struct TWidgetTypeTraits<class SGraphPinBool>
+{
+	static constexpr bool SupportsInvalidation() { return true; }
+};
+
 class SGraphPinBool : public SGraphPin
 {
 public:
diff --git a/Engine/Source/Editor/GraphEditor/Public/KismetPins/SGraphPinColor.h b/Engine/Source/Editor/GraphEditor/Public/KismetPins/SGraphPinColor.h
--- a/Engine/Source/Editor/GraphEditor/Public/KismetPins/SGraphPinColor.h
+++ b/Engine/Source/Editor/GraphEditor/Public/KismetPins/SGraphPinColor.h
@@ -18,6 +18,12 @@ class UEdGraphPin;
 struct FGeometry;
 struct FPointerEvent;
 
+template <>
+struct TWidgetTypeTraits<class SGraphPinColor>
+{
+	static constexpr bool SupportsInvalidation() { return true; }
+};
+
 class SGraphPinColor : public SGraphPin
 {
 public:
diff --git a/Engine/Source/Editor/GraphEditor/Public/KismetPins/SGraphPinEnum.h b/Engine/Source/Editor/GraphEditor/Public/KismetPins/SGraphPinEnum.h
--- a/Engine/Source/Editor/GraphEditor/Public/KismetPins/SGraphPinEnum.h
+++ b/Engine/Source/Editor/GraphEditor/Public/KismetPins/SGraphPinEnum.h
@@ -16,6 +16,12 @@
 class SWidget;
 class UEdGraphPin;
 
+template <>
+struct TWidgetTypeTraits<class SGraphPinEnum>
+{
+	static constexpr bool SupportsInvalidation() { return true; }
+};
+
 class SGraphPinEnum : public SGraphPin
 {
 public:
diff --git a/Engine/Source/Editor/GraphEditor/Public/KismetPins/SGraphPinExec.h b/Engine/Source/Editor/GraphEditor/Public/KismetPins/SGraphPinExec.h
--- a/Engine/Source/Editor/GraphEditor/Public/KismetPins/SGraphPinExec.h
+++ b/Engine/Source/Editor/GraphEditor/Public/KismetPins/SGraphPinExec.h
@@ -14,6 +14,13 @@ class SWidget;
 class UEdGraphPin;
 struct FSlateBrush;
 
+
+template <>
+struct TWidgetTypeTraits<class SGraphPinExec>
+{
+	static constexpr bool SupportsInvalidation() { return true; }
+};
+
 class SGraphPinExec : public SGraphPin
 {
 public:
diff --git a/Engine/Source/Editor/GraphEditor/Public/KismetPins/SGraphPinInteger.h b/Engine/Source/Editor/GraphEditor/Public/KismetPins/SGraphPinInteger.h
--- a/Engine/Source/Editor/GraphEditor/Public/KismetPins/SGraphPinInteger.h
+++ b/Engine/Source/Editor/GraphEditor/Public/KismetPins/SGraphPinInteger.h
@@ -18,6 +18,12 @@
 class SWidget;
 class UEdGraphPin;
 
+template <>
+struct TWidgetTypeTraits<class SGraphPinInteger>
+{
+	static constexpr bool SupportsInvalidation() { return true; }
+};
+
 class SGraphPinInteger : public SGraphPinNum<int32>
 {
 public:
diff --git a/Engine/Source/Editor/GraphEditor/Public/KismetPins/SGraphPinNum.h b/Engine/Source/Editor/GraphEditor/Public/KismetPins/SGraphPinNum.h
--- a/Engine/Source/Editor/GraphEditor/Public/KismetPins/SGraphPinNum.h
+++ b/Engine/Source/Editor/GraphEditor/Public/KismetPins/SGraphPinNum.h
@@ -11,6 +11,15 @@
 #include "ScopedTransaction.h"
 #include "SGraphPin.h"
 
+template <typename NumericType>
+class SGraphPinNum;
+
+template <typename NumericType>
+struct TWidgetTypeTraits<SGraphPinNum<NumericType>>
+{
+	static constexpr bool SupportsInvalidation() { return true; }
+};
+
 template <typename NumericType>
 class SGraphPinNum : public SGraphPin
 {
diff --git a/Engine/Source/Editor/GraphEditor/Public/KismetPins/SGraphPinString.h b/Engine/Source/Editor/GraphEditor/Public/KismetPins/SGraphPinString.h
--- a/Engine/Source/Editor/GraphEditor/Public/KismetPins/SGraphPinString.h
+++ b/Engine/Source/Editor/GraphEditor/Public/KismetPins/SGraphPinString.h
@@ -15,13 +15,19 @@
 class SWidget;
 class UEdGraphPin;
 
+template <>
+struct TWidgetTypeTraits<class SGraphPinString>
+{
+	static constexpr bool SupportsInvalidation() { return true; }
+};
+
 class SGraphPinString : public SGraphPin
 {
 public:
 	SLATE_BEGIN_ARGS(SGraphPinString) {}
 	SLATE_END_ARGS()
 
-		UE_API void Construct(const FArguments& InArgs, UEdGraphPin* InGraphPinObj);
+	UE_API void Construct(const FArguments& InArgs, UEdGraphPin* InGraphPinObj);
 
 protected:
 	//~ Begin SGraphPin Interface
diff --git a/Engine/Source/Editor/GraphEditor/Public/KismetPins/SGraphPinVector.h b/Engine/Source/Editor/GraphEditor/Public/KismetPins/SGraphPinVector.h
--- a/Engine/Source/Editor/GraphEditor/Public/KismetPins/SGraphPinVector.h
+++ b/Engine/Source/Editor/GraphEditor/Public/KismetPins/SGraphPinVector.h
@@ -15,6 +15,15 @@
 class SWidget;
 class UEdGraphPin;
 
+template <typename NumericType>
+class SGraphPinVector;
+
+template <typename NumericType>
+struct TWidgetTypeTraits<SGraphPinVector<NumericType>>
+{
+	static constexpr bool SupportsInvalidation() { return true; }
+};
+
 template <typename NumericType>
 class SGraphPinVector : public SGraphPin
 {
diff --git a/Engine/Source/Editor/GraphEditor/Public/KismetPins/SGraphPinVector2D.h b/Engine/Source/Editor/GraphEditor/Public/KismetPins/SGraphPinVector2D.h
--- a/Engine/Source/Editor/GraphEditor/Public/KismetPins/SGraphPinVector2D.h
+++ b/Engine/Source/Editor/GraphEditor/Public/KismetPins/SGraphPinVector2D.h
@@ -15,6 +15,17 @@
 class SWidget;
 class UEdGraphPin;
 
+
+template <typename NumericType>
+class SGraphPinVector2D;
+
+template <typename NumericType>
+struct TWidgetTypeTraits<SGraphPinVector2D<NumericType>>
+{
+	static constexpr bool SupportsInvalidation() { return true; }
+};
+
+
 template <typename NumericType>
 class SGraphPinVector2D : public SGraphPin
 {
diff --git a/Engine/Source/Editor/GraphEditor/Public/KismetPins/SGraphPinVector4.h b/Engine/Source/Editor/GraphEditor/Public/KismetPins/SGraphPinVector4.h
--- a/Engine/Source/Editor/GraphEditor/Public/KismetPins/SGraphPinVector4.h
+++ b/Engine/Source/Editor/GraphEditor/Public/KismetPins/SGraphPinVector4.h
@@ -16,6 +16,17 @@
 class SWidget;
 class UEdGraphPin;
 
+
+template <typename NumericType>
+class SGraphPinVector4;
+
+template <typename NumericType>
+struct TWidgetTypeTraits<SGraphPinVector4<NumericType>>
+{
+	static constexpr bool SupportsInvalidation() { return true; }
+};
+
+
 template <typename NumericType>
 class SGraphPinVector4 : public SGraphPin
 {
diff --git a/Engine/Source/Editor/GraphEditor/Public/SGraphPinComboBox.h b/Engine/Source/Editor/GraphEditor/Public/SGraphPinComboBox.h
--- a/Engine/Source/Editor/GraphEditor/Public/SGraphPinComboBox.h
+++ b/Engine/Source/Editor/GraphEditor/Public/SGraphPinComboBox.h
@@ -18,6 +18,12 @@ template <typename ItemType> class SListView;
 /** A delegate type invoked when a selection changes somewhere. */
 DECLARE_DELEGATE_RetVal_OneParam(FText, FGetComboItemDisplayString, int32);
 
+template <>
+struct TWidgetTypeTraits<class SPinComboBox>
+{
+	static constexpr bool SupportsInvalidation() { return true; }
+};
+
 //Class implementation to create combo box
 class SPinComboBox : public SCompoundWidget
 {
diff --git a/Engine/Source/Editor/MaterialEditor/Public/MaterialEditor/MaterialPins/SGraphPinMaterialInput.h b/Engine/Source/Editor/MaterialEditor/Public/MaterialEditor/MaterialPins/SGraphPinMaterialInput.h
--- a/Engine/Source/Editor/MaterialEditor/Public/MaterialEditor/MaterialPins/SGraphPinMaterialInput.h
+++ b/Engine/Source/Editor/MaterialEditor/Public/MaterialEditor/MaterialPins/SGraphPinMaterialInput.h
@@ -11,6 +11,12 @@
 
 class UEdGraphPin;
 
+template <>
+struct TWidgetTypeTraits<class SGraphPinMaterialInput>
+{
+	static constexpr bool SupportsInvalidation() { return true; }
+};
+
 class SGraphPinMaterialInput : public SGraphPin
 {
 public:
