Slate Dynamic Invalidation -Trivial Material Editor Graph Nodes and related widgets

#Synced-CL 50145007

diff --git a/Engine/Source/Editor/GraphEditor/Private/SGraphNodeResizable.cpp b/Engine/Source/Editor/GraphEditor/Private/SGraphNodeResizable.cpp
--- a/Engine/Source/Editor/GraphEditor/Private/SGraphNodeResizable.cpp
+++ b/Engine/Source/Editor/GraphEditor/Private/SGraphNodeResizable.cpp
@@ -119,6 +119,7 @@ FReply SGraphNodeResizable::OnMouseButtonUp( const FGeometry& MyGeometry, const
 		UserSize.Y = FMath::RoundToFloat(UserSize.Y);
 
 		GetNodeObj()->ResizeNode(UserSize);
+		Invalidate(EInvalidateWidgetReason::Prepass);
 
 		// End resize transaction
 		ResizeTransactionPtr.Reset();
@@ -206,6 +207,7 @@ FReply SGraphNodeResizable::OnMouseMove(const FGeometry& MyGeometry, const FPoin
 			UserSize = SnappedSize;
 			GraphNode->ResizeNode( UserSize );
 			DeltaNodePos = GetCorrectedNodePosition() - GetPosition2f();
+			Invalidate(EInvalidateWidgetReason::Prepass);
 		}
 
 		if (!ResizeTransactionPtr.IsValid() && UserSize != StoredUserSize)
diff --git a/Engine/Source/Editor/GraphEditor/Public/SCommentBubble.h b/Engine/Source/Editor/GraphEditor/Public/SCommentBubble.h
--- a/Engine/Source/Editor/GraphEditor/Public/SCommentBubble.h
+++ b/Engine/Source/Editor/GraphEditor/Public/SCommentBubble.h
@@ -31,6 +31,13 @@ struct FPointerEvent;
 DECLARE_DELEGATE_RetVal( bool, FIsGraphNodeHovered );
 DECLARE_DELEGATE_OneParam( FOnCommentBubbleToggled, bool );
 
+
+template <>
+struct TWidgetTypeTraits<class SCommentBubble>
+{
+	static constexpr bool SupportsInvalidation() { return true; }
+};
+
 class SCommentBubble : public SCompoundWidget
 {
 	SLATE_BEGIN_ARGS( SCommentBubble )
diff --git a/Engine/Source/Editor/GraphEditor/Public/SGraphNode.h b/Engine/Source/Editor/GraphEditor/Public/SGraphNode.h
--- a/Engine/Source/Editor/GraphEditor/Public/SGraphNode.h
+++ b/Engine/Source/Editor/GraphEditor/Public/SGraphNode.h
@@ -60,6 +60,12 @@ struct FSlateBrush;
 /////////////////////////////////////////////////////
 // SNodeTitle
 
+template <>
+struct TWidgetTypeTraits<class SNodeTitle>
+{
+	static constexpr bool SupportsInvalidation() { return true; }
+};
+
 class SNodeTitle : public SCompoundWidget
 {
 public:
@@ -122,6 +128,12 @@ protected:
 /////////////////////////////////////////////////////
 // SGraphNode
 
+template <>
+struct TWidgetTypeTraits<class SGraphNode>
+{
+	static constexpr bool SupportsInvalidation() { return true; }
+};
+
 class SGraphNode : public SNodePanel::SNode
 {
 public:
diff --git a/Engine/Source/Editor/GraphEditor/Public/SGraphNodeComment.h b/Engine/Source/Editor/GraphEditor/Public/SGraphNodeComment.h
--- a/Engine/Source/Editor/GraphEditor/Public/SGraphNodeComment.h
+++ b/Engine/Source/Editor/GraphEditor/Public/SGraphNodeComment.h
@@ -27,6 +27,13 @@ struct FGeometry;
 struct FPointerEvent;
 struct FSlateBrush;
 
+
+template <>
+struct TWidgetTypeTraits<class SGraphNodeComment>
+{
+	static constexpr bool SupportsInvalidation() { return true; }
+};
+
 class SGraphNodeComment : public SGraphNodeResizable
 {
 public:
diff --git a/Engine/Source/Editor/GraphEditor/Public/SGraphNodeDefault.h b/Engine/Source/Editor/GraphEditor/Public/SGraphNodeDefault.h
--- a/Engine/Source/Editor/GraphEditor/Public/SGraphNodeDefault.h
+++ b/Engine/Source/Editor/GraphEditor/Public/SGraphNodeDefault.h
@@ -11,6 +11,13 @@
 
 class UEdGraphNode;
 
+
+template <>
+struct TWidgetTypeTraits<class SGraphNodeDefault>
+{
+	static constexpr bool SupportsInvalidation() { return true; }
+};
+
 class SGraphNodeDefault : public SGraphNode
 {
 public:
diff --git a/Engine/Source/Editor/GraphEditor/Public/SGraphNodeKnot.h b/Engine/Source/Editor/GraphEditor/Public/SGraphNodeKnot.h
--- a/Engine/Source/Editor/GraphEditor/Public/SGraphNodeKnot.h
+++ b/Engine/Source/Editor/GraphEditor/Public/SGraphNodeKnot.h
@@ -18,6 +18,13 @@ class SGraphPin;
 class UEdGraphPin;
 struct FSlateBrush;
 
+
+template <>
+struct TWidgetTypeTraits<class SGraphNodeKnot>
+{
+	static constexpr bool SupportsInvalidation() { return true; }
+};
+
 /** The visual representation of a control point meant to adjust how connections are routed, also known as a Reroute node.
  * The input knot node should have properly implemented ShouldDrawNodeAsControlPointOnly to return true with valid indices for its pins.
  */
@@ -58,6 +65,13 @@ protected:
 	const FSlateBrush* ShadowBrushSelected;
 };
 
+
+template <>
+struct TWidgetTypeTraits<class SGraphPinKnot>
+{
+	static constexpr bool SupportsInvalidation() { return true; }
+};
+
 class SGraphPinKnot : public SGraphPin
 {
 public:
diff --git a/Engine/Source/Editor/GraphEditor/Public/SGraphNodeResizable.h b/Engine/Source/Editor/GraphEditor/Public/SGraphNodeResizable.h
--- a/Engine/Source/Editor/GraphEditor/Public/SGraphNodeResizable.h
+++ b/Engine/Source/Editor/GraphEditor/Public/SGraphNodeResizable.h
@@ -16,6 +16,13 @@ class FScopedTransaction;
 struct FGeometry;
 struct FPointerEvent;
 
+
+template <>
+struct TWidgetTypeTraits<class SGraphNodeResizable>
+{
+	static constexpr bool SupportsInvalidation() { return true; }
+};
+
 class SGraphNodeResizable : public SGraphNode
 {
 public:
diff --git a/Engine/Source/Editor/MaterialEditor/Public/MaterialEditor/MaterialNodes/SGraphNodeMaterialBase.h b/Engine/Source/Editor/MaterialEditor/Public/MaterialEditor/MaterialNodes/SGraphNodeMaterialBase.h
--- a/Engine/Source/Editor/MaterialEditor/Public/MaterialEditor/MaterialNodes/SGraphNodeMaterialBase.h
+++ b/Engine/Source/Editor/MaterialEditor/Public/MaterialEditor/MaterialNodes/SGraphNodeMaterialBase.h
@@ -90,6 +90,12 @@ private:
 	bool bIsRealtime;
 };
 
+template<>
+struct TWidgetTypeTraits<class SGraphNodeMaterialBase>
+{
+	static constexpr bool SupportsInvalidation() { return true; }
+};
+
 class SGraphNodeMaterialBase : public SGraphNode
 {
 public:
diff --git a/Engine/Source/Editor/MaterialEditor/Public/MaterialEditor/MaterialNodes/SGraphNodeMaterialCustom.h b/Engine/Source/Editor/MaterialEditor/Public/MaterialEditor/MaterialNodes/SGraphNodeMaterialCustom.h
--- a/Engine/Source/Editor/MaterialEditor/Public/MaterialEditor/MaterialNodes/SGraphNodeMaterialCustom.h
+++ b/Engine/Source/Editor/MaterialEditor/Public/MaterialEditor/MaterialNodes/SGraphNodeMaterialCustom.h
@@ -8,6 +8,12 @@
 
 #define UE_API MATERIALEDITOR_API
 
+template <>
+struct TWidgetTypeTraits<class SGraphNodeMaterialCustom>
+{
+	static constexpr bool SupportsInvalidation() { return true; }
+};
+
 class SGraphNodeMaterialCustom : public SGraphNodeMaterialBase
 {
 public:
diff --git a/Engine/Source/Editor/MaterialEditor/Public/MaterialEditor/MaterialNodes/SGraphNodeMaterialResult.h b/Engine/Source/Editor/MaterialEditor/Public/MaterialEditor/MaterialNodes/SGraphNodeMaterialResult.h
--- a/Engine/Source/Editor/MaterialEditor/Public/MaterialEditor/MaterialNodes/SGraphNodeMaterialResult.h
+++ b/Engine/Source/Editor/MaterialEditor/Public/MaterialEditor/MaterialNodes/SGraphNodeMaterialResult.h
@@ -12,6 +12,12 @@
 class FGraphNodeMetaData;
 class UMaterialGraphNode_Root;
 
+template <>
+struct TWidgetTypeTraits<class SGraphNodeMaterialResult>
+{
+	static constexpr bool SupportsInvalidation() { return true; }
+};
+
 class SGraphNodeMaterialResult : public SGraphNode
 {
 public:
diff --git a/Engine/Source/Editor/MaterialEditor/Public/MaterialNodes/SGraphNodeMaterialComment.h b/Engine/Source/Editor/MaterialEditor/Public/MaterialNodes/SGraphNodeMaterialComment.h
--- a/Engine/Source/Editor/MaterialEditor/Public/MaterialNodes/SGraphNodeMaterialComment.h
+++ b/Engine/Source/Editor/MaterialEditor/Public/MaterialNodes/SGraphNodeMaterialComment.h
@@ -11,6 +11,13 @@
 
 class UMaterialGraphNode_Comment;
 
+
+template <>
+struct TWidgetTypeTraits<class SGraphNodeMaterialComment>
+{
+	static constexpr bool SupportsInvalidation() { return true; }
+};
+
 class SGraphNodeMaterialComment : public SGraphNodeComment
 {
 public:
diff --git a/Engine/Source/Runtime/Slate/Public/Widgets/Notifications/SErrorText.h b/Engine/Source/Runtime/Slate/Public/Widgets/Notifications/SErrorText.h
--- a/Engine/Source/Runtime/Slate/Public/Widgets/Notifications/SErrorText.h
+++ b/Engine/Source/Runtime/Slate/Public/Widgets/Notifications/SErrorText.h
@@ -28,6 +28,12 @@ public:
 };
 
 
+template <>
+struct TWidgetTypeTraits<class SErrorText>
+{
+	static constexpr bool SupportsInvalidation() { return true; }
+};
+
 /**
  * Implements a widget that displays an error text message.
  */
@@ -63,9 +69,6 @@ public:
 	SLATE_API virtual TSharedRef<SWidget> AsWidget() override;
 
 private:
-
-	TAttribute< FSlateFontInfo > Font;
-
 	TAttribute<EVisibility> CustomVisibility;
 	EVisibility MyVisibility() const;
 
diff --git a/Engine/Source/Runtime/Slate/Public/Widgets/Notifications/SPopUpErrorText.h b/Engine/Source/Runtime/Slate/Public/Widgets/Notifications/SPopUpErrorText.h
--- a/Engine/Source/Runtime/Slate/Public/Widgets/Notifications/SPopUpErrorText.h
+++ b/Engine/Source/Runtime/Slate/Public/Widgets/Notifications/SPopUpErrorText.h
@@ -9,6 +9,12 @@
 #include "Widgets/Notifications/SErrorText.h"
 #include "Widgets/Input/SComboButton.h"
 
+template <>
+struct TWidgetTypeTraits<class SPopupErrorText>
+{
+	static constexpr bool SupportsInvalidation() { return true; }
+};
+
 class SPopupErrorText : public SComboButton, public IErrorReportingWidget
 {
 public:
