Slate Dynamic Invalidation - SNodePanel and SGraphPanel.

#Synced-CL 50897806

diff --git a/Engine/Plugins/Enterprise/DataprepEditor/Source/DataprepEditor/Private/Widgets/DataprepGraph/SDataprepGraphTrackNode.cpp b/Engine/Plugins/Enterprise/DataprepEditor/Source/DataprepEditor/Private/Widgets/DataprepGraph/SDataprepGraphTrackNode.cpp
--- a/Engine/Plugins/Enterprise/DataprepEditor/Source/DataprepEditor/Private/Widgets/DataprepGraph/SDataprepGraphTrackNode.cpp
+++ b/Engine/Plugins/Enterprise/DataprepEditor/Source/DataprepEditor/Private/Widgets/DataprepGraph/SDataprepGraphTrackNode.cpp
@@ -536,7 +536,7 @@ void SDataprepGraphTrackNode::OnActionsOrderChanged()
 
 bool SDataprepGraphTrackNode::RefreshLayout()
 {
-	if(OwnerGraphPanelPtr.IsValid() && OwnerGraphPanelPtr.Pin()->GetAllChildren()->Num() > 0)
+	if(OwnerGraphPanelPtr.IsValid() && OwnerGraphPanelPtr.Pin()->GetManagedChildren()->Num() > 0)
 	{
 		ensure(TrackWidgetPtr.IsValid());
 
diff --git a/Engine/Plugins/Experimental/Animation/Scribble/Source/ScribbleEditor/Private/SScribbleGraph.cpp b/Engine/Plugins/Experimental/Animation/Scribble/Source/ScribbleEditor/Private/SScribbleGraph.cpp
--- a/Engine/Plugins/Experimental/Animation/Scribble/Source/ScribbleEditor/Private/SScribbleGraph.cpp
+++ b/Engine/Plugins/Experimental/Animation/Scribble/Source/ScribbleEditor/Private/SScribbleGraph.cpp
@@ -397,7 +397,7 @@ void SScribbleGraphPanel::Tick(const FGeometry& AllottedGeometry, const double I
 		{
 			if (IsScribbleEnabled())
 			{
-				PanelToSync->SetViewOffset(ViewOffset);
+				PanelToSync->SetViewOffset(GetViewOffset());
 				PanelToSync->SetZoomLevel(ZoomLevel);
 			}
 			else
@@ -457,7 +457,7 @@ int32 SScribbleGraphPanel::OnPaint(const FPaintArgs& Args, const FGeometry& Allo
 	if (ScribbleEnabled.Get(false))
 	{
 		// draw a dark overlay below our content to indicate that we are in scribble mode
-		if (!ShouldDrawBackground.Get(false))
+		if (!GetShouldDrawBackground())
 		{
 			static const FSlateBrush* DefaultBackground = FAppStyle::GetBrush(TEXT("Graph.Panel.SolidBackground"));
 			static constexpr FLinearColor DarkenColor = FLinearColor(1,1,1, 0.6);
@@ -506,7 +506,7 @@ int32 SScribbleGraphPanel::OnPaint(const FPaintArgs& Args, const FGeometry& Allo
 					FVector2f ScribbleNodeCorner = ScribbleNodePositionWithOffset;
 
 					// draw a rounded box around the anchored content
-					const FVector2f ScribbleNodePositionInView = (ScribbleNodePositionWithOffset - ViewOffset) * GetZoomAmount();
+					const FVector2f ScribbleNodePositionInView = (ScribbleNodePositionWithOffset - GetViewOffset()) * GetZoomAmount();
              		const FVector2f ScribbleNodeSizeInView = ScribbleNodeSizeWithOffset * GetZoomAmount();
                     const FPaintGeometry PaintGeometry = AllottedGeometry.ToPaintGeometry(ScribbleNodeSizeInView, FSlateLayoutTransform(ScribbleNodePositionInView));
                     FSlateDrawElement::MakeBox(OutDrawElements, ++LastLayerId, PaintGeometry, RoundedFrameBrush.Get(), ESlateDrawEffect::None, FillColor);
@@ -542,8 +542,8 @@ int32 SScribbleGraphPanel::OnPaint(const FPaintArgs& Args, const FGeometry& Allo
 
 					// draw a connecting line to the box 
 					const TArray<FVector2f> AnchorVertices = {
-						(AnchorPosition.GetValue() - ViewOffset) * GetZoomAmount(),
-						(ScribbleNodeCorner - ViewOffset) * GetZoomAmount()
+						(AnchorPosition.GetValue() - GetViewOffset()) * GetZoomAmount(),
+						(ScribbleNodeCorner - GetViewOffset()) * GetZoomAmount()
 					};
 					FSlateDrawElement::MakeLines(OutDrawElements, LastLayerId, AllottedGeometry.ToPaintGeometry(), AnchorVertices, ESlateDrawEffect::None, LineColor, false, RoundedFrameBrush->OutlineSettings.Width);
 				}
diff --git a/Engine/Plugins/Runtime/Metasound/Source/MetasoundEditor/Private/MetasoundEditorGraphConnectionDrawingPolicy.cpp b/Engine/Plugins/Runtime/Metasound/Source/MetasoundEditor/Private/MetasoundEditorGraphConnectionDrawingPolicy.cpp
--- a/Engine/Plugins/Runtime/Metasound/Source/MetasoundEditor/Private/MetasoundEditorGraphConnectionDrawingPolicy.cpp
+++ b/Engine/Plugins/Runtime/Metasound/Source/MetasoundEditor/Private/MetasoundEditorGraphConnectionDrawingPolicy.cpp
@@ -469,7 +469,7 @@ namespace Metasound
 				return;
 			}
 
-			const FDrawConnectionData ConnectionData(Start, End, *EditorPtr, FGraphConnectionDrawingPolicy::ComputeSplineTangent(Start, End), Params, *Settings, LocalMousePosition);
+			const FDrawConnectionData ConnectionData(Start, End, *EditorPtr, FGraphConnectionDrawingPolicy::ComputeSplineTangent(Start, End), Params, *Settings, AbsoluteMousePosition);
 			if (Settings->bTreatSplinesLikePins)
 			{
 #if METASOUND_EDITOR_DEBUG_CONNECTIONS
diff --git a/Engine/Plugins/Runtime/nDisplay/Source/DisplayClusterConfigurator/Private/Views/OutputMapping/SDisplayClusterConfiguratorGraphEditor.cpp b/Engine/Plugins/Runtime/nDisplay/Source/DisplayClusterConfigurator/Private/Views/OutputMapping/SDisplayClusterConfiguratorGraphEditor.cpp
--- a/Engine/Plugins/Runtime/nDisplay/Source/DisplayClusterConfigurator/Private/Views/OutputMapping/SDisplayClusterConfiguratorGraphEditor.cpp
+++ b/Engine/Plugins/Runtime/nDisplay/Source/DisplayClusterConfigurator/Private/Views/OutputMapping/SDisplayClusterConfiguratorGraphEditor.cpp
@@ -981,7 +981,7 @@ void SDisplayClusterConfiguratorGraphEditor::ResetNodeTransform()
 void SDisplayClusterConfiguratorGraphEditor::SortNodes()
 {
 	// Sort the nodes via their depth to ensure that overlapping nodes are properly rendered on top of each order and mouse interactions are properly handled
-	TSlotlessChildren<SNodePanel::SNode>* VisibleChildren = static_cast<TSlotlessChildren<SNodePanel::SNode>*>(GetGraphPanel()->GetChildren());
+	TSlotlessChildren<SNodePanel::SNode>* VisibleChildren = static_cast<TSlotlessChildren<SNodePanel::SNode>*>(GetGraphPanel()->GetVisibleChildren());
 
 	struct SNodeLessThanSort
 	{
diff --git a/Engine/Plugins/VirtualProduction/SceneState/Source/SceneStateMachineEditor/Private/SceneStateMachineConnectionDrawingPolicy.cpp b/Engine/Plugins/VirtualProduction/SceneState/Source/SceneStateMachineEditor/Private/SceneStateMachineConnectionDrawingPolicy.cpp
--- a/Engine/Plugins/VirtualProduction/SceneState/Source/SceneStateMachineEditor/Private/SceneStateMachineConnectionDrawingPolicy.cpp
+++ b/Engine/Plugins/VirtualProduction/SceneState/Source/SceneStateMachineEditor/Private/SceneStateMachineConnectionDrawingPolicy.cpp
@@ -214,15 +214,15 @@ void FStateMachineConnectionDrawingPolicy::DrawArrowLine(const FVector2f& InStar
 	DrawConnection(WireLayerID, InStartPoint, EndPoint, InParams);
 
 	// Draw the transition grab handles in case the mouse is hovering the transition
-	const FVector2f ClosestPoint = FMath::ClosestPointOnSegment2D(LocalMousePosition, InStartPoint, EndPoint);
+	const FVector2f ClosestPoint = FMath::ClosestPointOnSegment2D(AbsoluteMousePosition, InStartPoint, EndPoint);
 
 	constexpr double RelinkHandleHoverRadius = 20.0;
 
-	if ((ClosestPoint - LocalMousePosition).Length() < RelinkHandleHoverRadius)
+	if ((ClosestPoint - AbsoluteMousePosition).Length() < RelinkHandleHoverRadius)
 	{
 		const double RelinkHoverRadiusSq = FMath::Square(RelinkHandleHoverRadius);
-		const double StartMouseSizeSq = FVector2f(InStartPoint - LocalMousePosition).SizeSquared();
-		const double EndMouseSizeSq = FVector2f(EndPoint - LocalMousePosition).SizeSquared();
+		const double StartMouseSizeSq = FVector2f(InStartPoint - AbsoluteMousePosition).SizeSquared();
+		const double EndMouseSizeSq = FVector2f(EndPoint - AbsoluteMousePosition).SizeSquared();
 
 		if (EndMouseSizeSq < RelinkHoverRadiusSq)
 		{
diff --git a/Engine/Source/Editor/AnimationBlueprintEditor/Private/StateMachineConnectionDrawingPolicy.cpp b/Engine/Source/Editor/AnimationBlueprintEditor/Private/StateMachineConnectionDrawingPolicy.cpp
--- a/Engine/Source/Editor/AnimationBlueprintEditor/Private/StateMachineConnectionDrawingPolicy.cpp
+++ b/Engine/Source/Editor/AnimationBlueprintEditor/Private/StateMachineConnectionDrawingPolicy.cpp
@@ -170,17 +170,17 @@ void FStateMachineConnectionDrawingPolicy::Internal_DrawLineWithArrow(const FVec
 	// Draw the transition grab handles in case the mouse is hovering the transition
 	bool bStartHovered = false;
 	bool bEndHovered = false;
-	const FVector2f FVecMousePos = FVector2f(LocalMousePosition.X, LocalMousePosition.Y);
+	const FVector2f FVecMousePos = FVector2f(AbsoluteMousePosition.X, AbsoluteMousePosition.Y);
 	const FVector2f ClosestPoint = FMath::ClosestPointOnSegment2D(FVecMousePos, StartPoint, EndPoint);
 	if ((ClosestPoint - FVecMousePos).Length() < RelinkHandleHoverRadius * ZoomFactor)
 	{
-		bStartHovered = (StartPoint - LocalMousePosition).Length() < RelinkHandleHoverRadius * ZoomFactor;
-		bEndHovered = (EndPoint - LocalMousePosition).Length() < RelinkHandleHoverRadius * ZoomFactor;
+		bStartHovered = (StartPoint - AbsoluteMousePosition).Length() < RelinkHandleHoverRadius * ZoomFactor;
+		bEndHovered = (EndPoint - AbsoluteMousePosition).Length() < RelinkHandleHoverRadius * ZoomFactor;
 		FVector2f HoverIndicatorPosition = bStartHovered ? StartPoint : EndPoint;
 
 		// Set the hovered pin results. This will be used by the SGraphPanel again.
-		const float SquaredDistToPin1 = (Params.AssociatedPin1 != nullptr) ? (StartPoint - LocalMousePosition).SizeSquared() : FLT_MAX;
-		const float SquaredDistToPin2 = (Params.AssociatedPin2 != nullptr) ? (EndPoint - LocalMousePosition).SizeSquared() : FLT_MAX;
+		const float SquaredDistToPin1 = (Params.AssociatedPin1 != nullptr) ? (StartPoint - AbsoluteMousePosition).SizeSquared() : FLT_MAX;
+		const float SquaredDistToPin2 = (Params.AssociatedPin2 != nullptr) ? (EndPoint - AbsoluteMousePosition).SizeSquared() : FLT_MAX;
 		UEdGraphPin* Pin1 = Params.AssociatedPin1;
 		UEdGraphPin* Pin2 = Params.AssociatedPin2;
 		if (Pin2)
diff --git a/Engine/Source/Editor/GraphEditor/Private/ConnectionDrawingPolicy.cpp b/Engine/Source/Editor/GraphEditor/Private/ConnectionDrawingPolicy.cpp
--- a/Engine/Source/Editor/GraphEditor/Private/ConnectionDrawingPolicy.cpp
+++ b/Engine/Source/Editor/GraphEditor/Private/ConnectionDrawingPolicy.cpp
@@ -118,7 +118,10 @@ FConnectionDrawingPolicy::FConnectionDrawingPolicy(int32 InBackLayerID, int32 In
 	, ZoomFactor(InZoomFactor)
 	, ClippingRect(InClippingRect)
 	, DrawElementsList(InDrawElements)
+PRAGMA_DISABLE_DEPRECATION_WARNINGS
 	, LocalMousePosition(0.0f, 0.0f)
+PRAGMA_ENABLE_DEPRECATION_WARNINGS
+	, AbsoluteMousePosition(0.0f, 0.0f)
 {
 	ArrowImage = FAppStyle::GetBrush( TEXT("Graph.Arrow") );
 	ArrowRadius = ArrowImage->ImageSize * ZoomFactor * 0.5f;
@@ -165,14 +168,21 @@ void FConnectionDrawingPolicy::DrawSplineWithArrow_DeprecationHelper(const UE::S
 	}
 }
 
-void FConnectionDrawingPolicy::DrawSplineWithArrow(const FGeometry& StartGeom, const FGeometry& EndGeom, const FConnectionParams& Params)
+void GetSplineEndPoints(const FGeometry& StartGeom, const FGeometry& EndGeom, const FVector2f& InArrowRadius, FVector2f& OutStartPoint, FVector2f& OutEndPoint)
 {
 	//@TODO: These values should be pushed into the Slate style, they are compensating for a bit of
 	// empty space inside of the pin brush images.
 	const float StartFudgeX = 4.0f;
 	const float EndFudgeX = 4.0f;
-	const FVector2f StartPoint = FGeometryHelper::VerticalMiddleRightOf(StartGeom) - FVector2f(StartFudgeX, 0.0f);
-	const FVector2f EndPoint = FGeometryHelper::VerticalMiddleLeftOf(EndGeom) - FVector2f(ArrowRadius.X - EndFudgeX, 0);
+	OutStartPoint = FGeometryHelper::VerticalMiddleRightOf(StartGeom) - FVector2f(StartFudgeX, 0.0f);
+	OutEndPoint = FGeometryHelper::VerticalMiddleLeftOf(EndGeom) - FVector2f(InArrowRadius.X - EndFudgeX, 0);
+}
+
+void FConnectionDrawingPolicy::DrawSplineWithArrow(const FGeometry& StartGeom, const FGeometry& EndGeom, const FConnectionParams& Params)
+{
+	FVector2f StartPoint;
+	FVector2f EndPoint;
+	GetSplineEndPoints(StartGeom, EndGeom, ArrowRadius, StartPoint, EndPoint);
 
 	DrawSplineWithArrow(StartPoint, EndPoint, Params);
 }
@@ -218,7 +228,15 @@ void FConnectionDrawingPolicy::SetHoveredPins(const TSet< FEdGraphPinReference >
 
 void FConnectionDrawingPolicy::SetMousePosition(const UE::Slate::FDeprecateVector2DParameter& InMousePos)
 {
-	LocalMousePosition = InMousePos;
+	SetAbsoluteMousePosition(InMousePos);
+}
+
+void FConnectionDrawingPolicy::SetAbsoluteMousePosition(FVector2f InAbsoluteMousePos)
+{
+PRAGMA_DISABLE_DEPRECATION_WARNINGS
+	LocalMousePosition = InAbsoluteMousePos;
+PRAGMA_ENABLE_DEPRECATION_WARNINGS
+	AbsoluteMousePosition = MoveTemp(InAbsoluteMousePos);
 }
 
 void FConnectionDrawingPolicy::SetSliceLine(const FMarqueeRect& InLine)
@@ -315,127 +333,14 @@ void FConnectionDrawingPolicy::DrawConnection_DeprecationHelper(int32 LayerId, c
 	const FVector2f P1Tangent = (Params.EndTangent.IsNearlyZero()) ? ((Params.EndDirection == EGPD_Input) ? SplineTangent : -SplineTangent) : Params.EndTangent;
 
 	bool bSliceLineIntersectsSpline = false;
-
-	if (Settings->bTreatSplinesLikePins)
-	{
-		// Distance to consider as an overlap
-		const float QueryDistanceTriggerThresholdSquared = FMath::Square(Settings->SplineHoverTolerance + Params.WireThickness * 0.5f);
-
-		// Distance to pass the bounding box cull test. This is used for the bCloseToSpline output that can be used as a
-		// dead zone to avoid mistakes caused by missing a double-click on a connection.
-		const float QueryDistanceForCloseSquared = FMath::Square(FMath::Sqrt(QueryDistanceTriggerThresholdSquared) + Settings->SplineCloseTolerance);
-
-		bool bSliceLineIntersectsSplineBounds = false;
-		bool bCloseToSpline = false;
-		{
-			// The curve will include the endpoints but can extend out of a tight bounds because of the tangents
-			// P0Tangent coefficient maximizes to 4/27 at a=1/3, and P1Tangent minimizes to -4/27 at a=2/3.
-			constexpr float MaximumTangentContribution = 4.0f / 27.0f;
-			FBox2f Bounds(ForceInit);
-
-			Bounds += FVector2f(P0);
-			Bounds += FVector2f(P0 + MaximumTangentContribution * P0Tangent);
-			Bounds += FVector2f(P1);
-			Bounds += FVector2f(P1 - MaximumTangentContribution * P1Tangent);
-
-			bCloseToSpline = Bounds.ComputeSquaredDistanceToPoint(LocalMousePosition) < QueryDistanceForCloseSquared;
-
-			// Re-use the bounds we just constructed to check if the slice line intersects the spline
-			if (SliceLine.IsValid())
-			{
-				bSliceLineIntersectsSplineBounds = FMath::LineBoxIntersection2D(Bounds, SliceLine.StartPoint, SliceLine.EndPoint);
-			}
-
-			// Draw the bounding box for debugging
-#if 0
-#define DrawSpaceLine(Point1, Point2, DebugWireColor) {const FVector2f FakeTangent = (Point2 - Point1).GetSafeNormal(); FSlateDrawElement::MakeDrawSpaceSpline(DrawElementsList, LayerId, Point1, FakeTangent, Point2, FakeTangent, ClippingRect, 1.0f, ESlateDrawEffect::None, DebugWireColor); }
-
-			if (bCloseToSpline)
-			{
-				const FLinearColor BoundsWireColor = bCloseToSpline ? FLinearColor::Green : FLinearColor::White;
-
-				FVector2f TL = Bounds.Min;
-				FVector2f BR = Bounds.Max;
-				FVector2f TR = FVector2f(Bounds.Max.X, Bounds.Min.Y);
-				FVector2f BL = FVector2f(Bounds.Min.X, Bounds.Max.Y);
-
-				DrawSpaceLine(TL, TR, BoundsWireColor);
-				DrawSpaceLine(TR, BR, BoundsWireColor);
-				DrawSpaceLine(BR, BL, BoundsWireColor);
-				DrawSpaceLine(BL, TL, BoundsWireColor);
-			}
-#endif
-		}
-		
-		constexpr int32 NumStepsToTest = 16;
-		constexpr float TestStepInterval = 1.0f / NumStepsToTest;
-		
-		// Only bother checking for precise spline intersection if we intersected the bounding box
-		if (bSliceLineIntersectsSplineBounds)
-		{
-			// Find the intersecting point on the spline
-			FVector2f OutIntersectionPoint;
-			FVector2f Point1 = FMath::CubicInterp(P0, P0Tangent, P1, P1Tangent, 0.0f);
-			for (float TestAlpha = 0.0f; TestAlpha < 1.0f; TestAlpha += TestStepInterval)
-			{
-				const FVector2f Point2 = FMath::CubicInterp(P0, P0Tangent, P1, P1Tangent, TestAlpha + TestStepInterval);
-				
-				if (IntersectionHelpers::SegmentIntersection2D(SliceLine.StartPoint, SliceLine.EndPoint, Point1, Point2, OutIntersectionPoint))
-				{
-					bSliceLineIntersectsSpline = true;
-					ConnectionsIntersectingSliceLine.Emplace(Params.AssociatedPin1, Params.AssociatedPin2);
-					break;
-				}
-
-				Point1 = Point2;
-			}
-		}
-
-		if (bCloseToSpline)
-		{
-			// Find the closest approach to the spline
-			FVector2f ClosestPoint(ForceInit);
-			float ClosestDistanceSquared = FLT_MAX;
-			FVector2f Point1 = FMath::CubicInterp(P0, P0Tangent, P1, P1Tangent, 0.0f);
-			for (float TestAlpha = 0.0f; TestAlpha < 1.0f; TestAlpha += TestStepInterval)
-			{
-				const FVector2f Point2 = FMath::CubicInterp(P0, P0Tangent, P1, P1Tangent, TestAlpha + TestStepInterval);
-
-				const FVector2f ClosestPointToSegment = FMath::ClosestPointOnSegment2D(LocalMousePosition, Point1, Point2);
-				const float DistanceSquared = (LocalMousePosition - ClosestPointToSegment).SizeSquared();
-
-				if (DistanceSquared < ClosestDistanceSquared)
-				{
-					ClosestDistanceSquared = DistanceSquared;
-					ClosestPoint = ClosestPointToSegment;
-				}
-
-				Point1 = Point2;
-			}
-
-			// Record the overlap
-			if (ClosestDistanceSquared < QueryDistanceTriggerThresholdSquared)
-			{
-				if (ClosestDistanceSquared < SplineOverlapResult.GetDistanceSquared())
-				{
-					const float SquaredDistToPin1 = (Params.AssociatedPin1 != nullptr) ? (P0 - ClosestPoint).SizeSquared() : FLT_MAX;
-					const float SquaredDistToPin2 = (Params.AssociatedPin2 != nullptr) ? (P1 - ClosestPoint).SizeSquared() : FLT_MAX;
-
-					SplineOverlapResult = FGraphSplineOverlapResult(Params.AssociatedPin1, Params.AssociatedPin2, ClosestDistanceSquared, SquaredDistToPin1, SquaredDistToPin2, true);
-				}
-			}
-			else if (ClosestDistanceSquared < QueryDistanceForCloseSquared)
-			{
-				SplineOverlapResult.SetCloseToSpline(true);
-			}
-		}
-	}
+	CheckSplineConnectionOverlapWithCursor(AbsoluteMousePosition, Start, End, Params, /*out*/ SplineOverlapResult, /*out*/ &bSliceLineIntersectsSpline);
 
 	// Apply an opacity fade to wires intersected by the slice line
 	FLinearColor WireColor = Params.WireColor;
 	if (bSliceLineIntersectsSpline)
 	{
 		WireColor.A *= SliceDeemphasisAlphaMultiplier;
+		ConnectionsIntersectingSliceLine.Emplace(Params.AssociatedPin1, Params.AssociatedPin2);
 	}
 
 	// Draw the spline itself
@@ -608,6 +513,61 @@ void FConnectionDrawingPolicy::Draw(TMap<TSharedRef<SWidget>, FArrangedWidget>&
 	DrawPinGeometries(InPinGeometries, ArrangedNodes);
 }
 
+bool FConnectionDrawingPolicy::HaveConnectionsGraphicallyChanged(const SGraphPanel& InGraphPanel, const FVector2f& AbsoluteCursorPosition) const
+{
+	FGraphSplineOverlapResult OverlapResult;
+
+	for (const FSimpleConnectionData& Connection : LastConnections)
+	{
+		FConnectionParams Params;
+		Params.AssociatedPin1 = Connection.EndPin->GetPinObj();
+		Params.AssociatedPin2 = Connection.StartPin ? Connection.StartPin->GetPinObj() : nullptr;
+		CheckSplineConnectionOverlapWithCursor(AbsoluteCursorPosition, Connection.StartPoint, Connection.EndPoint, Params, /*out*/ OverlapResult);
+	}
+
+	TSharedPtr<SGraphPin> NewPin1Widget;
+	TSharedPtr<SGraphPin> NewPin2Widget;
+	OverlapResult.GetPinWidgets(InGraphPanel, NewPin1Widget, NewPin2Widget);
+
+	TSet<FEdGraphPinReference> NewlyHoveredPins;
+	if (NewPin1Widget && NewPin1Widget->GetPinObj())
+	{
+		NewlyHoveredPins.Add(NewPin1Widget->GetPinObj());
+	}
+	if (NewPin2Widget && NewPin2Widget->GetPinObj())
+	{
+		NewlyHoveredPins.Add(NewPin2Widget->GetPinObj());
+	}
+
+	return HaveConnectionsGraphicallyChanged(NewlyHoveredPins);
+}
+
+bool FConnectionDrawingPolicy::HaveConnectionsGraphicallyChanged(const TSet<FEdGraphPinReference>& InNewlyHoveredPins) const
+{
+	if ( HoveredPins.Num() != InNewlyHoveredPins.Num() || !HoveredPins.Difference(InNewlyHoveredPins).IsEmpty() )
+	{
+		return true; // If the Hovered Pins are different, we do not check further
+	}
+
+	// Now we know that TempHoveredPins and HoveredPins are equivalent, so DetermineWiringStyle to give us the right information
+	// (as DetermineWiringStyle often checks the content of HoveredPins to decide how to draw them)
+	for (const FSimpleConnectionData& Connection : LastConnections)
+	{
+		FConnectionParams Params;
+		const_cast<FConnectionDrawingPolicy*>(this)->DetermineWiringStyle(Connection.EndPin->GetPinObj(), Connection.StartPin ? Connection.StartPin->GetPinObj() : nullptr, /*inout*/ Params);
+		if (Connection.StartPin && Connection.StartPin->AreConnectionsFaded() && Connection.EndPin->AreConnectionsFaded() )
+		{
+			Params.WireColor.A = 0.2f;
+		}
+
+		if (Params.WireThickness != Connection.Params.WireThickness || Params.WireColor != Connection.Params.WireColor)
+		{
+			return true;
+		}
+	}
+	return false;
+}
+
 void FConnectionDrawingPolicy::BuildPinToPinWidgetMap(TMap<TSharedRef<SWidget>, FArrangedWidget>& InPinGeometries)
 {
 	PinToPinWidgetMap.Empty();
@@ -651,13 +611,152 @@ void FConnectionDrawingPolicy::DrawPinGeometries(TMap<TSharedRef<SWidget>, FArra
 							Params.WireColor.A = 0.2f;
 						}
 					}
+
 					DrawSplineWithArrow(LinkStartWidgetGeometry->Geometry, LinkEndWidgetGeometry->Geometry, Params);
+
+					FVector2f StartPoint;
+					FVector2f EndPoint;
+					GetSplineEndPoints(LinkStartWidgetGeometry->Geometry, LinkEndWidgetGeometry->Geometry, ArrowRadius, StartPoint, EndPoint);
+					LastConnections.Add(FSimpleConnectionData{ConnectedPinWidget ? *ConnectedPinWidget : nullptr, StaticCastSharedRef<SGraphPin>(SomePinWidget), StartPoint, EndPoint, Params});
 				}
 			}
 		}
 	}
 }
 
+void FConnectionDrawingPolicy::CheckSplineConnectionOverlapWithCursor(const FVector2f& AbsoluteCursorPosition, const FVector2f& Start, const FVector2f& End, const FConnectionParams& Params, FGraphSplineOverlapResult& InOutOverlapResults, bool* bOutSliceLineIntersectsSpline) const
+{
+	const FVector2f& P0 = Start;
+	const FVector2f& P1 = End;
+
+	const FVector2f SplineTangent = ComputeSplineTangent(P0, P1);
+
+	const FVector2f P0Tangent = (Params.StartTangent.IsNearlyZero()) ? ((Params.StartDirection == EGPD_Output) ? SplineTangent : -SplineTangent) : Params.StartTangent;
+	const FVector2f P1Tangent = (Params.EndTangent.IsNearlyZero()) ? ((Params.EndDirection == EGPD_Input) ? SplineTangent : -SplineTangent) : Params.EndTangent;
+
+	if (bOutSliceLineIntersectsSpline)
+	{
+		*bOutSliceLineIntersectsSpline = false;
+	}
+
+	if (!Settings->bTreatSplinesLikePins)
+	{
+		return;
+	}
+
+	// Distance to consider as an overlap
+	const float QueryDistanceTriggerThresholdSquared = FMath::Square(Settings->SplineHoverTolerance + Params.WireThickness * 0.5f);
+
+	// Distance to pass the bounding box cull test. This is used for the bCloseToSpline output that can be used as a
+	// dead zone to avoid mistakes caused by missing a double-click on a connection.
+	const float QueryDistanceForCloseSquared = FMath::Square(FMath::Sqrt(QueryDistanceTriggerThresholdSquared) + Settings->SplineCloseTolerance);
+
+	bool bSliceLineIntersectsSplineBounds = false;
+	bool bCloseToSpline = false;
+	{
+		// The curve will include the endpoints but can extend out of a tight bounds because of the tangents
+		// P0Tangent coefficient maximizes to 4/27 at a=1/3, and P1Tangent minimizes to -4/27 at a=2/3.
+		constexpr float MaximumTangentContribution = 4.0f / 27.0f;
+		FBox2f Bounds(ForceInit);
+
+		Bounds += FVector2f(P0);
+		Bounds += FVector2f(P0 + MaximumTangentContribution * P0Tangent);
+		Bounds += FVector2f(P1);
+		Bounds += FVector2f(P1 - MaximumTangentContribution * P1Tangent);
+
+		bCloseToSpline = Bounds.ComputeSquaredDistanceToPoint(AbsoluteCursorPosition) < QueryDistanceForCloseSquared;
+
+		// Re-use the bounds we just constructed to check if the slice line intersects the spline
+		if (SliceLine.IsValid() && bOutSliceLineIntersectsSpline)
+		{
+			bSliceLineIntersectsSplineBounds = FMath::LineBoxIntersection2D(Bounds, SliceLine.StartPoint, SliceLine.EndPoint);
+		}
+
+		// Draw the bounding box for debugging
+#if 0
+#define DrawSpaceLine(Point1, Point2, DebugWireColor) {const FVector2f FakeTangent = (Point2 - Point1).GetSafeNormal(); FSlateDrawElement::MakeDrawSpaceSpline(DrawElementsList, LayerId, Point1, FakeTangent, Point2, FakeTangent, ClippingRect, 1.0f, ESlateDrawEffect::None, DebugWireColor); }
+
+		if (bCloseToSpline)
+		{
+			const FLinearColor BoundsWireColor = bCloseToSpline ? FLinearColor::Green : FLinearColor::White;
+
+			FVector2f TL = Bounds.Min;
+			FVector2f BR = Bounds.Max;
+			FVector2f TR = FVector2f(Bounds.Max.X, Bounds.Min.Y);
+			FVector2f BL = FVector2f(Bounds.Min.X, Bounds.Max.Y);
+
+			DrawSpaceLine(TL, TR, BoundsWireColor);
+			DrawSpaceLine(TR, BR, BoundsWireColor);
+			DrawSpaceLine(BR, BL, BoundsWireColor);
+			DrawSpaceLine(BL, TL, BoundsWireColor);
+		}
+#endif
+	}
+
+	constexpr int32 NumStepsToTest = 16;
+	constexpr float TestStepInterval = 1.0f / NumStepsToTest;
+
+	// Only bother checking for precise spline intersection if we intersected the bounding box
+	if (bSliceLineIntersectsSplineBounds)
+	{
+		// Find the intersecting point on the spline
+		FVector2f OutIntersectionPoint;
+		FVector2f Point1 = FMath::CubicInterp(P0, P0Tangent, P1, P1Tangent, 0.0f);
+		for (float TestAlpha = 0.0f; TestAlpha < 1.0f; TestAlpha += TestStepInterval)
+		{
+			const FVector2f Point2 = FMath::CubicInterp(P0, P0Tangent, P1, P1Tangent, TestAlpha + TestStepInterval);
+
+			if (IntersectionHelpers::SegmentIntersection2D(SliceLine.StartPoint, SliceLine.EndPoint, Point1, Point2, OutIntersectionPoint))
+			{
+				check(bOutSliceLineIntersectsSpline);
+				*bOutSliceLineIntersectsSpline = true;
+				break;
+			}
+
+			Point1 = Point2;
+		}
+	}
+
+	if (bCloseToSpline)
+	{
+		// Find the closest approach to the spline
+		FVector2f ClosestPoint(ForceInit);
+		float ClosestDistanceSquared = FLT_MAX;
+		FVector2f Point1 = FMath::CubicInterp(P0, P0Tangent, P1, P1Tangent, 0.0f);
+		for (float TestAlpha = 0.0f; TestAlpha < 1.0f; TestAlpha += TestStepInterval)
+		{
+			const FVector2f Point2 = FMath::CubicInterp(P0, P0Tangent, P1, P1Tangent, TestAlpha + TestStepInterval);
+
+			const FVector2f ClosestPointToSegment = FMath::ClosestPointOnSegment2D(AbsoluteCursorPosition, Point1, Point2);
+			const float DistanceSquared = (AbsoluteCursorPosition - ClosestPointToSegment).SizeSquared();
+
+			if (DistanceSquared < ClosestDistanceSquared)
+			{
+				ClosestDistanceSquared = DistanceSquared;
+				ClosestPoint = ClosestPointToSegment;
+			}
+
+			Point1 = Point2;
+		}
+
+		// Record the overlap
+		if (ClosestDistanceSquared < QueryDistanceTriggerThresholdSquared)
+		{
+			if (ClosestDistanceSquared < InOutOverlapResults.GetDistanceSquared())
+			{
+				const float SquaredDistToPin1 = (Params.AssociatedPin1 != nullptr) ? (P0 - ClosestPoint).SizeSquared() : FLT_MAX;
+				const float SquaredDistToPin2 = (Params.AssociatedPin2 != nullptr) ? (P1 - ClosestPoint).SizeSquared() : FLT_MAX;
+
+				InOutOverlapResults = FGraphSplineOverlapResult(Params.AssociatedPin1, Params.AssociatedPin2, ClosestDistanceSquared, SquaredDistToPin1, SquaredDistToPin2, true);
+			}
+		}
+		else if (ClosestDistanceSquared < QueryDistanceForCloseSquared)
+		{
+			InOutOverlapResults.SetCloseToSpline(true);
+		}
+	}
+}
+
 bool FConnectionDrawingPolicy::IsConnectionCulled( const FArrangedWidget& StartLink, const FArrangedWidget& EndLink ) const
 {
 	const float Top		= FMath::Min( StartLink.Geometry.AbsolutePosition.Y, EndLink.Geometry.AbsolutePosition.Y );
diff --git a/Engine/Source/Editor/GraphEditor/Private/DragConnection.cpp b/Engine/Source/Editor/GraphEditor/Private/DragConnection.cpp
--- a/Engine/Source/Editor/GraphEditor/Private/DragConnection.cpp
+++ b/Engine/Source/Editor/GraphEditor/Private/DragConnection.cpp
@@ -59,6 +59,7 @@ void FDragConnection::OnDragged(const class FDragDropEvent& DragDropEvent)
 	CursorDecoratorWindow->MoveWindowTo(DragDropEvent.GetScreenSpacePosition() + DecoratorAdjust);
 	// Request the active panel to scroll if required
 	GraphPanel->RequestDeferredPan(TargetPosition);
+	GraphPanel->Invalidate(EInvalidateWidgetReason::Paint);
 }
 
 void FDragConnection::HoverTargetChanged()
diff --git a/Engine/Source/Editor/GraphEditor/Private/GraphEditorSnapping/NodeSnappingManager.cpp b/Engine/Source/Editor/GraphEditor/Private/GraphEditorSnapping/NodeSnappingManager.cpp
--- a/Engine/Source/Editor/GraphEditor/Private/GraphEditorSnapping/NodeSnappingManager.cpp
+++ b/Engine/Source/Editor/GraphEditor/Private/GraphEditorSnapping/NodeSnappingManager.cpp
@@ -929,7 +929,7 @@ TSharedPtr<SGraphNode> FNodeSnappingManager::FindNodeForSnap(TSharedPtr<SGraphNo
 
 		FArrangedWidget ArrangedWidget = MyGeometry.MakeChild(
 			SomeChild, 
-			SomeChild->GetPosition2f() - OwningPanelPinned->ViewOffset, 
+			SomeChild->GetPosition2f() - OwningPanelPinned->GetViewOffset(),
 			SomeChild->GetDesiredSize(), 
 			OwningPanelPinned->GetZoomAmount()
 		);
diff --git a/Engine/Source/Editor/GraphEditor/Private/SGraphNode.cpp b/Engine/Source/Editor/GraphEditor/Private/SGraphNode.cpp
--- a/Engine/Source/Editor/GraphEditor/Private/SGraphNode.cpp
+++ b/Engine/Source/Editor/GraphEditor/Private/SGraphNode.cpp
@@ -105,8 +105,12 @@ void SNodeTitle::Tick( const FGeometry& AllottedGeometry, const double InCurrent
 	// Checks to see if the cached string is valid, and if not, updates it.
 	if (NodeTitleCache.IsOutOfDate(GraphNode.Get()))
 	{
-		NodeTitleCache.SetCachedText(TitleText.Get(), GraphNode.Get());
-		RebuildWidget();
+		// Do not rebuild if the text was empty and is still empty
+		if (!(NodeTitleCache.GetCachedText().IsEmpty() && TitleText.Get().IsEmpty()))
+		{
+			NodeTitleCache.SetCachedText(TitleText.Get(), GraphNode.Get());
+			RebuildWidget();
+		}
 	}
 }
 
@@ -586,6 +590,7 @@ void SGraphNode::MoveTo( const FVector2f& NewPosition, FNodeSet& NodeFilter, boo
 			NodeFilter.Add( SharedThis( this ) );
 			GraphNode->Modify(bMarkDirty);
 			GraphNode->SetPosition(NewPosition);
+			Invalidate(EInvalidateWidgetReason::Layout);
 		}
 	}
 }
@@ -960,10 +965,7 @@ void SGraphNode::UpdateGraphNode()
 	FGraphNodeMetaData TagMeta(TEXT("Graphnode"));
 	PopulateMetaTag(&TagMeta);
 	
-	TSharedPtr<SVerticalBox> InnerVerticalBox;
-
-
-	InnerVerticalBox = SNew(SVerticalBox)
+	TSharedPtr<SVerticalBox> InnerVerticalBox = SNew(SVerticalBox)
 		+SVerticalBox::Slot()
 		.AutoHeight()
 		.HAlign(HAlign_Fill)
diff --git a/Engine/Source/Editor/GraphEditor/Private/SGraphNodeComment.cpp b/Engine/Source/Editor/GraphEditor/Private/SGraphNodeComment.cpp
--- a/Engine/Source/Editor/GraphEditor/Private/SGraphNodeComment.cpp
+++ b/Engine/Source/Editor/GraphEditor/Private/SGraphNodeComment.cpp
@@ -292,6 +292,7 @@ FReply SGraphNodeComment::OnMouseButtonUp( const FGeometry& MyGeometry, const FP
 		UserSize.Y = FMath::RoundToFloat(UserSize.Y);
 
 		GetNodeObj()->ResizeNode(UserSize);
+		Invalidate(EInvalidateWidgetReason::Prepass);
 
 		// End resize transaction
 		ResizeTransactionPtr.Reset();
@@ -323,7 +324,7 @@ void SGraphNodeComment::HandleSelection(bool bSelected, bool bUpdateNodesUnderCo
 			if (CommentNode)
 			{
 				TSharedPtr<SGraphPanel> Panel = Comment->GetOwnerPanel();
-				FChildren* PanelChildren = Panel->GetAllChildren();
+				FChildren* PanelChildren = Panel->GetManagedChildren();
 				int32 NumChildren = PanelChildren->Num();
 				CommentNode->ClearNodesUnderComment();
 
@@ -420,7 +421,7 @@ void SGraphNodeComment::EndUserInteraction() const
 		const FSlateRect CommentRect( NodePosition.X, NodePosition.Y, NodePosition.X + NodeSize.X, NodePosition.Y + NodeSize.Y );
 
 		TSharedPtr<SGraphPanel> Panel = GetOwnerPanel();
-		FChildren* PanelChildren = Panel->GetAllChildren();
+		FChildren* PanelChildren = Panel->GetManagedChildren();
 		int32 NumChildren = PanelChildren->Num();
 		static FString SGraphNodeCommentType = "SGraphNodeComment";
 
diff --git a/Engine/Source/Editor/GraphEditor/Private/SGraphPanel.cpp b/Engine/Source/Editor/GraphEditor/Private/SGraphPanel.cpp
--- a/Engine/Source/Editor/GraphEditor/Private/SGraphPanel.cpp
+++ b/Engine/Source/Editor/GraphEditor/Private/SGraphPanel.cpp
@@ -88,6 +88,24 @@ DEFINE_LOG_CATEGORY_STATIC(LogGraphPanel, Log, All);
 //////////////////////////////////////////////////////////////////////////
 // SGraphPanel
 
+SLATE_IMPLEMENT_WIDGET(SGraphPanel)
+void SGraphPanel::PrivateRegisterAttributes(FSlateAttributeInitializer& AttributeInitializer)
+{
+	SLATE_ADD_MEMBER_ATTRIBUTE_DEFINITION(AttributeInitializer, ShowGraphStateOverlayAttribute, EInvalidateWidgetReason::Paint);
+	SLATE_ADD_MEMBER_ATTRIBUTE_DEFINITION(AttributeInitializer, ShouldDrawBackgroundAttribute, EInvalidateWidgetReason::Paint);
+	SLATE_ADD_MEMBER_ATTRIBUTE_DEFINITION(AttributeInitializer, ShouldDrawSurroundingShadowAttribute, EInvalidateWidgetReason::Paint);
+	SLATE_ADD_MEMBER_ATTRIBUTE_DEFINITION(AttributeInitializer, DisplayAsReadOnlyAttribute, EInvalidateWidgetReason::Paint);
+}
+
+
+SGraphPanel::SGraphPanel()
+	: ShowGraphStateOverlayAttribute(*this, true)
+	, ShouldDrawBackgroundAttribute(*this, true)
+	, ShouldDrawSurroundingShadowAttribute(*this, true)
+	, DisplayAsReadOnlyAttribute(*this, false)
+{
+}
+
 void SGraphPanel::Construct( const SGraphPanel::FArguments& InArgs )
 {
 	SNodePanel::Construct();
@@ -96,9 +114,9 @@ void SGraphPanel::Construct( const SGraphPanel::FArguments& InArgs )
 	this->GraphObj = InArgs._GraphObj;
 	this->DiffResults = InArgs._DiffResults;
 	this->FocusedDiffResult = InArgs._FocusedDiffResult;
-	this->SelectionManager.OnSelectionChanged = InArgs._OnSelectionChanged;
-	this->IsEditable = InArgs._IsEditable;
-	this->DisplayAsReadOnly = InArgs._DisplayAsReadOnly;
+	this->OnSelectionChangedDelegate = InArgs._OnSelectionChanged;
+	SetIsEditable(InArgs._IsEditable);
+	SetDisplayAsReadOnly(InArgs._DisplayAsReadOnly);
 	this->OnNodeDoubleClicked = InArgs._OnNodeDoubleClicked;
 	this->OnDropActors = InArgs._OnDropActors;
 	this->OnDropStreamingLevels = InArgs._OnDropStreamingLevels;
@@ -132,16 +150,16 @@ void SGraphPanel::Construct( const SGraphPanel::FArguments& InArgs )
 	MyRegisteredGraphChangedDelegate = FOnGraphChanged::FDelegate::CreateSP(this, &SGraphPanel::OnGraphChanged);
 	MyRegisteredGraphChangedDelegateHandle = this->GraphObj->AddOnGraphChangedHandler(MyRegisteredGraphChangedDelegate);
 	
-	ShowGraphStateOverlay = InArgs._ShowGraphStateOverlay;
+	SetShowGraphStateOverlay(InArgs._ShowGraphStateOverlay);
 	AllowConnectionSlicing = InArgs._AllowConnectionSlicing;
-	ShouldDrawBackground = InArgs._ShouldDrawBackground;
-	ShouldDrawSurroundingShadow = InArgs._ShouldDrawSurroundingShadow;
+	SetShouldDrawBackground(InArgs._ShouldDrawBackground);
+	SetShouldDrawSurroundingShadow(InArgs._ShouldDrawSurroundingShadow);
 
 	SavedMousePosForOnPaintEventLocalSpace = FVector2f::ZeroVector;
 	PreviousFrameSavedMousePosForSplineOverlap = FVector2f::ZeroVector;
 
 	TimeLeftToInvalidatePerTick = 0.0f;
-	
+
 	bHasCustomPrepass = true;
 	bCheckNodeGraphObjValidity = false;
 
@@ -172,7 +190,7 @@ int32 SGraphPanel::OnPaint( const FPaintArgs& Args, const FGeometry& AllottedGeo
 	// First paint the background
 	const UEditorExperimentalSettings& Options = *GetDefault<UEditorExperimentalSettings>();
 
-	if (ShouldDrawBackground.Get(true))
+	if (GetShouldDrawBackground())
 	{
 		const FSlateBrush* DefaultBackground = FAppStyle::GetBrush(TEXT("Graph.Panel.SolidBackground"));
 		const FSlateBrush* CustomBackground = &GetDefault<UEditorStyleSettings>()->GraphBackgroundBrush;
@@ -289,198 +307,193 @@ int32 SGraphPanel::OnPaint( const FPaintArgs& Args, const FGeometry& AllottedGeo
 			}
 
 
-			const bool bNodeIsVisible = FSlateRect::DoRectanglesIntersect( CurWidget.Geometry.GetLayoutBoundingRect(), MyCullingRect );
+			const bool bSelected = SelectionToVisualize->Contains( StaticCastSharedRef<SNodePanel::SNode>(CurWidget.Widget)->GetObjectBeingDisplayed() );
 
-			if (bNodeIsVisible)
-			{
-				const bool bSelected = SelectionToVisualize->Contains( StaticCastSharedRef<SNodePanel::SNode>(CurWidget.Widget)->GetObjectBeingDisplayed() );
-				
-				UEdGraphNode* NodeObj = Cast<UEdGraphNode>(ChildNode->GetObjectBeingDisplayed());
-				float Alpha = 1.0f;
+			UEdGraphNode* NodeObj = Cast<UEdGraphNode>(ChildNode->GetObjectBeingDisplayed());
+			float Alpha = 1.0f;
 
-				// Handle Node renaming once the node is visible
-				if( bSelected && ChildNode->IsRenamePending() )
+			// Handle Node renaming once the node is visible
+			if( bSelected && ChildNode->IsRenamePending() )
+			{
+				// Only open a rename when the window has focus
+				TSharedPtr<SWindow> OwnerWindow = FSlateApplication::Get().FindWidgetWindow(AsShared());
+				if (!OwnerWindow.IsValid() || FSlateApplication::Get().HasFocusedDescendants(OwnerWindow.ToSharedRef()))
 				{
-					// Only open a rename when the window has focus
-					TSharedPtr<SWindow> OwnerWindow = FSlateApplication::Get().FindWidgetWindow(AsShared());
-					if (!OwnerWindow.IsValid() || FSlateApplication::Get().HasFocusedDescendants(OwnerWindow.ToSharedRef()))
-					{
-						ChildNode->ApplyRename();
-					}
+					ChildNode->ApplyRename();
 				}
+			}
 
-				/** if this graph is being diffed, highlight the changes in the graph */
-				if(DiffResults.IsValid())
+			/** if this graph is being diffed, highlight the changes in the graph */
+			if(DiffResults.IsValid())
+			{
+				/** When diffing nodes, color code shadow based on diff result */
+				if (NodeDiffResults.Contains(NodeObj))
 				{
-					/** When diffing nodes, color code shadow based on diff result */
-					if (NodeDiffResults.Contains(NodeObj))
+					const FDiffSingleResult& DiffResult = NodeDiffResults[NodeObj];
+					for (const SNode::DiffHighlightInfo& Highlight : ChildNode->GetDiffHighlights(DiffResult))
 					{
-						const FDiffSingleResult& DiffResult = NodeDiffResults[NodeObj];
-						for (const SNode::DiffHighlightInfo& Highlight : ChildNode->GetDiffHighlights(DiffResult))
-						{
-							FSlateDrawElement::MakeBox(
-								OutDrawElements,
-								NodeDiffHighlightLayerID,
-								CurWidget.Geometry.ToInflatedPaintGeometry(NodeShadowSize),
-								Highlight.Brush,
-								ESlateDrawEffect::None,
-								Highlight.Tint
-								);
-						}
+						FSlateDrawElement::MakeBox(
+							OutDrawElements,
+							NodeDiffHighlightLayerID,
+							CurWidget.Geometry.ToInflatedPaintGeometry(NodeShadowSize),
+							Highlight.Brush,
+							ESlateDrawEffect::None,
+							Highlight.Tint
+							);
 					}
 				}
-				
-				/** When diffing, set the backround of the differing pins to their diff colors */
-				for (UEdGraphPin* Pin : NodeObj->Pins)
+			}
+
+			/** When diffing, set the backround of the differing pins to their diff colors */
+			for (UEdGraphPin* Pin : NodeObj->Pins)
+			{
+				if (TSharedPtr<SGraphPin> PinWidget = ChildNode->FindWidgetForPin(Pin))
 				{
-					if (TSharedPtr<SGraphPin> PinWidget = ChildNode->FindWidgetForPin(Pin))
+					if (FDiffSingleResult* DiffResult = PinDiffResults.Find(Pin))
 					{
-						if (FDiffSingleResult* DiffResult = PinDiffResults.Find(Pin))
+						// if the diff result associated with this pin is focused, highlight the pin
+						if (DiffResults.IsValid() && FocusedDiffResult.IsSet())
 						{
-							// if the diff result associated with this pin is focused, highlight the pin
-							if (DiffResults.IsValid() && FocusedDiffResult.IsSet())
+							const int32 Index = FocusedDiffResult.Get();
+							if (DiffResults->IsValidIndex(Index))
 							{
-								const int32 Index = FocusedDiffResult.Get();
-								if (DiffResults->IsValidIndex(Index))
-								{
-									const FDiffSingleResult& Focused = (*DiffResults)[Index];
-									PinWidget->SetDiffHighlighted(*DiffResult == Focused);
-								}
+								const FDiffSingleResult& Focused = (*DiffResults)[Index];
+								PinWidget->SetDiffHighlighted(*DiffResult == Focused);
 							}
-						
-							FLinearColor PinDiffColor = DiffResult->GetDisplayColor();
-							PinDiffColor.A = 0.7f;
-							PinWidget->SetPinDiffColor(PinDiffColor);
-							PinWidget->SetFadeConnections(false);
 						}
-						else
-						{
-							PinWidget->SetDiffHighlighted(false);
-							PinWidget->SetPinDiffColor(TOptional<FLinearColor>());
 
-							// when zoomed out, fade out pin connections that aren't involved in a diff
-							PinWidget->SetFadeConnections(ZoomLevel <= 6 && (!NodeDiffResults.Contains(NodeObj) || NodeDiffResults[NodeObj].Pin1));
-						}
+						FLinearColor PinDiffColor = DiffResult->GetDisplayColor();
+						PinDiffColor.A = 0.7f;
+						PinWidget->SetPinDiffColor(PinDiffColor);
+						PinWidget->SetFadeConnections(false);
+					}
+					else
+					{
+						PinWidget->SetDiffHighlighted(false);
+						PinWidget->SetPinDiffColor(TOptional<FLinearColor>());
+
+						// when zoomed out, fade out pin connections that aren't involved in a diff
+						PinWidget->SetFadeConnections(ZoomLevel <= 6 && (!NodeDiffResults.Contains(NodeObj) || NodeDiffResults[NodeObj].Pin1));
 					}
 				}
+			}
+
+			// Draw the node's shadow.
+			if (bDrawShadowsThisFrame || bSelected)
+			{
+				// draw drop shadow on connected pins:
+				ChildNode->DrawChildrenNodeShadow(
+					*SelectionToVisualize,
+					/*FSlateWindowElementList& */OutDrawElements,
+					NodeShadowsLayerId,
+					CurWidget.Geometry,
+					FLinearColor(1.0f, 1.0f, 1.0f, Alpha)
+				);
+			}
+
+			// Draw the comments and information popups for this node, if it has any.
+			{
+				const SNodePanel::SNode::FNodeSlot* CommentSlot = ChildNode->GetSlot( ENodeZone::TopCenter );
+				float CommentBubbleY = CommentSlot ? -CommentSlot->GetSlotOffset().Y : 0.f;
+				Context.bSelected = bSelected;
+				TArray<FGraphInformationPopupInfo> Popups;
 
-				// Draw the node's shadow.
-				if (bDrawShadowsThisFrame || bSelected)
 				{
-					// draw drop shadow on connected pins:
-					ChildNode->DrawChildrenNodeShadow(
-						*SelectionToVisualize,
-						/*FSlateWindowElementList& */OutDrawElements, 
-						NodeShadowsLayerId,
-						CurWidget.Geometry,
-						FLinearColor(1.0f, 1.0f, 1.0f, Alpha)
-					);
+					ChildNode->GetNodeInfoPopups(&Context, /*out*/ Popups);
 				}
 
-				// Draw the comments and information popups for this node, if it has any.
+				for (int32 PopupIndex = 0; PopupIndex < Popups.Num(); ++PopupIndex)
 				{
-					const SNodePanel::SNode::FNodeSlot* CommentSlot = ChildNode->GetSlot( ENodeZone::TopCenter );
-					float CommentBubbleY = CommentSlot ? -CommentSlot->GetSlotOffset().Y : 0.f;
-					Context.bSelected = bSelected;
-					TArray<FGraphInformationPopupInfo> Popups;
+					FGraphInformationPopupInfo& Popup = Popups[PopupIndex];
+					PaintComment(Popup.Message, CurWidget.Geometry, MyCullingRect, OutDrawElements, ChildLayerId, Popup.BackgroundColor, /*inout*/ CommentBubbleY, InWidgetStyle);
+				}
+			}
 
-					{
-						ChildNode->GetNodeInfoPopups(&Context, /*out*/ Popups);
-					}
+			int32 CurWidgetsMaxLayerId;
+			{
+				/* When dragging off a pin, we want to duck the alpha of some nodes */
+				TSharedPtr< SGraphPin > OnlyStartPin = (1 == PreviewConnectorFromPins.Num()) ? PreviewConnectorFromPins[0].FindInGraphPanel(*this) : TSharedPtr< SGraphPin >();
+				const bool bNodeIsNotUsableInCurrentContext = Schema->FadeNodeWhenDraggingOffPin(NodeObj, OnlyStartPin.IsValid() ? OnlyStartPin.Get()->GetPinObj() : nullptr);
 
-					for (int32 PopupIndex = 0; PopupIndex < Popups.Num(); ++PopupIndex)
-					{
-						FGraphInformationPopupInfo& Popup = Popups[PopupIndex];
-						PaintComment(Popup.Message, CurWidget.Geometry, MyCullingRect, OutDrawElements, ChildLayerId, Popup.BackgroundColor, /*inout*/ CommentBubbleY, InWidgetStyle);
-					}
-				}
+				const bool bCleanDiff = DiffResults.IsValid() && !NodeDiffResults.Contains(NodeObj);
 
-				int32 CurWidgetsMaxLayerId;
+				FWidgetStyle NodeStyleToUse = InWidgetStyle;
+				if (bNodeIsNotUsableInCurrentContext)
 				{
-					/* When dragging off a pin, we want to duck the alpha of some nodes */
-					TSharedPtr< SGraphPin > OnlyStartPin = (1 == PreviewConnectorFromPins.Num()) ? PreviewConnectorFromPins[0].FindInGraphPanel(*this) : TSharedPtr< SGraphPin >();
-					const bool bNodeIsNotUsableInCurrentContext = Schema->FadeNodeWhenDraggingOffPin(NodeObj, OnlyStartPin.IsValid() ? OnlyStartPin.Get()->GetPinObj() : nullptr);
-					
-					const bool bCleanDiff = DiffResults.IsValid() && !NodeDiffResults.Contains(NodeObj);
-					
-					FWidgetStyle NodeStyleToUse = InWidgetStyle;
-					if (bNodeIsNotUsableInCurrentContext)
-					{
-						NodeStyleToUse = FadedStyle;
-					}
-					else if (ZoomLevel <= 6 && bCleanDiff)
-					{
-						NodeStyleToUse = FadedStyle;
-					}
-					NodeStyleToUse.BlendColorAndOpacityTint(FLinearColor(1.0f, 1.0f, 1.0f, Alpha));
-
-					// Draw the node.
-					CurWidgetsMaxLayerId = CurWidget.Widget->Paint(NewArgs, CurWidget.Geometry, MyCullingRect, OutDrawElements, ChildLayerId, NodeStyleToUse, !DisplayAsReadOnly.Get() && ShouldBeEnabled( bParentEnabled ) );
+					NodeStyleToUse = FadedStyle;
 				}
+				else if (ZoomLevel <= 6 && bCleanDiff)
+				{
+					NodeStyleToUse = FadedStyle;
+				}
+				NodeStyleToUse.BlendColorAndOpacityTint(FLinearColor(1.0f, 1.0f, 1.0f, Alpha));
+
+				// Draw the node.
+				CurWidgetsMaxLayerId = CurWidget.Widget->Paint(NewArgs, CurWidget.Geometry, MyCullingRect, OutDrawElements, ChildLayerId, NodeStyleToUse, !DisplayAsReadOnlyAttribute.Get() && ShouldBeEnabled( bParentEnabled ) );
+			}
+
+			// Draw the node's overlay, if it has one.
+			{
+				// Get its size
+				const FVector2f WidgetSize = CurWidget.Geometry.Size;
 
-				// Draw the node's overlay, if it has one.
 				{
-					// Get its size
-					const FVector2f WidgetSize = CurWidget.Geometry.Size;
+					TArray<FOverlayBrushInfo> OverlayBrushes;
+					ChildNode->GetOverlayBrushes(bSelected, WidgetSize, /*out*/ OverlayBrushes);
 
+					for (int32 BrushIndex = 0; BrushIndex < OverlayBrushes.Num(); ++BrushIndex)
 					{
-						TArray<FOverlayBrushInfo> OverlayBrushes;
-						ChildNode->GetOverlayBrushes(bSelected, WidgetSize, /*out*/ OverlayBrushes);
-
-						for (int32 BrushIndex = 0; BrushIndex < OverlayBrushes.Num(); ++BrushIndex)
+						FOverlayBrushInfo& OverlayInfo = OverlayBrushes[BrushIndex];
+						const FSlateBrush* OverlayBrush = OverlayInfo.Brush;
+						if (OverlayBrush != nullptr)
 						{
-							FOverlayBrushInfo& OverlayInfo = OverlayBrushes[BrushIndex];
-							const FSlateBrush* OverlayBrush = OverlayInfo.Brush;
-							if (OverlayBrush != nullptr)
-							{
-								FPaintGeometry BouncedGeometry = CurWidget.Geometry.ToPaintGeometry(OverlayBrush->ImageSize, FSlateLayoutTransform(OverlayInfo.OverlayOffset));
-
-								// Handle bouncing during PIE
-								const float BounceValue = FMath::Sin(2.0f * PI * BounceCurve.GetLerp());
-								BouncedGeometry.DrawPosition += FVector2f(OverlayInfo.AnimationEnvelope * BounceValue * ZoomFactor);
-
-								FLinearColor FinalColorAndOpacity(InWidgetStyle.GetColorAndOpacityTint()* OverlayBrush->GetTint(InWidgetStyle));
-								//FinalColorAndOpacity.A = Alpha;
-
-								CurWidgetsMaxLayerId++;
-								FSlateDrawElement::MakeBox(
-									OutDrawElements,
-									CurWidgetsMaxLayerId,
-									BouncedGeometry,
-									OverlayBrush,
-									ESlateDrawEffect::None,
-									FinalColorAndOpacity
-									);
-							}
+							FPaintGeometry BouncedGeometry = CurWidget.Geometry.ToPaintGeometry(OverlayBrush->ImageSize, FSlateLayoutTransform(OverlayInfo.OverlayOffset));
+
+							// Handle bouncing during PIE
+							const float BounceValue = FMath::Sin(2.0f * PI * BounceCurve.GetLerp());
+							BouncedGeometry.DrawPosition += FVector2f(OverlayInfo.AnimationEnvelope * BounceValue * ZoomFactor);
+
+							FLinearColor FinalColorAndOpacity(InWidgetStyle.GetColorAndOpacityTint()* OverlayBrush->GetTint(InWidgetStyle));
+							//FinalColorAndOpacity.A = Alpha;
 
+							CurWidgetsMaxLayerId++;
+							FSlateDrawElement::MakeBox(
+								OutDrawElements,
+								CurWidgetsMaxLayerId,
+								BouncedGeometry,
+								OverlayBrush,
+								ESlateDrawEffect::None,
+								FinalColorAndOpacity
+								);
 						}
+
 					}
+				}
 
-					{
-						TArray<FOverlayWidgetInfo> OverlayWidgets = ChildNode->GetOverlayWidgets(bSelected, WidgetSize);
+				{
+					TArray<FOverlayWidgetInfo> OverlayWidgets = ChildNode->GetOverlayWidgets(bSelected, WidgetSize);
 
-						for (int32 WidgetIndex = 0; WidgetIndex < OverlayWidgets.Num(); ++WidgetIndex)
+					for (int32 WidgetIndex = 0; WidgetIndex < OverlayWidgets.Num(); ++WidgetIndex)
+					{
+						FOverlayWidgetInfo& OverlayInfo = OverlayWidgets[WidgetIndex];
+						if (SWidget* Widget = OverlayInfo.Widget.Get())
 						{
-							FOverlayWidgetInfo& OverlayInfo = OverlayWidgets[WidgetIndex];
-							if (SWidget* Widget = OverlayInfo.Widget.Get())
+							FSlateAttributeMetaData::UpdateOnlyVisibilityAttributes(*Widget, FSlateAttributeMetaData::EInvalidationPermission::AllowInvalidationIfConstructed);
+							if (Widget->GetVisibility() == EVisibility::Visible)
 							{
-								FSlateAttributeMetaData::UpdateOnlyVisibilityAttributes(*Widget, FSlateAttributeMetaData::EInvalidationPermission::AllowInvalidationIfConstructed);
-								if (Widget->GetVisibility() == EVisibility::Visible)
-								{
-									// call SlatePrepass as these widgets are not in the 'normal' child hierarchy
-									Widget->SlatePrepass(AllottedGeometry.GetAccumulatedLayoutTransform().GetScale());
+								// call SlatePrepass as these widgets are not in the 'normal' child hierarchy
+								Widget->SlatePrepass(AllottedGeometry.GetAccumulatedLayoutTransform().GetScale());
 
-									const FGeometry WidgetGeometry = CurWidget.Geometry.MakeChild(Widget->GetDesiredSize(), FSlateLayoutTransform(OverlayInfo.OverlayOffset));
+								const FGeometry WidgetGeometry = CurWidget.Geometry.MakeChild(Widget->GetDesiredSize(), FSlateLayoutTransform(OverlayInfo.OverlayOffset));
 
-									Widget->Paint(NewArgs, WidgetGeometry, MyCullingRect, OutDrawElements, CurWidgetsMaxLayerId, InWidgetStyle, bParentEnabled);
-								}
+								Widget->Paint(NewArgs, WidgetGeometry, MyCullingRect, OutDrawElements, CurWidgetsMaxLayerId, InWidgetStyle, bParentEnabled);
 							}
 						}
 					}
 				}
-
-				MaxLayerId = FMath::Max( MaxLayerId, CurWidgetsMaxLayerId + 1 );
 			}
+
+			MaxLayerId = FMath::Max( MaxLayerId, CurWidgetsMaxLayerId + 1 );
 		}
 	}
 
@@ -489,15 +502,10 @@ int32 SGraphPanel::OnPaint( const FPaintArgs& Args, const FGeometry& AllottedGeo
 	// Draw connections between pins 
 	if (Children.Num() > 0 )
 	{
-		FConnectionDrawingPolicy* ConnectionDrawingPolicy = nullptr;
-		if (NodeFactory.IsValid())
-		{
-			ConnectionDrawingPolicy = NodeFactory->CreateConnectionPolicy(Schema, WireLayerId, MaxLayerId, ZoomFactor, MyCullingRect, OutDrawElements, GraphObj);
-		}
-		else
-		{
-			ConnectionDrawingPolicy = FNodeFactory::CreateConnectionPolicy(Schema, WireLayerId, MaxLayerId, ZoomFactor, MyCullingRect, OutDrawElements, GraphObj);
-		}
+		ConnectionDrawingPolicy = TUniquePtr<FConnectionDrawingPolicy>{ NodeFactory.IsValid() ?
+			NodeFactory->CreateConnectionPolicy(Schema, WireLayerId, MaxLayerId, ZoomFactor, MyCullingRect, OutDrawElements, GraphObj) :
+			FNodeFactory::CreateConnectionPolicy(Schema, WireLayerId, MaxLayerId, ZoomFactor, MyCullingRect, OutDrawElements, GraphObj)
+		};
 
 		const bool bUseDrawStateCaching = ConnectionDrawingPolicy->UseDrawStateCaching();
 
@@ -527,7 +535,7 @@ int32 SGraphPanel::OnPaint( const FPaintArgs& Args, const FGeometry& AllottedGeo
 			ConnectionDrawingPolicy->SetHoveredPins(CurrentHoveredPins, OverridePins, TimeWhenMouseEnteredPin);
 			ConnectionDrawingPolicy->SetMarkedPin(MarkedPin);
 			
-			ConnectionDrawingPolicy->SetMousePosition(AllottedGeometry.LocalToAbsolute(SavedMousePosForOnPaintEventLocalSpace));
+			ConnectionDrawingPolicy->SetAbsoluteMousePosition(AllottedGeometry.LocalToAbsolute(SavedMousePosForOnPaintEventLocalSpace));
 
 			if (IsRelinkingConnection())
 			{
@@ -667,18 +675,20 @@ int32 SGraphPanel::OnPaint( const FPaintArgs& Args, const FGeometry& AllottedGeo
 				const_cast<SGraphPanel*>(this)->CurrentSlice->ConnectionsIntersectingSliceLine = MoveTemp(ConnectionDrawingPolicy->ConnectionsIntersectingSliceLine);
 			}
 		}
-
-		delete ConnectionDrawingPolicy;
+	}
+	else
+	{
+		ConnectionDrawingPolicy.Reset();
 	}
 
-	if (ShouldDrawSurroundingShadow.Get(true))
+	if (GetShouldDrawSurroundingShadow())
 	{
 		// Draw a shadow overlay around the edges of the graph
 		++MaxLayerId;
 		PaintSurroundSunkenShadow(FAppStyle::GetBrush(TEXT("Graph.Shadow")), AllottedGeometry, MyCullingRect, OutDrawElements, MaxLayerId);
 	}
 
-	if (ShowGraphStateOverlay.Get())
+	if (GetShowGraphStateOverlay())
 	{
 		const FSlateBrush* BorderBrush = nullptr;
 		if ((GEditor->bIsSimulatingInEditor || GEditor->PlayWorld != nullptr))
@@ -686,7 +696,7 @@ int32 SGraphPanel::OnPaint( const FPaintArgs& Args, const FGeometry& AllottedGeo
 			// Draw a surrounding indicator when PIE is active, to make it clear that the graph is read-only, etc...
 			BorderBrush = FAppStyle::GetBrush(TEXT("Graph.PlayInEditor"));
 		}
-		else if (!IsEditable.Get())
+		else if (!GetIsEditable())
 		{
 			// Draw a different border when we're not simulating but the graph is read-only
 			BorderBrush = FAppStyle::GetBrush(TEXT("Graph.ReadOnlyBorder"));
@@ -821,7 +831,7 @@ void SGraphPanel::UpdateSelectedNodesPositions(const FVector2f& PositionIncremen
 
 FReply SGraphPanel::OnKeyDown( const FGeometry& MyGeometry, const FKeyEvent& InKeyEvent )
 {
-	if( IsEditable.Get() )
+	if( GetIsEditable() )
 	{
 		const bool bIsModifierActive = InKeyEvent.IsCommandDown() || InKeyEvent.IsAltDown() || InKeyEvent.IsShiftDown() || InKeyEvent.IsControlDown();
 		if (!bIsModifierActive)
@@ -1013,10 +1023,11 @@ FReply SGraphPanel::OnMouseButtonDown(const FGeometry& MyGeometry, const FPointe
 			{
 				LastClickSplineOverlap = PreviousFrameSplineOverlap;
 				// Don't allow starting a slice while hovering a node, or for read-only graphs
-				if (!NodeUnderMousePtr.IsValid() && IsGraphEditable() && !DisplayAsReadOnly.Get())
+				if (!NodeUnderMousePtr.IsValid() && IsGraphEditable() && !DisplayAsReadOnlyAttribute.Get())
 				{
 					const FVector2f GraphMousePos = PanelCoordToGraphCoord( MyGeometry.AbsoluteToLocal( MouseEvent.GetScreenSpacePosition() ) );
 					CurrentSlice = FConnectionSlice(GraphMousePos);
+					Invalidate(EInvalidateWidgetReason::Paint);
 					return FReply::Handled().CaptureMouse(SharedThis(this));
 				}
 			}
@@ -1029,6 +1040,7 @@ FReply SGraphPanel::OnMouseButtonDown(const FGeometry& MyGeometry, const FPointe
 				{
 					const UEdGraphSchema* Schema = GraphObj->GetSchema();
 					Schema->BreakSinglePinLink(Pin1, Pin2);
+					Invalidate(EInvalidateWidgetReason::Paint);
 					return FReply::Handled();
 				}
 			}
@@ -1054,6 +1066,7 @@ FReply SGraphPanel::OnMouseButtonDown(const FGeometry& MyGeometry, const FPointe
 			}
 
 			// fall back legacy behavior, just allow the spline to handle mouse down:
+			Invalidate(EInvalidateWidgetReason::Paint);
 			return BestPinFromHoveredSpline->OnPinMouseDown(MyGeometry, MouseEvent);
 		}
 	}
@@ -1069,6 +1082,7 @@ FReply SGraphPanel::OnMouseButtonDown(const FGeometry& MyGeometry, const FPointe
 				const UEdGraphSchema* Schema = GraphObj->GetSchema();
 				if (Schema->IsConnectionRelinkingAllowed(SourcePin->GetPinObj()))
 				{
+					Invalidate(EInvalidateWidgetReason::Paint);
 					return SourcePin->OnPinMouseDown(MyGeometry, MouseEvent);
 				}
 			}
@@ -1079,6 +1093,7 @@ FReply SGraphPanel::OnMouseButtonDown(const FGeometry& MyGeometry, const FPointe
 	{
 		if (const FReply Reply = OnClicked.Execute(MyGeometry, MouseEvent); Reply.IsEventHandled())
 		{
+			Invalidate(EInvalidateWidgetReason::Paint);
 			return FReply::Handled();
 		}
 	}
@@ -1151,6 +1166,7 @@ FReply SGraphPanel::OnMouseButtonUp(const FGeometry& MyGeometry, const FPointerE
 	{
 		if (SGraphPin* BestPinFromHoveredSpline = GetBestPinFromHoveredSpline())
 		{
+			Invalidate(EInvalidateWidgetReason::Paint);
 			return BestPinFromHoveredSpline->OnMouseButtonUp(MyGeometry, MouseEvent);
 		}
 	}
@@ -1200,7 +1216,7 @@ void SGraphPanel::AddPinToHoverSet(UEdGraphPin* HoveredPin)
 	TimeWhenMouseEnteredPin = FSlateApplication::Get().GetCurrentTime();
 
 	// About covers the fade in time when highlighting pins or splines.
-	TimeLeftToInvalidatePerTick += 1.5f;
+	TimeLeftToInvalidatePerTick = 1.5f;
 
 	// This handle should always be for this function
 	if (!ActiveTimerHandleInvalidatePerTick.IsValid())
@@ -1223,7 +1239,7 @@ void SGraphPanel::ArrangeChildrenForContextMenuSummon(const FGeometry& AllottedG
 		const TSharedRef<SNode>& SomeChild = VisibleChildren[ChildIndex];
 		if (!SomeChild->RequiresSecondPassLayout())
 		{
-			ArrangedChildren.AddWidget(AllottedGeometry.MakeChild(SomeChild, SomeChild->GetPosition2f() - ViewOffset, SomeChild->GetDesiredSizeForMarquee2f(), GetZoomAmount()));
+			ArrangedChildren.AddWidget(AllottedGeometry.MakeChild(SomeChild, SomeChild->GetPosition2f() - GetViewOffset(), SomeChild->GetDesiredSizeForMarquee2f(), GetZoomAmount()));
 		}
 	}
 
@@ -1234,7 +1250,7 @@ void SGraphPanel::ArrangeChildrenForContextMenuSummon(const FGeometry& AllottedG
 		if (SomeChild->RequiresSecondPassLayout())
 		{
 			SomeChild->PerformSecondPassLayout(NodeToWidgetLookup);
-			ArrangedChildren.AddWidget(AllottedGeometry.MakeChild(SomeChild, SomeChild->GetPosition2f() - ViewOffset, SomeChild->GetDesiredSizeForMarquee2f(), GetZoomAmount()));
+			ArrangedChildren.AddWidget(AllottedGeometry.MakeChild(SomeChild, SomeChild->GetPosition2f() - GetViewOffset(), SomeChild->GetDesiredSizeForMarquee2f(), GetZoomAmount()));
 		}
 	}
 }
@@ -1266,7 +1282,7 @@ TSharedPtr<SWidget> SGraphPanel::OnSummonContextMenu(const FGeometry& MyGeometry
 bool SGraphPanel::OnHandleLeftMouseRelease(const FGeometry& MyGeometry, const FPointerEvent& MouseEvent)
 {
 	TSharedPtr<SGraphPin> PreviewConnectionPin = PreviewConnectorFromPins.Num() > 0 ? PreviewConnectorFromPins[0].FindInGraphPanel(*this) : nullptr;
-	if (PreviewConnectionPin.IsValid() && IsEditable.Get())
+	if (PreviewConnectionPin.IsValid() && GetIsEditable())
 	{
 		TSet< TSharedRef<SWidget> > AllConnectors;
 		for (int32 ChildIndex = 0; ChildIndex < Children.Num(); ++ChildIndex)
@@ -1322,6 +1338,14 @@ FReply SGraphPanel::OnMouseMove(const FGeometry& MyGeometry, const FPointerEvent
 	// Save the mouse position to use in OnPaint for spline hit detection
 	SavedMousePosForOnPaintEventLocalSpace = MyGeometry.AbsoluteToLocal(MouseEvent.GetScreenSpacePosition());
 
+	if (Children.Num() > 0 && ConnectionDrawingPolicy)
+	{
+		UE::Slate::FDeprecateVector2DResult AbsolutePos = GetPersistentState().AllottedGeometry.LocalToAbsolute(SavedMousePosForOnPaintEventLocalSpace);
+		if (ConnectionDrawingPolicy->HaveConnectionsGraphicallyChanged(*this, AbsolutePos))
+		{
+			Invalidate(EInvalidateWidgetReason::Paint);
+		}
+	}
 	// Invalidate the spline results if we moved very far
 	const FVector2f MouseDelta = SavedMousePosForOnPaintEventLocalSpace - PreviousFrameSavedMousePosForSplineOverlap;
 	const float MouseDeltaLengthSquared = MouseDelta.SizeSquared();
@@ -1333,10 +1357,11 @@ FReply SGraphPanel::OnMouseMove(const FGeometry& MyGeometry, const FPointerEvent
 		// Our frame-latent approximation is going to be totally junk if the mouse is moving quickly
 		OnSplineHoverStateChanged(FGraphSplineOverlapResult());
 	}
-	
+
 	if(NodeSnappingManager && NodeSnappingManager->IsDragging())
 	{
 		NodeSnappingManager->UpdateDrag(MyGeometry, MouseEvent);
+		Invalidate(EInvalidateWidgetReason::Paint);
 		return FReply::Handled();
 	}
 
@@ -1344,6 +1369,8 @@ FReply SGraphPanel::OnMouseMove(const FGeometry& MyGeometry, const FPointerEvent
 	// Also consume the mouse move since we don't really want other graph behaviour happening while slicing
 	if (CurrentSlice)
 	{
+		Invalidate(EInvalidateWidgetReason::Paint);
+
 		if (this->HasMouseCapture())
 		{
 			const FVector2f GraphMousePos = PanelCoordToGraphCoord(MyGeometry.AbsoluteToLocal(MouseEvent.GetScreenSpacePosition()));
@@ -1397,7 +1424,7 @@ FReply SGraphPanel::OnDragOver( const FGeometry& MyGeometry, const FDragDropEven
 	}
 
 	// Handle Read only graphs
-	if (!IsEditable.Get())
+	if (!GetIsEditable())
 	{
 		TSharedPtr<FGraphEditorDragDropAction> GraphDragDropOp = DragDropEvent.GetOperationAs<FGraphEditorDragDropAction>();
 
@@ -1462,7 +1489,7 @@ FReply SGraphPanel::OnDrop( const FGeometry& MyGeometry, const FDragDropEvent& D
 	FSlateApplication::Get().SetKeyboardFocus(AsShared(), EFocusCause::SetDirectly);
 
 	TSharedPtr<FDragDropOperation> Operation = DragDropEvent.GetOperation();
-	if (!Operation.IsValid() || !IsEditable.Get())
+	if (!Operation.IsValid() || !GetIsEditable())
 	{
 		return FReply::Unhandled();
 	}
@@ -1746,6 +1773,7 @@ FReply SGraphPanel::OnMouseDownDisconnectClosestPinFromWire(TSharedRef<SGraphPin
 	const UEdGraphSchema* Schema = GraphObj->GetSchema();
 	Schema->BreakSinglePinLink(ClosestPin->GetPinObj(), ConnectedPin->GetPinObj());
 	ClosestPin->bIsMovingLinks = true;
+	Invalidate(EInvalidateWidgetReason::Paint);
 	return FReply::Handled().BeginDragDrop(DragEvent.ToSharedRef());
 }
 
@@ -1816,7 +1844,8 @@ void SGraphPanel::RemoveAllNodes()
 {
 	NodeGuidMap.Empty();
 	CurrentHoveredPins.Empty();
-	
+	ConnectionDrawingPolicy.Reset();
+
 	for (int32 Iter = 0; Iter != Children.Num(); ++Iter)
 	{
 		GetChild(Iter)->InvalidateGraphData();
@@ -1908,7 +1937,7 @@ void SGraphPanel::DismissContextMenu()
 void SGraphPanel::AttachGraphEvents(TSharedPtr<SGraphNode> CreatedSubNode)
 {
 	check(CreatedSubNode.IsValid());
-	CreatedSubNode->SetIsEditable(IsEditable);
+	CreatedSubNode->SetIsEditable(GetIsEditableAttribute().ToAttribute());
 	CreatedSubNode->SetDoubleClickEvent(OnNodeDoubleClicked);
 	CreatedSubNode->SetVerifyTextCommitEvent(OnVerifyTextCommit);
 	CreatedSubNode->SetTextCommittedEvent(OnTextCommitted);
@@ -2285,7 +2314,7 @@ void SGraphPanel::AddNode(UEdGraphNode* Node, AddNodeBehavior Behavior)
 		Behavior == NotUserAdded ? false :
 					(UserAddedNodes.Find(Node) != nullptr);
 
-	NewNode->SetIsEditable(IsEditable);
+	NewNode->SetIsEditable(GetIsEditableAttribute().ToAttribute());
 	NewNode->SetDoubleClickEvent(OnNodeDoubleClicked);
 	NewNode->SetVerifyTextCommitEvent(OnVerifyTextCommit);
 	NewNode->SetTextCommittedEvent(OnTextCommitted);
@@ -2386,6 +2415,12 @@ void SGraphPanel::RemoveAllNodesWithInvalidPointers()
 	}
 }
 
+void SGraphPanel::OnSelectionChanged(const TSet<UObject*>& Objects)
+{
+	SNodePanel::OnSelectionChanged(Objects);
+	OnSelectionChangedDelegate.ExecuteIfBound(Objects);
+}
+
 TSharedPtr<SGraphNode> SGraphPanel::GetNodeWidgetFromGuid(FGuid Guid) const
 {
 	return NodeGuidMap.FindRef(Guid).Pin();
diff --git a/Engine/Source/Editor/GraphEditor/Private/SNodePanel.cpp b/Engine/Source/Editor/GraphEditor/Private/SNodePanel.cpp
--- a/Engine/Source/Editor/GraphEditor/Private/SNodePanel.cpp
+++ b/Engine/Source/Editor/GraphEditor/Private/SNodePanel.cpp
@@ -230,15 +230,38 @@ namespace NodePanelDefs
 };
 
 
-void SNodePanel::SNode::FNodeSlot::Construct(const FChildren& SlotOwner, FSlotArguments&& InArgs)
+SLATE_IMPLEMENT_WIDGET(SNodePanel)
+void SNodePanel::PrivateRegisterAttributes(FSlateAttributeInitializer& AttributeInitializer)
 {
-	TSlotBase<FNodeSlot>::Construct(SlotOwner, MoveTemp(InArgs));
-	TAlignmentWidgetSlotMixin<FNodeSlot>::ConstructMixin(SlotOwner, MoveTemp(InArgs));
+	SLATE_ADD_MEMBER_ATTRIBUTE_DEFINITION(AttributeInitializer, IsEditableAttribute, EInvalidateWidgetReason::Paint);
+	SLATE_ADD_MEMBER_ATTRIBUTE_DEFINITION(AttributeInitializer, ViewOffsetAttribute, EInvalidateWidgetReason::Layout);
+}
 
-	if (InArgs._Padding.IsSet())
-	{
-		SlotPadding = MoveTemp(InArgs._Padding);
-	}
+
+// equivalent of SLATE_IMPLEMENT_WIDGET(SNodePanel::SNode)
+FSlateWidgetClassRegistration ClassRegistration__SNodePanel_SNode = SNodePanel::SNode::StaticWidgetClass();
+
+void SNodePanel::SNode::PrivateRegisterAttributes(FSlateAttributeInitializer& AttributeInitializer)
+{
+	FSlateWidgetSlotAttributeInitializer Initializer = SLATE_ADD_PANELCHILDREN_DEFINITION(AttributeInitializer, Children);
+	FNodeSlot::RegisterAttributes(Initializer);
+
+	SLATE_ADD_MEMBER_ATTRIBUTE_DEFINITION(AttributeInitializer, DesiredSizeScaleAttribute, EInvalidateWidgetReason::Prepass);
+}
+
+void SNodePanel::SNode::FNodeSlot::RegisterAttributes(FSlateWidgetSlotAttributeInitializer& AttributeInitializer)
+{
+	TBasicLayoutWidgetSlot::RegisterAttributes(AttributeInitializer);
+
+	SLATE_ADD_SLOT_ATTRIBUTE_DEFINITION_WITH_NAME(FNodeSlot, AttributeInitializer, "Slot.AllowScale", AllowScale, EInvalidateWidgetReason::Prepass | EInvalidateWidgetReason::Paint);
+	SLATE_ADD_SLOT_ATTRIBUTE_DEFINITION_WITH_NAME(FNodeSlot, AttributeInitializer, "Slot.Size", Size, EInvalidateWidgetReason::Layout);
+	SLATE_ADD_SLOT_ATTRIBUTE_DEFINITION_WITH_NAME(FNodeSlot, AttributeInitializer, "Slot.Offset", Offset, EInvalidateWidgetReason::Layout)
+		.UpdatePrerequisite("Slot.Size");
+}
+
+void SNodePanel::SNode::FNodeSlot::Construct(const FChildren& SlotOwner, FSlotArguments&& InArgs)
+{
+	TBasicLayoutWidgetSlot<FNodeSlot>::Construct(SlotOwner, MoveTemp(InArgs));
 
 PRAGMA_DISABLE_DEPRECATION_WARNINGS
 	if (InArgs._SlotOffset.IsSet())
@@ -249,7 +272,7 @@ PRAGMA_ENABLE_DEPRECATION_WARNINGS
 
 	if (InArgs._SlotOffset2f.IsSet())
 	{
-		Offset = MoveTemp(InArgs._SlotOffset2f);
+		Offset.Assign(*this, MoveTemp(InArgs._SlotOffset2f));
 	}
 
 PRAGMA_DISABLE_DEPRECATION_WARNINGS
@@ -261,11 +284,11 @@ PRAGMA_ENABLE_DEPRECATION_WARNINGS
 
 	if (InArgs._SlotSize2f.IsSet())
 	{
-		Size = MoveTemp(InArgs._SlotSize2f);
+		SetSlotSize(MoveTemp(InArgs._SlotSize2f));
 	}
 	if (InArgs._AllowScaling.IsSet())
 	{
-		AllowScale = MoveTemp(InArgs._AllowScaling);
+		AllowScale.Assign(*this, MoveTemp(InArgs._AllowScaling));
 	}
 }
 
@@ -300,9 +323,12 @@ TArray<SNodePanel::SNode::DiffHighlightInfo> SNodePanel::SNode::GetDiffHighlight
 }
 
 SNodePanel::SNodePanel()
-	: Children(this)
+	: ZoomLevel(0)
+	, Children(this)
 	, VisibleChildren(this)
 	, NodeSnappingManager(nullptr)
+	, IsEditableAttribute(*this, true)
+	, ViewOffsetAttribute(*this, FVector2f::ZeroVector)
 {
 	bHasRelativeLayoutScale = true;
 }
@@ -314,27 +340,28 @@ void SNodePanel::OnArrangeChildren(const FGeometry& AllottedGeometry, FArrangedC
 
 void SNodePanel::ArrangeChildNodes(const FGeometry& AllottedGeometry, FArrangedChildren& ArrangedChildren) const
 {
-	const TSlotlessChildren<SNode>& ChildrenToArrange = ArrangedChildren.Accepts(EVisibility::Hidden) ? Children : VisibleChildren;
+	FChildren* ChildrenToArrange = ArrangedChildren.Accepts(EVisibility::Hidden) ? (FChildren*)&Children : (FChildren*)&VisibleChildren;
+
 	// First pass nodes
-	for (int32 ChildIndex = 0; ChildIndex < ChildrenToArrange.Num(); ++ChildIndex)
+	ChildrenToArrange->ForEachWidget([&, this](SWidget& ChildWidget)
 	{
-		const TSharedRef<SNode>& SomeChild = ChildrenToArrange[ChildIndex];
-		if (!SomeChild->RequiresSecondPassLayout())
+		SNode& SomeChild = static_cast<SNode&>(ChildWidget);
+		if (!SomeChild.RequiresSecondPassLayout())
 		{
-			ArrangedChildren.AddWidget(AllottedGeometry.MakeChild(SomeChild, SomeChild->GetPosition2f() - ViewOffset, SomeChild->GetDesiredSize(), GetZoomAmount()));
+			ArrangedChildren.AddWidget(AllottedGeometry.MakeChild(SomeChild.AsShared(), SomeChild.GetPosition2f() - GetViewOffset(), SomeChild.GetDesiredSize(), GetZoomAmount()));
 		}
-	}
+	});
 
 	// Second pass nodes
-	for (int32 ChildIndex = 0; ChildIndex < ChildrenToArrange.Num(); ++ChildIndex)
+	ChildrenToArrange->ForEachWidget([&, this](SWidget& ChildWidget)
 	{
-		const TSharedRef<SNode>& SomeChild = ChildrenToArrange[ChildIndex];
-		if (SomeChild->RequiresSecondPassLayout())
+		SNode& SomeChild = static_cast<SNode&>(ChildWidget);
+		if (SomeChild.RequiresSecondPassLayout())
 		{
-			SomeChild->PerformSecondPassLayout(NodeToWidgetLookup);
-			ArrangedChildren.AddWidget(AllottedGeometry.MakeChild(SomeChild, SomeChild->GetPosition2f() - ViewOffset, SomeChild->GetDesiredSize(), GetZoomAmount()));
+			SomeChild.PerformSecondPassLayout(NodeToWidgetLookup);
+			ArrangedChildren.AddWidget(AllottedGeometry.MakeChild(SomeChild.AsShared(), SomeChild.GetPosition2f() - GetViewOffset(), SomeChild.GetDesiredSize(), GetZoomAmount()));
 		}
-	}
+	});
 }
 
 FVector2D SNodePanel::ComputeDesiredSize( float ) const
@@ -350,6 +377,16 @@ FChildren* SNodePanel::GetChildren()
 }
 
 FChildren* SNodePanel::GetAllChildren()
+{
+	return &VisibleChildren;
+}
+
+FChildren* SNodePanel::GetVisibleChildren()
+{
+	return &VisibleChildren;
+}
+
+FChildren* SNodePanel::GetManagedChildren()
 {
 	return &Children;
 }
@@ -378,7 +415,7 @@ FSlateColor SNodePanel::GetZoomTextColorAndOpacity() const
 
 UE::Slate::FDeprecateVector2DResult SNodePanel::GetViewOffset() const
 {
-	return ViewOffset;
+	return ViewOffsetAttribute.Get();
 }
 
 bool SNodePanel::GetZoomTargetRect(FVector2D& TopLeft, FVector2D& BottomRight) const
@@ -409,7 +446,7 @@ void SNodePanel::Construct()
 	PreviousZoomLevel = ZoomLevels->GetDefaultZoomLevel();
 	PostChangedZoom();
 
-	ViewOffset = FVector2f::ZeroVector;
+	SetViewOffset(FVector2f::ZeroVector);
 	TotalMouseDelta = 0;
 	TotalMouseDeltaXY = 0;
 	bDeferredZoomToSelection = false;
@@ -429,7 +466,7 @@ void SNodePanel::Construct()
 
 	bIsPanning = false;
 	bIsZoomingWithTrackpad = false;
-	IsEditable.Set(true);
+	SetIsEditable(true);
 
 	ZoomLevelFade = FCurveSequence( 0.0f, 1.0f );
 	ZoomLevelFade.Play( this->AsShared() );
@@ -442,7 +479,7 @@ void SNodePanel::Construct()
 	DeferredPanPosition = FVector2f::ZeroVector;
 	bRequestDeferredPan = false;
 
-	OldViewOffset = ViewOffset;
+	OldViewOffset = GetViewOffset();
 	OldZoomAmount = GetZoomAmount();
 	ZoomStartOffset = FVector2f::ZeroVector;
 	TotalGestureMagnify = 0.0f;
@@ -450,6 +487,8 @@ void SNodePanel::Construct()
 	ScopedTransactionPtr.Reset();
 
 	bVisualUpdatePending = false;
+
+	SelectionManager.OnSelectionChanged.BindRaw(this, &SNodePanel::OnSelectionChanged);
 }
 
 FVector2f SNodePanel::ComputeEdgePanAmount(const FGeometry& MyGeometry, const FVector2f& TargetPosition)
@@ -497,7 +536,7 @@ void SNodePanel::UpdateViewOffset(const FGeometry& MyGeometry, const FVector2f&
 		return;
 	}
 	const FVector2f PanAmount = ComputeEdgePanAmount( MyGeometry, TargetPosition ) / GetZoomAmount();
-	ViewOffset += PanAmount;
+	SetViewOffset(GetViewOffset() + PanAmount);
 }
 
 void SNodePanel::RequestDeferredPan(const UE::Slate::FDeprecateVector2DParameter& UpdatePosition)
@@ -569,7 +608,7 @@ EActiveTimerReturnType SNodePanel::HandleZoomToFit(double InCurrentTime, float I
 	if (bDoneZooming && bDoneScrolling)
 	{
 		// One final push to make sure we centered in the end
-		ViewOffset = DesiredViewCenter - ( 0.5f * CachedGeometry.GetLocalSize() / GetZoomAmount() );
+		SetViewOffset(DesiredViewCenter - ( 0.5f * CachedGeometry.GetLocalSize() / GetZoomAmount() ));
 		
 		// Reset ZoomPadding
 		ZoomPadding = NodePanelDefs::DefaultZoomPadding;
@@ -652,15 +691,20 @@ void SNodePanel::Tick( const FGeometry& AllottedGeometry, const double InCurrent
 
 	// Reset the current bookmark if the location and/or zoom level has been changed.
 	const float CurZoomAmount = GetZoomAmount();
-	if (CurrentBookmarkGuid.IsValid() && (OldViewOffset != ViewOffset || OldZoomAmount != CurZoomAmount))
+	if (CurrentBookmarkGuid.IsValid() && (OldViewOffset != GetViewOffset() || OldZoomAmount != CurZoomAmount))
 	{
 		CurrentBookmarkGuid.Invalidate();
 	}
 
 	OldZoomAmount = CurZoomAmount;
-	OldViewOffset = ViewOffset;
+	OldViewOffset = GetViewOffset();
 
 	SPanel::Tick( AllottedGeometry, InCurrentTime, InDeltaTime );
+
+	if (!bCanMoveToTargetObjectThisFrame)
+	{
+		Invalidate(EInvalidateWidgetReason::Paint);
+	}
 }
 
 // The system calls this method to notify the widget that a mouse button was pressed within it. This event is bubbled.
@@ -701,6 +745,7 @@ FReply SNodePanel::OnMouseButtonDown( const FGeometry& MyGeometry, const FPointe
 		}
 		CancelZoomToFit();
 
+		Invalidate(EInvalidateWidgetReason::Paint);
 		return ReplyState;
 	};
 
@@ -756,6 +801,7 @@ FReply SNodePanel::OnMouseButtonDown( const FGeometry& MyGeometry, const FPointe
 			}
 		}
 
+		Invalidate(EInvalidateWidgetReason::Paint);
 		return ReplyState;
 	}
 	else if (bIsRightMouseButtonEffecting && ( GetDefault<UGraphEditorSettings>()->PanningMouseButton == EGraphPanningMouseButton::Right || GetDefault<UGraphEditorSettings>()->PanningMouseButton == EGraphPanningMouseButton::Both ) )
@@ -773,6 +819,7 @@ FReply SNodePanel::OnMouseButtonDown( const FGeometry& MyGeometry, const FPointe
 		{
 			if(SnapManager->BeginDragWithSnap(MyGeometry, MouseEvent))
 			{
+				Invalidate(EInvalidateWidgetReason::Paint);
 				return FReply::Handled().CaptureMouse( SharedThis(this) );
 			}
 		}
@@ -795,6 +842,7 @@ FReply SNodePanel::OnMouseButtonDown( const FGeometry& MyGeometry, const FPointe
 				{
 					// Track the node that we're dragging; we will move it in OnMouseMove.
 					this->OnBeginNodeInteraction(NodeWidgetUnderMouse, MousePositionInNode);
+					Invalidate(EInvalidateWidgetReason::Paint);
 					return FReply::Handled().CaptureMouse( SharedThis(this) );
 				}
 			}
@@ -807,6 +855,7 @@ FReply SNodePanel::OnMouseButtonDown( const FGeometry& MyGeometry, const FPointe
 		// If we're marquee selecting, then we're not clicking on a node!
 		NodeUnderMousePtr.Reset();
 
+		Invalidate(EInvalidateWidgetReason::Paint);
 		return FReply::Handled().CaptureMouse( SharedThis(this) );
 	}
 	else
@@ -827,6 +876,8 @@ FReply SNodePanel::OnMouseMove(const FGeometry& MyGeometry, const FPointerEvent&
 
 	if ( this->HasMouseCapture() )
 	{
+		Invalidate(EInvalidateWidgetReason::Paint);
+
 		const FVector2f CursorDelta = MouseEvent.GetCursorDelta();
 		// Track how much the mouse moved since the mouse down.
 		TotalMouseDelta += CursorDelta.Size();
@@ -875,7 +926,7 @@ FReply SNodePanel::OnMouseMove(const FGeometry& MyGeometry, const FPointerEvent&
 			PastePosition = PanelCoordToGraphCoord( 0.5f * MyGeometry.GetLocalSize() );
 
 			this->bIsPanning = true;
-			ViewOffset -= CursorDelta / GetZoomAmount();
+			SetViewOffset(GetViewOffset() - CursorDelta / GetZoomAmount());
 
 			// Stop the zoom-to-fit in favor of user control
 			CancelZoomToFit();
@@ -895,7 +946,7 @@ FReply SNodePanel::OnMouseMove(const FGeometry& MyGeometry, const FPointerEvent&
 			PastePosition = PanelCoordToGraphCoord(0.5f * MyGeometry.Size);
 
 			this->bIsPanning = true;
-			ViewOffset -= CursorDelta / GetZoomAmount();
+			SetViewOffset(GetViewOffset() - CursorDelta / GetZoomAmount());
 
 			return ReplyState;
 		}
@@ -903,7 +954,7 @@ FReply SNodePanel::OnMouseMove(const FGeometry& MyGeometry, const FPointerEvent&
 		{
 			TSharedPtr<SNode> NodeBeingDragged = NodeUnderMousePtr.Pin();
 
-			if ( IsEditable.Get() )
+			if ( GetIsEditable() )
 			{
 				// Update the amount to pan panel
 				UpdateViewOffset(MyGeometry, MouseEvent.GetScreenSpacePosition());
@@ -933,7 +984,7 @@ FReply SNodePanel::OnMouseMove(const FGeometry& MyGeometry, const FPointerEvent&
 							// Perform movement in 3 passes:
 							
 							// 1. Gather all selected nodes positions and calculate new positions
-							struct FDefferedNodePosition 
+							struct FDefferedNodePosition
 							{ 
 								SNode*		Node; 
 								FVector2f	NewPosition; 
@@ -1145,6 +1196,8 @@ void SNodePanel::OnMouseCaptureLost(const FCaptureLostEvent& CaptureLostEvent)
 		FinalizeNodeMovements();
 	}
 
+	Invalidate(EInvalidateWidgetReason::Paint);
+
 	SPanel::OnMouseCaptureLost(CaptureLostEvent);
 }
 
@@ -1170,7 +1223,7 @@ FCursorReply SNodePanel::OnCursorQuery( const FGeometry& MyGeometry, const FPoin
 
 FReply SNodePanel::OnKeyDown( const FGeometry& MyGeometry, const FKeyEvent& InKeyEvent )
 {
-	if( IsEditable.Get() )
+	if( GetIsEditable())
 	{
 		LastKeyChordDetected = FInputChord(InKeyEvent.GetKey(), EModifierKey::FromBools(InKeyEvent.IsControlDown(), InKeyEvent.IsAltDown(), InKeyEvent.IsShiftDown(), InKeyEvent.IsCommandDown()));
 	}
@@ -1220,7 +1273,7 @@ FReply SNodePanel::OnTouchGesture( const FGeometry& MyGeometry, const FPointerEv
 		const bool bUseDirectionInvertedFromDevice = DirectionSetting == EScrollGestureDirection::Natural || (DirectionSetting == EScrollGestureDirection::UseSystemSetting && GestureEvent.IsDirectionInvertedFromDevice());
 
 		this->bIsPanning = true;
-		ViewOffset -= (bUseDirectionInvertedFromDevice == GestureEvent.IsDirectionInvertedFromDevice() ? GestureDelta : -GestureDelta) / GetZoomAmount();
+		SetViewOffset(GetViewOffset() - (bUseDirectionInvertedFromDevice == GestureEvent.IsDirectionInvertedFromDevice() ? GestureDelta : -GestureDelta) / GetZoomAmount());
 		
 		// Stop the zoom-to-fit in favor of user control
 		CancelZoomToFit();
@@ -1391,7 +1444,7 @@ bool SNodePanel::Contains(UObject* Node) const
 
 void SNodePanel::RestoreViewSettings(const UE::Slate::FDeprecateVector2DParameter& InViewOffset, float InZoomAmount, const FGuid& InBookmarkGuid)
 {
-	ViewOffset = InViewOffset;
+	SetViewOffset(InViewOffset);
 
 	if (InZoomAmount <= 0.0f)
 	{
@@ -1411,7 +1464,7 @@ void SNodePanel::RestoreViewSettings(const UE::Slate::FDeprecateVector2DParamete
 
 	// If we have been forced to a specific position, set the old values equal to the new ones.
 	// This is so our locked window isn't forced to update according to this movement.
-	OldViewOffset = ViewOffset;
+	OldViewOffset = GetViewOffset();
 	OldZoomAmount = GetZoomAmount();
 
 	// Update the current bookmark ID.
@@ -1452,8 +1505,8 @@ void SNodePanel::PaintBackgroundAsLines(const FSlateBrush* BackgroundImage, cons
 	
 	const float GridCellSize = NominalGridSize * ZoomFactor * Inflation;
 
-	const float GraphSpaceGridX0 = FancyMod(ViewOffset.X, Inflation * NominalGridSize * static_cast<float>(RulePeriod));
-	const float GraphSpaceGridY0 = FancyMod(ViewOffset.Y, Inflation * NominalGridSize * static_cast<float>(RulePeriod));
+	const float GraphSpaceGridX0 = FancyMod(GetViewOffset().X, Inflation * NominalGridSize * static_cast<float>(RulePeriod));
+	const float GraphSpaceGridY0 = FancyMod(GetViewOffset().Y, Inflation * NominalGridSize * static_cast<float>(RulePeriod));
 
 	float ImageOffsetX = GraphSpaceGridX0 * -ZoomFactor;
 	float ImageOffsetY = GraphSpaceGridY0 * -ZoomFactor;
@@ -1764,10 +1817,10 @@ bool SNodePanel::GetBoundsForNodes(bool bSelectionSetOnly, FVector2f& MinCorner,
 bool SNodePanel::ScrollToLocation(const FGeometry& MyGeometry, FVector2f DesiredCenterPosition, const float InDeltaTime)
 {
 	const FVector2f HalfOFScreenInGraphSpace = 0.5f * MyGeometry.GetLocalSize() / GetZoomAmount();
-	FVector2f CurrentPosition = ViewOffset + HalfOFScreenInGraphSpace;
+	FVector2f CurrentPosition = GetViewOffset() + HalfOFScreenInGraphSpace;
 
 	FVector2f NewPosition = FMath::Vector2DInterpTo(CurrentPosition, DesiredCenterPosition, InDeltaTime, 10.f);
-	ViewOffset = NewPosition - HalfOFScreenInGraphSpace;
+	SetViewOffset(NewPosition - HalfOFScreenInGraphSpace);
 
 	// If within 1 pixel of target, stop interpolating
 	return ((NewPosition - DesiredCenterPosition).SizeSquared() < 1.f);
@@ -1909,7 +1962,7 @@ void SNodePanel::ChangeZoomLevel(int32 ZoomLevelDelta, const FVector2f& WidgetSp
 	ZoomLevelFade.Play( this->AsShared() );
 
 	// Re-center the screen so that it feels like zooming around the cursor.
-	this->ViewOffset = PointToMaintainGraphSpace - WidgetSpaceZoomOrigin / GetZoomAmount();
+	this->SetViewOffset(PointToMaintainGraphSpace - WidgetSpaceZoomOrigin / GetZoomAmount());
 }
 
 void SNodePanel::SetZoomLevel(int32 InZoomLevel)
@@ -2016,6 +2069,7 @@ void SNodePanel::PostChangedZoom()
 
 	// Invalidate the current bookmark.
 	CurrentBookmarkGuid.Invalidate();
+	Invalidate(EInvalidateWidgetReason::Prepass);
 }
 
 void SNodePanel::RequestZoomToFit()
@@ -2038,7 +2092,12 @@ void SNodePanel::CancelZoomToFit()
 	}
 }
 
+void SNodePanel::OnSelectionChanged(const TSet<UObject*>& Objects)
+{
+	Invalidate(EInvalidateWidgetReason::Paint);
+}
+
 bool SNodePanel::HasMoved() const
 {
-	return (!FMath::IsNearlyEqual(GetZoomAmount(), OldZoomAmount) || !ViewOffset.Equals(OldViewOffset, SMALL_NUMBER));
+	return (!FMath::IsNearlyEqual(GetZoomAmount(), OldZoomAmount) || !GetViewOffset().Equals(OldViewOffset, SMALL_NUMBER));
 }
diff --git a/Engine/Source/Editor/GraphEditor/Public/ConnectionDrawingPolicy.h b/Engine/Source/Editor/GraphEditor/Public/ConnectionDrawingPolicy.h
--- a/Engine/Source/Editor/GraphEditor/Public/ConnectionDrawingPolicy.h
+++ b/Engine/Source/Editor/GraphEditor/Public/ConnectionDrawingPolicy.h
@@ -129,13 +129,15 @@ protected:
 	float ZoomFactor; 
 	float HoverDeemphasisDarkFraction;
 	float SliceDeemphasisAlphaMultiplier;
-	const FSlateRect& ClippingRect;
+	const FSlateRect ClippingRect;
 	FSlateWindowElementList& DrawElementsList;
 	TMap< UEdGraphPin*, TSharedPtr<SGraphPin> > PinToPinWidgetMap;
 	TSet< FEdGraphPinReference > HoveredPins;
 	TMap<TSharedRef<SWidget>, FArrangedWidget>* PinGeometries;
 	double LastHoverTimeEvent;
+	UE_DEPRECATED(5.8, "Use AbsoluteMousePosition instead, no other change of code needed. The name of this variable was wrong as SetMousePosition was passed the result of 'AllottedGeometry.LocalToAbsolute'.")
 	FDeprecateSlateVector2D LocalMousePosition;
+	FVector2f AbsoluteMousePosition;
 
 	/** List of currently relinked connections. */
 	TArray<FRelinkConnection> RelinkConnections;
@@ -152,7 +154,9 @@ public:
 	// Update the drawing policy with the set of hovered pins (which can be empty)
 	UE_API void SetHoveredPins(const TSet< FEdGraphPinReference >& InHoveredPins, const TArray< TSharedPtr<SGraphPin> >& OverridePins, double HoverTime);
 
+	UE_DEPRECATED(5.8, "Use SetAbsoluteMousePosition instead, no other change of code needed.")
 	UE_API void SetMousePosition(const UE::Slate::FDeprecateVector2DParameter& InMousePos);
+	UE_API void SetAbsoluteMousePosition(FVector2f InAbsoluteMousePos);
 
 	// If a (valid) slice line has been set, we'll intersect any drawn splines via DrawConnection() with that line,
 	// and add any intersecting pin pairs to the public ConnectionsIntersectingSliceLine array. 
@@ -217,7 +221,19 @@ public:
 	UE_API virtual bool IsConnectionCulled( const FArrangedWidget& StartLink, const FArrangedWidget& EndLink ) const;
 
 	UE_API virtual TSharedPtr<IToolTip> GetConnectionToolTip(const SGraphPanel& GraphPanel, const FGraphSplineOverlapResult& OverlapData) const;
-	
+
+	/**
+	 * Used in MouseMove to check if the connections have graphically changed and need the GraphPanel to invalidate.
+	 * The default implementation checks for a difference in HoveredPins, or in difference in connection WireColor or WireThickness.
+	 */
+	UE_API virtual bool HaveConnectionsGraphicallyChanged(const SGraphPanel& InGraphPanel, const FVector2f& AbsoluteCursorPosition) const;
+
+	/**
+	 * Used in MouseMove to check if the connections have graphically changed and need the GraphPanel to invalidate.
+	 * The default implementation checks for a difference in HoveredPins, or in difference in connection WireColor or WireThickness.
+	 */
+	UE_API virtual bool HaveConnectionsGraphicallyChanged(const TSet<FEdGraphPinReference>& InNewlyHoveredPins) const;
+
 protected:
 	// Helper function used by Draw(). Called before DrawPinGeometries to populate PinToPinWidgetMap
 	UE_API virtual void BuildPinToPinWidgetMap(TMap<TSharedRef<SWidget>, FArrangedWidget>& InPinGeometries);
@@ -225,6 +241,34 @@ protected:
 	// Helper function used by Draw(). Iterates over the pin geometries, drawing connections between them
 	UE_API virtual void DrawPinGeometries(TMap<TSharedRef<SWidget>, FArrangedWidget>& InPinGeometries, FArrangedChildren& ArrangedNodes);
 
+	/**
+	 * Checks if the given spline overlaps with the Cursor.
+	 * @param AbsoluteCursorPosition Absolute position of the cursor
+	 * @param Start Start point of the Spline
+	 * @param End End point of the Spline
+	 * @param Params Parameters of the Spline
+	 * @param InOutOverlapResults Will hold the overlap result. Pass in previous results used to find the closest spline to the cursor
+	 * @param bOutSliceLineIntersectsSpline If not null, will return true if the SliceLine intersects this spline
+	 */
+	UE_API void CheckSplineConnectionOverlapWithCursor(
+		const FVector2f& AbsoluteCursorPosition,
+		const FVector2f& Start,
+		const FVector2f& End,
+		const FConnectionParams& Params,
+		/*out*/ FGraphSplineOverlapResult& InOutOverlapResults,
+		/*out*/ bool* bOutSliceLineIntersectsSpline = nullptr
+		) const;
+
+	struct FSimpleConnectionData
+	{
+		TSharedPtr<SGraphPin> StartPin;
+		TSharedRef<SGraphPin> EndPin;
+		FVector2f StartPoint;
+		FVector2f EndPoint;
+		FConnectionParams Params;
+	};
+	/** Keeps track of the last connections drawn to check if they would have changed on MouseMove */
+	TArray<FSimpleConnectionData> LastConnections;
 private:
 	// Note - these helpers are temporary, to be replaced entirely by the public API's 2f/float functions once the 2D/double deprecated functions can be removed
 	UE_API void DrawSplineWithArrow_DeprecationHelper(const UE::Slate::FDeprecateVector2DParameter& StartPoint, const UE::Slate::FDeprecateVector2DParameter& EndPoint, const FConnectionParams& Params);
diff --git a/Engine/Source/Editor/GraphEditor/Public/SGraphPanel.h b/Engine/Source/Editor/GraphEditor/Public/SGraphPanel.h
--- a/Engine/Source/Editor/GraphEditor/Public/SGraphPanel.h
+++ b/Engine/Source/Editor/GraphEditor/Public/SGraphPanel.h
@@ -87,8 +87,16 @@ struct FConnectionSlice
 		: Line(StartPosition) { }
 };
 
+
+template <>
+struct TWidgetTypeTraits<class SGraphPanel>
+{
+	static constexpr bool SupportsInvalidation() { return true; }
+};
+
 class SGraphPanel : public SNodePanel, public FGCObject
 {
+	SLATE_DECLARE_WIDGET_API(SGraphPanel, SNodePanel, UE_API)
 public:
 	DECLARE_DELEGATE_RetVal_OneParam(FActionMenuContent, FOnGetContextMenuFor, const FGraphContextMenuArguments& /*SpawnInfo*/)
 
@@ -192,6 +200,8 @@ PRAGMA_ENABLE_DEPRECATION_WARNINGS
 		SLATE_ARGUMENT(bool, AllowPanning)
 	SLATE_END_ARGS()
 
+	UE_API SGraphPanel();
+
 	/**
 	 * Construct a widget
 	 *
@@ -282,7 +292,7 @@ public:
 	UE_API void AttachGraphEvents(TSharedPtr<SGraphNode> CreatedSubNode);
 
 	/** Returns if this graph is editable */
-	bool IsGraphEditable() const { return IsEditable.Get(); }
+	bool IsGraphEditable() const { return GetIsEditable(); }
 
 	/** Attempt to retrieve the bounds for the specified node */
 	UE_DEPRECATED(5.6, "Use the version of the function accepting FVector2f; this Slate API no longer interfaces directly with double-precision scalars and vectors.")
@@ -327,6 +337,10 @@ protected:
 
 	/** Helper method to remove all nodes from the panel holding garbage / invalid pointers */
 	UE_API void RemoveAllNodesWithInvalidPointers();
+
+	/** Called when the graph selection changed */
+	UE_API virtual void OnSelectionChanged(const TSet<UObject*>& Objects) override;
+
 public:
 	/** Pin marked via shift-clicking */
 	TWeakPtr<SGraphPin> MarkedPin;
@@ -339,6 +353,23 @@ public:
 
 	const FGraphSplineOverlapResult& GetPreviousFrameSplineOverlap() const { return PreviousFrameSplineOverlap; }
 
+protected:
+	void SetShowGraphStateOverlay(TAttribute<bool> InShowGraphStateOverlayAttribute) { ShowGraphStateOverlayAttribute.Assign(*this, MoveTemp(InShowGraphStateOverlayAttribute)); }
+	bool GetShowGraphStateOverlay() const { return ShowGraphStateOverlayAttribute.Get(); }
+	TSlateAttributeRef<bool> GetShowGraphStateOverlayAttribute() const { return TSlateAttributeRef<bool>{SharedThis(this), ShowGraphStateOverlayAttribute}; }
+
+	void SetShouldDrawBackground(TAttribute<bool> InShouldDrawBackgroundAttribute) { ShouldDrawBackgroundAttribute.Assign(*this, MoveTemp(InShouldDrawBackgroundAttribute)); }
+	bool GetShouldDrawBackground() const { return ShouldDrawBackgroundAttribute.Get(); }
+	TSlateAttributeRef<bool> GetShouldDrawBackgroundAttribute() const { return TSlateAttributeRef<bool>{SharedThis(this), ShouldDrawBackgroundAttribute}; }
+
+	void SetShouldDrawSurroundingShadow(TAttribute<bool> InShouldDrawSurroundingShadowAttribute) { ShouldDrawSurroundingShadowAttribute.Assign(*this, MoveTemp(InShouldDrawSurroundingShadowAttribute)); }
+	bool GetShouldDrawSurroundingShadow() const { return ShouldDrawSurroundingShadowAttribute.Get(); }
+	TSlateAttributeRef<bool> GetShouldDrawSurroundingShadowAttribute() const { return TSlateAttributeRef<bool>{SharedThis(this), ShouldDrawSurroundingShadowAttribute}; }
+
+	void SetDisplayAsReadOnly(TAttribute<bool> InDisplayAsReadOnlyAttribute) { DisplayAsReadOnlyAttribute.Assign(*this, MoveTemp(InDisplayAsReadOnlyAttribute)); }
+	bool GetDisplayAsReadOnly() const { return DisplayAsReadOnlyAttribute.Get(); }
+	TSlateAttributeRef<bool> GetDisplayAsReadOnlyAttribute() const { return TSlateAttributeRef<bool>{SharedThis(this), DisplayAsReadOnlyAttribute}; }
+
 private:
 
 	/** A map of guid -> graph nodes */
@@ -422,21 +453,33 @@ protected:
 
 	/** Called when the graph itself is clicked */
 	SGraphEditor::FOnMouseButtonDown OnClicked;
-	
+
+	/** Called when the graph selection changed */
+	SGraphEditor::FOnSelectionChanged OnSelectionChangedDelegate;
+
 	/** Whether to draw the overlay indicating we're in PIE */
 	bool bShowPIENotification;
 
 	/** Whether to draw decorations for graph state (PIE / ReadOnly etc.) */
-	TAttribute<bool> ShowGraphStateOverlay;
+	UE_DEPRECATED(5.8, "Use SetShowGraphStateOverlay / GetShowGraphStateOverlay / GetShowGraphStateOverlayAttribute")
+	TSlateDeprecatedTAttribute<bool> ShowGraphStateOverlay;
 
 	/** Whether to allow slicing of connections in this graph panel */
 	TAttribute<bool> AllowConnectionSlicing;
 
 	/** Whether this panel should be drawing a background */
-	TAttribute<bool> ShouldDrawBackground;
+	UE_DEPRECATED(5.8, "Use SetShouldDrawBackground / GetShouldDrawBackground / GetShouldDrawBackgroundAttribute")
+	TSlateDeprecatedTAttribute<bool> ShouldDrawBackground;
 
 	/** Whether this panel should be drawing a sunken shadow on its borders */
-	TAttribute<bool> ShouldDrawSurroundingShadow;
+	UE_DEPRECATED(5.8, "Use SetShouldDrawSurroundingShadow / GetShouldDrawSurroundingShadow / GetShouldDrawSurroundingShadowAttribute")
+	TSlateDeprecatedTAttribute<bool> ShouldDrawSurroundingShadow;
+
+	/**
+	 * The last created policy which is cached to be used on MouseMove to invalidate when hovering connections.
+	 * This must not be used for drawing outside of Paint
+	 */
+	mutable TUniquePtr<FConnectionDrawingPolicy> ConnectionDrawingPolicy;
 
 private:
 	/** Set of nodes selected by the user, tracked while a visual update is pending */
@@ -445,8 +488,17 @@ private:
 	/** Set of user-added nodes for the panel, tracked while a visual update is pending */
 	TSet<const class UEdGraphNode*> UserAddedNodes;
 
+	/** Whether to draw decorations for graph state (PIE / ReadOnly etc.) */
+	TSlateAttribute<bool> ShowGraphStateOverlayAttribute;
+
+	/** Whether this panel should be drawing a background */
+	TSlateAttribute<bool> ShouldDrawBackgroundAttribute;
+
+	/** Whether this panel should be drawing a sunken shadow on its borders */
+	TSlateAttribute<bool> ShouldDrawSurroundingShadowAttribute;
+
 	/** Should the graph display all nodes in a read-only state (grayed)? This does not affect functionality of using them (IsEditable) */
-	TAttribute<bool> DisplayAsReadOnly;
+	TSlateAttribute<bool> DisplayAsReadOnlyAttribute;
 
 	FOnGraphChanged::FDelegate MyRegisteredGraphChangedDelegate;
 	FDelegateHandle            MyRegisteredGraphChangedDelegateHandle;
diff --git a/Engine/Source/Editor/GraphEditor/Public/SNodePanel.h b/Engine/Source/Editor/GraphEditor/Public/SNodePanel.h
--- a/Engine/Source/Editor/GraphEditor/Public/SNodePanel.h
+++ b/Engine/Source/Editor/GraphEditor/Public/SNodePanel.h
@@ -311,16 +311,24 @@ namespace ENodeZone
 	};
 }
 
+template <>
+struct TWidgetTypeTraits<class SNodePanel>
+{
+	static constexpr bool SupportsInvalidation() { return true; }
+};
+
 class SNodePanel : public SPanel
 {
+	SLATE_DECLARE_WIDGET_API(SNodePanel, SPanel, UE_API)
 public:
 
 	class SNode : public SPanel
 	{
+		SLATE_DECLARE_WIDGET_API(SNode, SPanel, UE_API)
 	public:
 
 		/** A slot that support alignment of content and padding and z-order */
-		class FNodeSlot : public TSlotBase<FNodeSlot>, public TAlignmentWidgetSlotMixin<FNodeSlot>
+		class FNodeSlot : public TBasicLayoutWidgetSlot<FNodeSlot>
 		{
 		public:
 			friend SNode;
@@ -330,16 +338,16 @@ public:
 			{ }
 
 			FNodeSlot(ENodeZone::Type InZone)
-				: TSlotBase<FNodeSlot>()
-				, TAlignmentWidgetSlotMixin<FNodeSlot>(HAlign_Fill, VAlign_Fill)
+				: TBasicLayoutWidgetSlot<FNodeSlot>()
 				, Zone(InZone)
-				, SlotPadding(0.0f)
-				, Offset(FVector2f::ZeroVector)
-				, AllowScale(true)
-			{ }
+				, Offset(*this, FVector2f::ZeroVector)
+				, Size(*this, FVector2f::ZeroVector)
+				, AllowScale(*this, true)
+				, bSizeAttrIsSet(false)
+			{
+			}
 
-			SLATE_SLOT_BEGIN_ARGS_OneMixin(FNodeSlot, TSlotBase<FNodeSlot>, TAlignmentWidgetSlotMixin<FNodeSlot>)
-				SLATE_ATTRIBUTE(FMargin, Padding)
+			SLATE_SLOT_BEGIN_ARGS(FNodeSlot, TBasicLayoutWidgetSlot<FNodeSlot>)
 				SLATE_ATTRIBUTE_DEPRECATED(FVector2D, SlotOffset, 5.6, "Use SlotOffset2f; this Slate API no longer interfaces directly with double-precision scalars and vectors.")
 				SLATE_ATTRIBUTE(FVector2f, SlotOffset2f)
 				SLATE_ATTRIBUTE_DEPRECATED(FVector2D, SlotSize, 5.6, "UseSlotSize2f; this Slate API no longer interfaces directly with double-precision scalars and vectors.")
@@ -348,6 +356,7 @@ public:
 			SLATE_SLOT_END_ARGS()
 
 			UE_API void Construct(const FChildren& SlotOwner, FSlotArguments&& InArgs);
+			static UE_API void RegisterAttributes(FSlateWidgetSlotAttributeInitializer& AttributeInitializer);
 
 		public:
 			ENodeZone::Type GetZoneType() const
@@ -355,30 +364,20 @@ public:
 				return Zone;
 			}
 
-			void SetPadding(TAttribute<FMargin> InPadding)
-			{
-				SlotPadding = MoveTemp(InPadding);
-			}
-
-			FMargin GetPadding() const
-			{
-				return SlotPadding.Get();
-			}
-
 			UE_DEPRECATED(5.6, "This exists as a temporary bridge for deprecation of SlotSize Attribute. Please use the FVector2f version.")
 			void SetSlotOffset(TAttribute<FVector2D> InOffset)
 			{
 			PRAGMA_DISABLE_DEPRECATION_WARNINGS
 				Offset_Deprecated = MoveTemp(InOffset);
-				Offset.BindLambda([this]()
+				SetSlotOffset(TAttribute<FVector2f>::CreateLambda([this]()
 				{
 					return UE::Slate::CastToVector2f(Offset_Deprecated.Get());
-				});
+				}));
 			PRAGMA_ENABLE_DEPRECATION_WARNINGS
 			}
 			void SetSlotOffset(TAttribute<FVector2f> InOffset)
 			{
-				Offset = MoveTemp(InOffset);
+				Offset.Assign(*this, MoveTemp(InOffset));
 			}
 
 			UE::Slate::FDeprecateVector2DResult GetSlotOffset() const
@@ -391,15 +390,16 @@ public:
 			{
 			PRAGMA_DISABLE_DEPRECATION_WARNINGS
 				Size_Deprecated = MoveTemp(InSize);
-				Size.BindLambda([this]()
+				SetSlotSize(TAttribute<FVector2f>::CreateLambda([this]()
 				{
 					return UE::Slate::CastToVector2f(Size_Deprecated.Get());
-				});
+				}));
 			PRAGMA_ENABLE_DEPRECATION_WARNINGS
 			}
 			void SetSlotSize(TAttribute<FVector2f> InSize)
 			{
-				Size = MoveTemp(InSize);
+				bSizeAttrIsSet = InSize.IsSet();
+				Size.Assign(*this, MoveTemp(InSize));
 			}
 
 			UE::Slate::FDeprecateVector2DResult GetSlotSize() const
@@ -409,7 +409,7 @@ public:
 
 			void SetAllowScaling(TAttribute<bool> InAllowScaling)
 			{
-				AllowScale = MoveTemp(InAllowScaling);
+				AllowScale.Assign(*this, MoveTemp(InAllowScaling));
 			}
 
 			bool GetAllowScaling() const
@@ -420,14 +420,14 @@ public:
 		private:
 			/** The child widget contained in this slot. */
 			ENodeZone::Type Zone;
-			TAttribute<FMargin> SlotPadding;
-			TAttribute<FVector2f> Offset;
-			TAttribute<FVector2f> Size;
+			TSlateSlotAttribute<FVector2f> Offset;
+			TSlateSlotAttribute<FVector2f> Size;
 			UE_DEPRECATED(5.6, "This exists as a temporary bridge for deprecation of SlotSize Attribute.")
 			TAttribute<FVector2D> Size_Deprecated;
 			UE_DEPRECATED(5.6, "This exists as a temporary bridge for deprecation of SlotSize Attribute.")
 			TAttribute<FVector2D> Offset_Deprecated;
-			TAttribute<bool> AllowScale;
+			TSlateSlotAttribute<bool> AllowScale;
+			bool bSizeAttrIsSet;
 		};
 
 		typedef TSet<TWeakPtr<SNodePanel::SNode>> FNodeSet;
@@ -476,7 +476,7 @@ public:
 					const EVisibility ChildVisibility = CenterZone.GetWidget()->GetVisibility();
 					if( ChildVisibility != EVisibility::Collapsed )
 					{
-						return ( FVector2D(CenterZone.GetWidget()->GetDesiredSize() + CenterZone.SlotPadding.Get().GetDesiredSize()) ) * FVector2D(DesiredSizeScale.Get());
+						return ( FVector2D(CenterZone.GetWidget()->GetDesiredSize() + CenterZone.GetPadding().GetDesiredSize()) ) * FVector2D(GetDesiredSizeScale());
 					}
 				}
 			}
@@ -507,11 +507,11 @@ public:
 				const EVisibility ChildVisibility = CurChild.GetWidget()->GetVisibility();
 				if ( ArrangedChildren.Accepts(ChildVisibility) )
 				{
-					const FMargin SlotPadding(CurChild.SlotPadding.Get());
+					const FMargin SlotPadding(CurChild.GetPadding());
 					// If this child is not allowed to scale, its scale relative to its parent should undo the parent widget's scaling.
 					FVector2f Size;
 
-					if( CurChild.Size.IsSet() )
+					if( CurChild.bSizeAttrIsSet )
 					{
 						Size = CurChild.Size.Get();
 					}
@@ -753,15 +753,30 @@ public:
 
 	protected:
 		SNode()
-		: DesiredSizeScale(FVector2f::UnitVector)
+		: DesiredSizeScaleAttribute(*this, FVector2f::UnitVector)
 		, Children(this)
 		{
 			bHasRelativeLayoutScale = true;
 		}
 
+	protected:
+		void SetDesiredSizeScale(TAttribute<FDeprecateSlateVector2D> InDesiredSizeScale)
+		{
+			DesiredSizeScaleAttribute.Assign(*this, MoveTemp(InDesiredSizeScale));
+		}
+		TSlateAttributeRef<FDeprecateSlateVector2D> GetDesiredSizeScaleAttribute() const
+		{
+			return TSlateAttributeRef<FDeprecateSlateVector2D>{SharedThis(this), DesiredSizeScaleAttribute};
+		}
+		FDeprecateSlateVector2D GetDesiredSizeScale() const
+		{
+			return DesiredSizeScaleAttribute.Get();
+		}
+
 	protected:
 
-		TAttribute<FDeprecateSlateVector2D> DesiredSizeScale;
+		UE_DEPRECATED(5.8, "Use SetDesiredSizeScale / GetDesiredSizeScale / GetDesiredSizeScaleAttribute")
+		TSlateDeprecatedTAttribute<FDeprecateSlateVector2D> DesiredSizeScale;
 
 		// SBorder Begin
 		UE_DEPRECATED(5.8, "The attribute was unused and can be removed.")
@@ -795,6 +810,9 @@ public:
 		TSlateDeprecatedTAttribute<FSlateColor> ForegroundColor;
 		// SPanel End
 
+	private:
+		TSlateAttribute<FDeprecateSlateVector2D> DesiredSizeScaleAttribute;
+
 	private:
 
 		TSharedPtr<SNodePanel> GetParentPanel() const
@@ -812,6 +830,8 @@ public:
 	// SPanel interface
 	UE_API virtual void OnArrangeChildren(const FGeometry& AllottedGeometry, FArrangedChildren& ArrangedChildren) const override;
 	UE_API virtual FVector2D ComputeDesiredSize(float) const override;
+
+	UE_DEPRECATED(5.8, "Use GetVisibleChildren for visible children, or GetManagedChildren to get all the nodes of this panel.")
 	UE_API virtual FChildren* GetChildren() override;
 	// End of SPanel interface
 
@@ -850,7 +870,7 @@ public:
 	UE_API UE::Slate::FDeprecateVector2DResult GetViewOffset() const;
 
 	/** Sets the current view offset */
-	void SetViewOffset(const UE::Slate::FDeprecateVector2DResult& InViewOffset) { ViewOffset = InViewOffset; }
+	void SetViewOffset(const UE::Slate::FDeprecateVector2DResult& InViewOffset) { ViewOffsetAttribute.Assign(*this, InViewOffset); }
 
 	/** 
 	 * when a panel is scrolling/zooming to a target, this can be called to get it's destination
@@ -911,8 +931,13 @@ public:
 	/** Returns if the panel has been panned or zoomed since the last update */
 	UE_API bool HasMoved() const;
 
-	/** Returns all the panel children rather than only visible */
-	UE_API FChildren* GetAllChildren();
+	/** Returns only the visible children as this is the function called by the invalidation system */
+	UE_DEPRECATED(5.8, "Use GetVisibleChildren for visible children, or GetManagedChildren to get all the nodes of this panel. GetAllChildren needs to return only the VisibleChildren to work with invalidation.")
+	UE_API virtual FChildren* GetAllChildren() override;
+
+	/** Returns all the visible panel children */
+	UE_API FChildren* GetVisibleChildren();
+	UE_API FChildren* GetManagedChildren();
 
 protected:
 	/** Initialize members */
@@ -1061,6 +1086,9 @@ protected:
 	// Cancels any active zoom-to-fit action
 	UE_API void CancelZoomToFit();
 
+	/** Called when the graph selection changed */
+	UE_API virtual void OnSelectionChanged(const TSet<UObject*>& Objects);
+
 public:
 	
 	// Sets the zoom levels container
@@ -1082,13 +1110,27 @@ public:
 
 	void SetCurrentHoveredNode(UObject* InNode) { CurrentHoveredNode = InNode; }
 	UObject* GetCurrentHoveredNode() { return CurrentHoveredNode; }
-	
+
 protected:
-	
+	void SetIsEditable(TAttribute<bool> InIsEditableAttribute) { IsEditableAttribute.Assign(*this, MoveTemp(InIsEditableAttribute)); }
+	bool GetIsEditable() const { return IsEditableAttribute.Get(); }
+	TSlateAttributeRef<bool> GetIsEditableAttribute() const
+	{
+		return TSlateAttributeRef<bool>{SharedThis(this), IsEditableAttribute};
+	}
+
+	TSlateAttributeRef<FVector2f> GetViewOffsetAttribute() const
+	{
+		return TSlateAttributeRef<FVector2f>{SharedThis(this), ViewOffsetAttribute};
+	}
+
+protected:
+
 	// The interface for mapping ZoomLevel values to actual node scaling values
 	TSharedPtr<FZoomLevelsContainer> ZoomLevels;
 
 	/** The position within the graph at which the user is looking */
+	UE_DEPRECATED(5.8, "Use SetViewOffset / GetViewOffset / GetViewOffsetAttribute")
 	FDeprecateSlateVector2D ViewOffset;
 
 	/** The position within the graph at which the user was looking last tick */
@@ -1110,7 +1152,7 @@ protected:
 	bool bIsZoomingWithTrackpad;
 
 	/** The graph node widgets owned by this panel */
-	TSlotlessChildren<SNode> Children;
+	TUnmanagedSlotlessChildren<SNode> Children;
 	TSlotlessChildren<SNode> VisibleChildren;
 
 	/** The node that the user is dragging. Null when they are not dragging a node. */
@@ -1142,7 +1184,8 @@ protected:
 	FMarqueeOperation Marquee;
 
 	/** Is the graph editable (can nodes be moved, etc...)? */
-	TAttribute<bool> IsEditable;
+	UE_DEPRECATED(5.8, "Use SetIsEditable / GetIsEditable / GetIsEditableAttribute")
+	TSlateDeprecatedTAttribute<bool> IsEditable;
 
 	/** Given a node, find the corresponding widget */
 	TMap< UObject*, TSharedRef<SNode> > NodeToWidgetLookup;
@@ -1229,6 +1272,12 @@ private:
 	UE_API EActiveTimerReturnType HandleZoomToFit(double InCurrentTime, float InDeltaTime);
 
 private:
+	/** Is the graph editable (can nodes be moved, etc...)? */
+	TSlateAttribute<bool> IsEditableAttribute;
+
+	/** The position within the graph at which the user is looking */
+	TSlateAttribute<FVector2f> ViewOffsetAttribute;
+
 	/** The handle to the active timer */
 	TWeakPtr<FActiveTimerHandle> ActiveTimerHandle;
 
@@ -1240,4 +1289,10 @@ private:
 	FGuid CurrentBookmarkGuid;
 };
 
+template <>
+struct TWidgetTypeTraits<SNodePanel::SNode>
+{
+	static constexpr bool SupportsInvalidation() { return true; }
+};
+
 #undef UE_API
diff --git a/Engine/Source/Editor/MaterialEditor/Private/MaterialEditor.cpp b/Engine/Source/Editor/MaterialEditor/Private/MaterialEditor.cpp
--- a/Engine/Source/Editor/MaterialEditor/Private/MaterialEditor.cpp
+++ b/Engine/Source/Editor/MaterialEditor/Private/MaterialEditor.cpp
@@ -7826,7 +7826,7 @@ TSharedRef<SGraphEditor> FMaterialEditor::CreateGraphEditorWidget(TSharedRef<cla
 		.MaterialInfoList(&MaterialInfoList);
 
 	return SNew(SGraphEditor)
-		.GraphSupportsInvalidation(false)
+		.GraphSupportsInvalidation(true)
 		.AdditionalCommands(GraphEditorCommands)
 		.IsEditable(true)
 		.TitleBar(TitleBarWidget)
diff --git a/Engine/Source/Editor/WorldBrowser/Private/Tiles/SWorldComposition.cpp b/Engine/Source/Editor/WorldBrowser/Private/Tiles/SWorldComposition.cpp
--- a/Engine/Source/Editor/WorldBrowser/Private/Tiles/SWorldComposition.cpp
+++ b/Engine/Source/Editor/WorldBrowser/Private/Tiles/SWorldComposition.cpp
@@ -200,7 +200,7 @@ public:
 		if (bHasScrollByRequest)
 		{
 			bHasScrollByRequest = false;
-			ViewOffset += RequestedScrollByValue;
+			SetViewOffset(GetViewOffset() + RequestedScrollByValue);
 		}
 	}
 	
@@ -395,7 +395,7 @@ public:
 				PastePosition = PanelCoordToGraphCoord( 0.5f * MyGeometry.GetLocalSize() );
 
 				this->bIsPanning = true;
-				ViewOffset -= CursorDelta / GetZoomAmount();
+				SetViewOffset(GetViewOffset() - CursorDelta / GetZoomAmount());
 
 				return ReplyState;
 			}
@@ -403,7 +403,7 @@ public:
 			{
 				TSharedPtr<SNode> NodeBeingDragged = NodeUnderMousePtr.Pin();
 
-				if ( IsEditable.Get() )
+				if ( GetIsEditable() )
 				{
 					// Update the amount to pan panel
 					UpdateViewOffset(MyGeometry, MouseEvent.GetScreenSpacePosition());
diff --git a/Engine/Source/Runtime/Slate/Public/Framework/Docking/SDockingTabStack.h b/Engine/Source/Runtime/Slate/Public/Framework/Docking/SDockingTabStack.h
--- a/Engine/Source/Runtime/Slate/Public/Framework/Docking/SDockingTabStack.h
+++ b/Engine/Source/Runtime/Slate/Public/Framework/Docking/SDockingTabStack.h
@@ -22,8 +22,6 @@ class FWidgetPath;
 class SDockingTabWell;
 class SBorder;
 
-template<typename ChildType> class TSlotlessChildren;
-
 struct FDockingStackOptionalContent
 {
 	FDockingStackOptionalContent()
diff --git a/Engine/Source/Runtime/SlateCore/Public/Layout/Children.h b/Engine/Source/Runtime/SlateCore/Public/Layout/Children.h
--- a/Engine/Source/Runtime/SlateCore/Public/Layout/Children.h
+++ b/Engine/Source/Runtime/SlateCore/Public/Layout/Children.h
@@ -830,8 +830,10 @@ private:
  * In those cases, using TSlotlessChildren is convenient.
  *
  * TSlotlessChildren should not be used for general-purpose widgets.
+ *
+ * If bManagedByWidget is false, the children will not be attached or detached from the parent and no invalidation will be triggerred
  */
-template<typename ChildType>
+template<typename ChildType, bool bManagedByWidget = true>
 class TSlotlessChildren final : public FChildren
 {
 private:
@@ -861,12 +863,18 @@ private:
 	}
 
 public:
-	TSlotlessChildren(SWidget* InOwner, bool InbChangesInvalidatePrepass = true)
+	TSlotlessChildren(SWidget* InOwner, bool InbChangesInvalidatePrepass = true) requires (bManagedByWidget)
 		: FChildren(InOwner)
 		, bChangesInvalidatePrepass(InbChangesInvalidatePrepass)
 	{
 	}
 
+	TSlotlessChildren(SWidget* InOwner) requires (!bManagedByWidget)
+		: FChildren(InOwner)
+		, bChangesInvalidatePrepass(false)
+	{
+	}
+
 	TSlotlessChildren(std::nullptr_t, bool InbChangesInvalidatePrepass = true) = delete;
 
 	virtual int32 Num() const override
@@ -886,16 +894,22 @@ public:
 
 	int32 Add( const TSharedRef<ChildType>& Child )
 	{
-		if (bChangesInvalidatePrepass)
-		{ 
-			GetOwner().Invalidate(EInvalidateWidgetReason::ChildOrder);
+		if constexpr (bManagedByWidget)
+		{
+			if (bChangesInvalidatePrepass)
+			{
+				GetOwner().Invalidate(EInvalidateWidgetReason::ChildOrder);
+			}
 		}
 
 		int32 Index = Children.Add(Child);
 
-		if (Child != SNullWidget::NullWidget)
+		if constexpr (bManagedByWidget)
 		{
-			Child->AssignParentWidget(GetOwner().AsShared());
+			if (Child != SNullWidget::NullWidget)
+			{
+				Child->AssignParentWidget(GetOwner().AsShared());
+			}
 		}
 
 		return Index;
@@ -903,12 +917,15 @@ public:
 
 	void Reset(int32 NewSize = 0)
 	{
-		for (int32 ChildIndex = 0; ChildIndex < Children.Num(); ChildIndex++)
+		if constexpr (bManagedByWidget)
 		{
-			TSharedRef<SWidget> Child = GetChildAt(ChildIndex);
-			if (Child != SNullWidget::NullWidget)
+			for (int32 ChildIndex = 0; ChildIndex < Children.Num(); ChildIndex++)
 			{
-				Child->ConditionallyDetatchParentWidget(&GetOwner());
+				TSharedRef<SWidget> Child = GetChildAt(ChildIndex);
+				if (Child != SNullWidget::NullWidget)
+				{
+					Child->ConditionallyDetatchParentWidget(&GetOwner());
+				}
 			}
 		}
 
@@ -927,12 +944,15 @@ public:
 
 	void Empty(int32 Slack = 0)
 	{
-		for (int32 ChildIndex = 0; ChildIndex < Children.Num(); ChildIndex++)
+		if constexpr (bManagedByWidget)
 		{
-			TSharedRef<SWidget> Child = GetChildAt(ChildIndex);
-			if (Child != SNullWidget::NullWidget)
+			for (int32 ChildIndex = 0; ChildIndex < Children.Num(); ChildIndex++)
 			{
-				Child->ConditionallyDetatchParentWidget(&GetOwner());
+				TSharedRef<SWidget> Child = GetChildAt(ChildIndex);
+				if (Child != SNullWidget::NullWidget)
+				{
+					Child->ConditionallyDetatchParentWidget(&GetOwner());
+				}
 			}
 		}
 
@@ -951,23 +971,33 @@ public:
 
 	void Insert(const TSharedRef<ChildType>& Child, int32 Index)
 	{
-		if (bChangesInvalidatePrepass) 
+		if constexpr (bManagedByWidget)
 		{
-			GetOwner().Invalidate(EInvalidateWidgetReason::ChildOrder);
+			if (bChangesInvalidatePrepass)
+			{
+				GetOwner().Invalidate(EInvalidateWidgetReason::ChildOrder);
+			}
 		}
 
 		Children.Insert(Child, Index);
-		if (Child != SNullWidget::NullWidget)
+
+		if constexpr (bManagedByWidget)
 		{
-			Child->AssignParentWidget(GetOwner().AsShared());
+			if (Child != SNullWidget::NullWidget)
+			{
+				Child->AssignParentWidget(GetOwner().AsShared());
+			}
 		}
 	}
 
 	int32 Remove( const TSharedRef<ChildType>& Child )
 	{
-		if (Child != SNullWidget::NullWidget)
+		if constexpr (bManagedByWidget)
 		{
-			Child->ConditionallyDetatchParentWidget(&GetOwner());
+			if (Child != SNullWidget::NullWidget)
+			{
+				Child->ConditionallyDetatchParentWidget(&GetOwner());
+			}
 		}
 
 		return Children.Remove( Child );
@@ -975,10 +1005,13 @@ public:
 
 	void RemoveAt( int32 Index )
 	{
-		TSharedRef<SWidget> Child = GetChildAt(Index);
-		if (Child != SNullWidget::NullWidget)
+		if constexpr (bManagedByWidget)
 		{
-			Child->ConditionallyDetatchParentWidget(&GetOwner());
+			TSharedRef<SWidget> Child = GetChildAt(Index);
+			if (Child != SNullWidget::NullWidget)
+			{
+				Child->ConditionallyDetatchParentWidget(&GetOwner());
+			}
 		}
 
 		// Note: Child above ensures the instance we're removing from the array won't run its destructor until after it's fully removed from the array
@@ -1002,18 +1035,24 @@ public:
 	void Sort( const PREDICATE_CLASS& Predicate )
 	{
 		Children.Sort( Predicate );
-		if (bChangesInvalidatePrepass)
+		if constexpr (bManagedByWidget)
 		{
-			GetOwner().Invalidate(EInvalidateWidgetReason::ChildOrder);
+			if (bChangesInvalidatePrepass)
+			{
+				GetOwner().Invalidate(EInvalidateWidgetReason::ChildOrder);
+			}
 		}
 	}
 
 	void Swap( int32 IndexA, int32 IndexB )
 	{
 		Children.Swap(IndexA, IndexB);
-		if (bChangesInvalidatePrepass)
+		if constexpr (bManagedByWidget)
 		{
-			GetOwner().Invalidate(EInvalidateWidgetReason::ChildOrder);
+			if (bChangesInvalidatePrepass)
+			{
+				GetOwner().Invalidate(EInvalidateWidgetReason::ChildOrder);
+			}
 		}
 	}
 
@@ -1021,6 +1060,9 @@ private:
 	bool bChangesInvalidatePrepass;
 };
 
+template <typename ChildType>
+using TUnmanagedSlotlessChildren = TSlotlessChildren<ChildType, false>;
+
 
 /** Required to implement GetChildren() in a way that can dynamically return the currently active child. */
 template<typename SlotType, typename Attr>
