[Fix] Quick fix for `FConnectionDrawingPolicy::HaveConnectionsGraphicallyChanged` to exit early if the last connection pins are not valid anymore, which would otherwise lead to a crash when retrieving their PinObject.
This seem to only happen in certain graphs.

#Synced-CL 50897806

diff --git a/Engine/Source/Editor/GraphEditor/Private/ConnectionDrawingPolicy.cpp b/Engine/Source/Editor/GraphEditor/Private/ConnectionDrawingPolicy.cpp
--- a/Engine/Source/Editor/GraphEditor/Private/ConnectionDrawingPolicy.cpp
+++ b/Engine/Source/Editor/GraphEditor/Private/ConnectionDrawingPolicy.cpp
@@ -519,6 +519,10 @@ bool FConnectionDrawingPolicy::HaveConnectionsGraphicallyChanged(const SGraphPan
 
 	for (const FSimpleConnectionData& Connection : LastConnections)
 	{
+		if (!Connection.EndPin->GetParentWidget() || (Connection.StartPin && !Connection.StartPin->GetParentWidget()))
+		{
+			return true; // pins are not valid anymore, just redraw
+		}
 		FConnectionParams Params;
 		Params.AssociatedPin1 = Connection.EndPin->GetPinObj();
 		Params.AssociatedPin2 = Connection.StartPin ? Connection.StartPin->GetPinObj() : nullptr;
@@ -553,6 +557,10 @@ bool FConnectionDrawingPolicy::HaveConnectionsGraphicallyChanged(const TSet<FEdG
 	// (as DetermineWiringStyle often checks the content of HoveredPins to decide how to draw them)
 	for (const FSimpleConnectionData& Connection : LastConnections)
 	{
+		if (!Connection.EndPin->GetParentWidget() || (Connection.StartPin && !Connection.StartPin->GetParentWidget()))
+		{
+			return true; // pins are not valid anymore, just redraw
+		}
 		FConnectionParams Params;
 		const_cast<FConnectionDrawingPolicy*>(this)->DetermineWiringStyle(Connection.EndPin->GetPinObj(), Connection.StartPin ? Connection.StartPin->GetPinObj() : nullptr, /*inout*/ Params);
 		if (Connection.StartPin && Connection.StartPin->AreConnectionsFaded() && Connection.EndPin->AreConnectionsFaded() )
