Fix Material Graph widgets not updating when nodes are moved using arrow key presses.

#Synced-CL 51454850

diff --git a/Engine/Source/Editor/GraphEditor/Private/SGraphNode.cpp b/Engine/Source/Editor/GraphEditor/Private/SGraphNode.cpp
--- a/Engine/Source/Editor/GraphEditor/Private/SGraphNode.cpp
+++ b/Engine/Source/Editor/GraphEditor/Private/SGraphNode.cpp
@@ -590,7 +590,7 @@ void SGraphNode::MoveTo( const FVector2f& NewPosition, FNodeSet& NodeFilter, boo
 			NodeFilter.Add( SharedThis( this ) );
 			GraphNode->Modify(bMarkDirty);
 			GraphNode->SetPosition(NewPosition);
-			Invalidate(EInvalidateWidgetReason::Layout);
+			Invalidate(EInvalidateWidgetReason::RenderTransform);
 		}
 	}
 }
