Fix for `BlueprintConnectionDrawingPolicy` being given invalid graph nodes by using `FGraphPinHandle` instead of raw `UEdGraphPin` pointers (which can end up being stale), and by ensuring null `UEdGraphPins` are not passed to `DetermineWiringStyle` as some derived classes expect the first argument to never be null.

#Synced-CL 51747987

diff --git a/Engine/Source/Editor/GraphEditor/Private/ConnectionDrawingPolicy.cpp b/Engine/Source/Editor/GraphEditor/Private/ConnectionDrawingPolicy.cpp
--- a/Engine/Source/Editor/GraphEditor/Private/ConnectionDrawingPolicy.cpp
+++ b/Engine/Source/Editor/GraphEditor/Private/ConnectionDrawingPolicy.cpp
@@ -519,13 +519,15 @@ bool FConnectionDrawingPolicy::HaveConnectionsGraphicallyChanged(const SGraphPan
 
 	for (const FSimpleConnectionData& Connection : LastConnections)
 	{
-		if (!Connection.EndPin->GetParentWidget() || (Connection.StartPin && !Connection.StartPin->GetParentWidget()))
+		UEdGraphPin* EndPinObj = Connection.EndPinHandle.GetPinObj(InGraphPanel);
+		if (!EndPinObj)
 		{
-			return true; // pins are not valid anymore, just redraw
+			return true; // pin is not valid anymore, just redraw
 		}
+
 		FConnectionParams Params;
-		Params.AssociatedPin1 = Connection.EndPin->GetPinObj();
-		Params.AssociatedPin2 = Connection.StartPin ? Connection.StartPin->GetPinObj() : nullptr;
+		Params.AssociatedPin1 = EndPinObj;
+		Params.AssociatedPin2 = Connection.StartPinHandle.GetPinObj(InGraphPanel);
 		CheckSplineConnectionOverlapWithCursor(AbsoluteCursorPosition, Connection.StartPoint, Connection.EndPoint, Params, /*out*/ OverlapResult);
 	}
 
@@ -543,10 +545,10 @@ bool FConnectionDrawingPolicy::HaveConnectionsGraphicallyChanged(const SGraphPan
 		NewlyHoveredPins.Add(NewPin2Widget->GetPinObj());
 	}
 
-	return HaveConnectionsGraphicallyChanged(NewlyHoveredPins);
+	return HaveConnectionsGraphicallyChanged(InGraphPanel, NewlyHoveredPins);
 }
 
-bool FConnectionDrawingPolicy::HaveConnectionsGraphicallyChanged(const TSet<FEdGraphPinReference>& InNewlyHoveredPins) const
+bool FConnectionDrawingPolicy::HaveConnectionsGraphicallyChanged(const SGraphPanel& InGraphPanel, const TSet<FEdGraphPinReference>& InNewlyHoveredPins) const
 {
 	if ( HoveredPins.Num() != InNewlyHoveredPins.Num() || !HoveredPins.Difference(InNewlyHoveredPins).IsEmpty() )
 	{
@@ -557,13 +559,15 @@ bool FConnectionDrawingPolicy::HaveConnectionsGraphicallyChanged(const TSet<FEdG
 	// (as DetermineWiringStyle often checks the content of HoveredPins to decide how to draw them)
 	for (const FSimpleConnectionData& Connection : LastConnections)
 	{
-		if (!Connection.EndPin->GetParentWidget() || (Connection.StartPin && !Connection.StartPin->GetParentWidget()))
+		UEdGraphPin* EndPinObj = Connection.EndPinHandle.GetPinObj(InGraphPanel);
+		if (!EndPinObj)
 		{
-			return true; // pins are not valid anymore, just redraw
+			return true; // pin is not valid anymore, just redraw
 		}
+
 		FConnectionParams Params;
-		const_cast<FConnectionDrawingPolicy*>(this)->DetermineWiringStyle(Connection.EndPin->GetPinObj(), Connection.StartPin ? Connection.StartPin->GetPinObj() : nullptr, /*inout*/ Params);
-		if (Connection.StartPin && Connection.StartPin->AreConnectionsFaded() && Connection.EndPin->AreConnectionsFaded() )
+		const_cast<FConnectionDrawingPolicy*>(this)->DetermineWiringStyle(EndPinObj, Connection.StartPinHandle.GetPinObj(InGraphPanel), /*inout*/ Params);
+		if (Connection.bAreConnectionsFaded)
 		{
 			Params.WireColor.A = 0.2f;
 		}
@@ -610,6 +614,8 @@ void FConnectionDrawingPolicy::DrawPinGeometries(TMap<TSharedRef<SWidget>, FArra
 				if (( LinkEndWidgetGeometry && LinkStartWidgetGeometry ) && !IsConnectionCulled( *LinkStartWidgetGeometry, *LinkEndWidgetGeometry ))
 				{
 					FConnectionParams Params;
+					bool bAreConnectionsFaded = false;
+
 					DetermineWiringStyle(ThePin, TargetPin, /*inout*/ Params);
 					const TSharedPtr<SGraphPin>* ConnectedPinWidget = PinToPinWidgetMap.Find(TargetPin);
 					if (ConnectedPinWidget && ConnectedPinWidget->IsValid())
@@ -617,6 +623,7 @@ void FConnectionDrawingPolicy::DrawPinGeometries(TMap<TSharedRef<SWidget>, FArra
 						if ( PinWidget.AreConnectionsFaded() && (*ConnectedPinWidget)->AreConnectionsFaded() )
 						{
 							Params.WireColor.A = 0.2f;
+							bAreConnectionsFaded = true;
 						}
 					}
 
@@ -625,7 +632,12 @@ void FConnectionDrawingPolicy::DrawPinGeometries(TMap<TSharedRef<SWidget>, FArra
 					FVector2f StartPoint;
 					FVector2f EndPoint;
 					GetSplineEndPoints(LinkStartWidgetGeometry->Geometry, LinkEndWidgetGeometry->Geometry, ArrowRadius, StartPoint, EndPoint);
-					LastConnections.Add(FSimpleConnectionData{ConnectedPinWidget ? *ConnectedPinWidget : nullptr, StaticCastSharedRef<SGraphPin>(SomePinWidget), StartPoint, EndPoint, Params});
+
+					TSharedPtr<SGraphPin> StartPinWidget = ConnectedPinWidget ? *ConnectedPinWidget : nullptr;
+					TSharedRef<SGraphPin> EndPinWidget = StaticCastSharedRef<SGraphPin>(SomePinWidget);
+					FGraphPinHandle StartPinHandle(StartPinWidget ? StartPinWidget->GetPinObj() : nullptr);
+					FGraphPinHandle EndPinHandle(EndPinWidget->GetPinObj());
+					LastConnections.Add(FSimpleConnectionData{ StartPinHandle, EndPinHandle, StartPoint, EndPoint, Params, bAreConnectionsFaded });
 				}
 			}
 		}
diff --git a/Engine/Source/Editor/GraphEditor/Public/ConnectionDrawingPolicy.h b/Engine/Source/Editor/GraphEditor/Public/ConnectionDrawingPolicy.h
--- a/Engine/Source/Editor/GraphEditor/Public/ConnectionDrawingPolicy.h
+++ b/Engine/Source/Editor/GraphEditor/Public/ConnectionDrawingPolicy.h
@@ -232,7 +232,7 @@ public:
 	 * Used in MouseMove to check if the connections have graphically changed and need the GraphPanel to invalidate.
 	 * The default implementation checks for a difference in HoveredPins, or in difference in connection WireColor or WireThickness.
 	 */
-	UE_API virtual bool HaveConnectionsGraphicallyChanged(const TSet<FEdGraphPinReference>& InNewlyHoveredPins) const;
+	UE_API virtual bool HaveConnectionsGraphicallyChanged(const SGraphPanel& InGraphPanel, const TSet<FEdGraphPinReference>& InNewlyHoveredPins) const;
 
 protected:
 	// Helper function used by Draw(). Called before DrawPinGeometries to populate PinToPinWidgetMap
@@ -261,11 +261,12 @@ protected:
 
 	struct FSimpleConnectionData
 	{
-		TSharedPtr<SGraphPin> StartPin;
-		TSharedRef<SGraphPin> EndPin;
+		FGraphPinHandle StartPinHandle;
+		FGraphPinHandle EndPinHandle;
 		FVector2f StartPoint;
 		FVector2f EndPoint;
 		FConnectionParams Params;
+		bool bAreConnectionsFaded = false;
 	};
 	/** Keeps track of the last connections drawn to check if they would have changed on MouseMove */
 	TArray<FSimpleConnectionData> LastConnections;
