Fix for `SGraphPanel` ensuring that `PreviewConnectorEndpoint` is properly initialized when starting a new connection.

#Synced-CL 51747987

diff --git a/Engine/Source/Editor/GraphEditor/Private/SGraphPanel.cpp b/Engine/Source/Editor/GraphEditor/Private/SGraphPanel.cpp
--- a/Engine/Source/Editor/GraphEditor/Private/SGraphPanel.cpp
+++ b/Engine/Source/Editor/GraphEditor/Private/SGraphPanel.cpp
@@ -157,6 +157,7 @@ void SGraphPanel::Construct( const SGraphPanel::FArguments& InArgs )
 
 	SavedMousePosForOnPaintEventLocalSpace = FVector2f::ZeroVector;
 	PreviousFrameSavedMousePosForSplineOverlap = FVector2f::ZeroVector;
+	PreviewConnectorEndpoint = FVector2f::ZeroVector;
 
 	TimeLeftToInvalidatePerTick = 0.0f;
 
@@ -1395,6 +1396,7 @@ void SGraphPanel::OnDragEnter( const FGeometry& MyGeometry, const FDragDropEvent
 	if (DragConnectionOp.IsValid())
 	{
 		DragConnectionOp->SetHoveredGraph( SharedThis(this) );
+		PreviewConnectorEndpoint = MyGeometry.AbsoluteToLocal(DragDropEvent.GetScreenSpacePosition());
 	}
 }
 
@@ -1789,6 +1791,7 @@ void SGraphPanel::OnBeginMakingConnection(FGraphPinHandle PinHandle)
 		DismissContextMenu();
 
 		PreviewConnectorFromPins.Add(PinHandle);
+		PreviewConnectorEndpoint = SavedMousePosForOnPaintEventLocalSpace;
 		bIsDrawStateCached = false;
 	}
 }
