Fix: Node paint ordering and Z-order (overlapping fix). The original single-loop paint pass had no selection-aware ordering.

[RN] Graph Nodes when dragged are now on top of the other nodes.
#Synced-CL 0

diff --git a/Engine/Source/Editor/GraphEditor/Private/SGraphPanel.cpp b/Engine/Source/Editor/GraphEditor/Private/SGraphPanel.cpp
--- a/Engine/Source/Editor/GraphEditor/Private/SGraphPanel.cpp
+++ b/Engine/Source/Editor/GraphEditor/Private/SGraphPanel.cpp
@@ -288,14 +288,37 @@ int32 SGraphPanel::OnPaint( const FPaintArgs& Args, const FGeometry& AllottedGeo
 		// Context for rendering node infos
 		FKismetNodeInfoContext Context(GraphObj);
 
-		for (int32 ChildIndex = 0; ChildIndex < ArrangedChildren.Num(); ++ChildIndex)
+		// Build paint order: non-selected first, selected/moving last.
+		TArray<int32> PaintOrder;
+		TArray<int32> SelectedIndices;
+		PaintOrder.Reserve(ArrangedChildren.Num());
+		SelectedIndices.Reserve(ArrangedChildren.Num());
+		for (int32 i = 0; i < ArrangedChildren.Num(); ++i)
 		{
+			if (SelectionToVisualize->Contains(StaticCastSharedRef<SNodePanel::SNode>(ArrangedChildren[i].Widget)->GetObjectBeingDisplayed()))
+			{
+				SelectedIndices.Add(i);
+			}
+			else
+			{
+				PaintOrder.Add(i);
+			}
+		}
+		const int32 FirstSelectedIndex = PaintOrder.Num();
+		PaintOrder.Append(SelectedIndices);
+
+		// Non-selected nodes first, selected/moving nodes last.
+		// This ensures that any node being dragged or selected always renders on top of all others.
+		for (int32 OrderIndex = 0; OrderIndex < PaintOrder.Num(); ++OrderIndex)
+		{
+			int32 ChildIndex = PaintOrder[OrderIndex];
+			const bool bSelected = (OrderIndex >= FirstSelectedIndex);
 			FArrangedWidget& CurWidget = ArrangedChildren[ChildIndex];
 			TSharedRef<SGraphNode> ChildNode = StaticCastSharedRef<SGraphNode>(CurWidget.Widget);
-			
+
 			// Examine node to see what layers we should be drawing in
-			int32 ShadowLayerId = NodeShadowsLayerId;
-			int32 ChildLayerId = NodeLayerId;
+			int32 ShadowLayerId = bSelected ? MaxLayerId : NodeShadowsLayerId;
+			int32 ChildLayerId = MaxLayerId;
 
 			// If a comment node, draw in the dedicated comment slots
 			{
@@ -308,8 +331,6 @@ int32 SGraphPanel::OnPaint( const FPaintArgs& Args, const FGeometry& AllottedGeo
 			}
 
 
-			const bool bSelected = SelectionToVisualize->Contains( StaticCastSharedRef<SNodePanel::SNode>(CurWidget.Widget)->GetObjectBeingDisplayed() );
-
 			UEdGraphNode* NodeObj = Cast<UEdGraphNode>(ChildNode->GetObjectBeingDisplayed());
 			float Alpha = 1.0f;
 
@@ -386,7 +407,7 @@ int32 SGraphPanel::OnPaint( const FPaintArgs& Args, const FGeometry& AllottedGeo
 				ChildNode->DrawChildrenNodeShadow(
 					*SelectionToVisualize,
 					/*FSlateWindowElementList& */OutDrawElements,
-					NodeShadowsLayerId,
+					ShadowLayerId,
 					CurWidget.Geometry,
 					FLinearColor(1.0f, 1.0f, 1.0f, Alpha)
 				);
@@ -2008,17 +2029,17 @@ public:
 			return A.Rank > B.Rank;
 		});
 
-		TSet<UEdGraphNode*> VistedNodes;
+		TSet<UEdGraphNode*> VisitedNodes;
 		for (FRankedNode& RankedNode : RankedNodes)
 		{
-			StraightenConnectionsForNode(RankedNode.Node, VistedNodes, EEdGraphPinDirection::EGPD_Output);
-			if (VistedNodes.Num() == RankedNodes.Num())
+			StraightenConnectionsForNode(RankedNode.Node, VisitedNodes, EEdGraphPinDirection::EGPD_Output);
+			if (VisitedNodes.Num() == RankedNodes.Num())
 			{
 				return;
 			}
 
-			StraightenConnectionsForNode(RankedNode.Node, VistedNodes, EEdGraphPinDirection::EGPD_Input);
-			if (VistedNodes.Num() == RankedNodes.Num())
+			StraightenConnectionsForNode(RankedNode.Node, VisitedNodes, EEdGraphPinDirection::EGPD_Input);
+			if (VisitedNodes.Num() == RankedNodes.Num())
 			{
 				return;
 			}
