Fix minor error in messages of the graph diff tool

#Synced-CL 52484127

diff --git a/Engine/Source/Editor/GraphEditor/Private/GraphDiffControl.cpp b/Engine/Source/Editor/GraphEditor/Private/GraphDiffControl.cpp
--- a/Engine/Source/Editor/GraphEditor/Private/GraphDiffControl.cpp
+++ b/Engine/Source/Editor/GraphEditor/Private/GraphDiffControl.cpp
@@ -304,9 +304,9 @@ static void DiffR_PinDefaultValueChanged(FDiffResults& Results, UEdGraphPin* New
 		Args.Add(TEXT("PinNameForValue1"), FText::FromName(NewPin->PinName));
 		Args.Add(TEXT("PinValue1"), FText::FromString(OldPin->GetDefaultAsString()));
 		Args.Add(TEXT("PinValue2"), FText::FromString(NewPin->GetDefaultAsString()));
-		Diff.ToolTip = FText::Format(LOCTEXT("DIF_PinDefaultValueToolTip", "Pin '{PinNameForValue1}' Default Value was '{PinValue1}', but is now '{PinValue2}"), Args);
+		Diff.ToolTip = FText::Format(LOCTEXT("DIF_PinDefaultValueToolTip", "Pin '{PinNameForValue1}' Default Value was '{PinValue1}', but is now '{PinValue2}'"), Args);
 		Diff.Category = EDiffType::MODIFICATION;
-		Diff.DisplayString = FText::Format(LOCTEXT("DIF_PinDefaultValue", "Pin Default '{PinNameForValue1}' '{PinValue1}' -> '{PinValue2}']"), Args);
+		Diff.DisplayString = FText::Format(LOCTEXT("DIF_PinDefaultValue", "Pin Default '{PinNameForValue1}' ['{PinValue1}' -> '{PinValue2}']"), Args);
 	}
 
 	Results.Add(Diff);
